MotionDesign : Vector Customization
- AllowPreserveRatio : read and write property metadata directly
- Use MotionDesignVectorWidget metadata to use this customization
- To lock axis set AllowPreserveRatio="{AxisToLock}" where AxisToLock can be X, Y, Z or empty for free

diff --git a/Engine/Plugins/Experimental/Avalanche/Source/AvalancheEditor/Private/DetailView/Customizations/AvaVectorPropertyTypeCustomization.cpp b/Engine/Plugins/Experimental/Avalanche/Source/AvalancheEditor/Private/DetailView/Customizations/AvaVectorPropertyTypeCustomization.cpp
--- a/Engine/Plugins/Experimental/Avalanche/Source/AvalancheEditor/Private/DetailView/Customizations/AvaVectorPropertyTypeCustomization.cpp
+++ b/Engine/Plugins/Experimental/Avalanche/Source/AvalancheEditor/Private/DetailView/Customizations/AvaVectorPropertyTypeCustomization.cpp
@@ -10,12 +10,10 @@
 #include "DynamicMeshes/AvaShape2DDynMeshBase.h"
 #include "DynamicMeshes/AvaShape3DDynMeshBase.h"
 #include "Editor.h"
-#include "IDetailChildrenBuilder.h"
 #include "Viewport/AvaViewportExtension.h"
 #include "ViewportClient/IAvaViewportClient.h"
-#include "Widgets/Input/SButton.h"
-#include "Widgets/Input/SComboButton.h"
-#include "Widgets/Input/STextComboBox.h"
+#include "Widgets/Layout/SBox.h"
+#include "Widgets/Input/SComboBox.h"
 #include "Widgets/Layout/SScaleBox.h"
 
 #define LOCTEXT_NAMESPACE "AvaVectorPropertyTypeCustomization"
@@ -37,14 +35,13 @@ void FAvaVectorPropertyTypeCustomization::CustomizeHeader(TSharedRef<IPropertyHa
 		// All outer objects should have the same world? So just use the first
 		if (OuterObjects.Num() > 0)
 		{
-			UAvaEditorSubsystem* AvaEditorSubsystem = UAvaEditorSubsystem::Get(OuterObjects[0]);
-			if (AvaEditorSubsystem)
+			if (UAvaEditorSubsystem* AvaEditorSubsystem = UAvaEditorSubsystem::Get(OuterObjects[0]))
 			{
-				TSharedPtr<FAvaViewportExtension> ViewportExtension = AvaEditorSubsystem->FindExtension<FAvaViewportExtension>();
-				if (ViewportExtension.IsValid())
+				if (TSharedPtr<FAvaViewportExtension> ViewportExtension = AvaEditorSubsystem->FindExtension<FAvaViewportExtension>())
 				{
 					TArray<TSharedPtr<IAvaViewportClient>> ViewportClients = ViewportExtension->GetViewportClients();
-					if (ViewportClients.Num() > 0)
+
+					if (!ViewportClients.IsEmpty())
 					{
 						ViewportClient = ViewportClients[0];
 					}
@@ -53,122 +50,55 @@ void FAvaVectorPropertyTypeCustomization::CustomizeHeader(TSharedRef<IPropertyHa
 		}
 	}
 
-	bool bPreserveRatio = false;
-	TSharedPtr<SWidget> PreserveRatioWidget = SNullWidget::NullWidget;
+	// Assign Name Widget
+	HeaderRow.NameContent()[StructPropertyHandle->CreatePropertyNameWidget()];
 
-	if (StructPropertyHandle->HasMetaData("AllowPreserveRatio"))
-	{
-		bPreserveRatio = true;
-		
-		static const FVector2D ImageSize(16.f);
+	// Fill available space
+	HeaderRow.ValueWidget.HorizontalAlignment = EHorizontalAlignment::HAlign_Fill;
+	HeaderRow.ValueWidget.VerticalAlignment = EVerticalAlignment::VAlign_Fill;
 
-		// fill available space
-		HeaderRow.NameWidget.HorizontalAlignment = EHorizontalAlignment::HAlign_Fill;
-		HeaderRow.NameWidget.VerticalAlignment = EVerticalAlignment::VAlign_Fill;
+	FName RatioNone(TEXT("Ratio"));
+	RatioNone.SetNumber(static_cast<int32>(ERatioMode::None));
+	RatioModes.Add(RatioNone);
 
-		auto ComboBoxButtonBuilder = [this](const ERatioMode ButtonMode)->TSharedRef<SButton>
-		{
-			const FSlateBrush* ButtonImage = GetComboButtonBrush(ButtonMode);
-			const FText& ButtonText = GetComboButtonText(ButtonMode);
-			
-			return SNew(SButton)
-					.HAlign(HAlign_Fill)
-					.VAlign(VAlign_Fill)
-					.Cursor(EMouseCursor::Hand)
-					.OnClicked(this, &FAvaVectorPropertyTypeCustomization::OnComboButtonClicked, ButtonMode)
-					.ContentPadding(FMargin(-8.f, 0.f))
-					[
-						SNew(SHorizontalBox)
-						+ SHorizontalBox::Slot()
-						.HAlign(HAlign_Left)
-						.AutoWidth()
-						[
-							SNew(SScaleBox)
-							[
-								SNew(SImage)
-								.DesiredSizeOverride(ImageSize)
-								.Image(ButtonImage)
-							]
-						]
-						+ SHorizontalBox::Slot()
-						.AutoWidth()
-						[
-							SNew(STextBlock)
-							.Justification(ETextJustify::Center)
-							.Text(ButtonText)
-							.Margin(FMargin(5.f, 0.f, 0.f, 0.f))
-						]
-					];
-		};
-
-		TArray<ERatioMode> RatioModes {
-			ERatioMode::None,
-			ERatioMode::PreserveXY
-		};
-		
-		if (bIsVector3d)
-		{
-			RatioModes.Add(ERatioMode::PreserveYZ);
-			RatioModes.Add(ERatioMode::PreserveXZ);
-			RatioModes.Add(ERatioMode::PreserveXYZ);
-		}
+	FName RatioXY(TEXT("Ratio"));
+	RatioXY.SetNumber(static_cast<int32>(ERatioMode::PreserveXY));
+	RatioModes.Add(RatioXY);
 
-		const TSharedRef<SVerticalBox> ButtonVerticalBox = SNew(SVerticalBox);
-		for (const ERatioMode& Mode : RatioModes)
-		{
-			ButtonVerticalBox->AddSlot()
-			.Padding(0.f)
-			[
-				ComboBoxButtonBuilder(Mode)
-			];
-		}
+	if (bIsVector3d)
+	{
+		FName RatioXZ(TEXT("Ratio"));
+		RatioXZ.SetNumber(static_cast<int32>(ERatioMode::PreserveXZ));
+		RatioModes.Add(RatioXZ);
 
-		PreserveRatioWidget = SNew(SBox)
-			.MinDesiredWidth(60.f)
-			[
-				SAssignNew(ComboButton, SComboButton)
-				.VAlign(EVerticalAlignment::VAlign_Center)
-				.Cursor(EMouseCursor::Hand)
-				.Method(EPopupMethod::UseCurrentWindow)
-				.HasDownArrow(false)
-				.ContentPadding(FMargin(-5.f, -2.f))
-				.ButtonContent()
-				[
-					SNew(SHorizontalBox)
-					+ SHorizontalBox::Slot()
-					.AutoWidth()
-					[
-						SNew(SScaleBox)
-						[
-							SNew(SImage)
-							.DesiredSizeOverride(ImageSize)
-							.Image(this, &FAvaVectorPropertyTypeCustomization::GetCurrentComboButtonBrush)
-						]
-					]
-					+ SHorizontalBox::Slot()
-					.AutoWidth()
-					[
-						SNew(STextBlock)
-						.Justification(ETextJustify::Center)
-						.Text(this, &FAvaVectorPropertyTypeCustomization::GetCurrentComboButtonText)
-						.Margin(FMargin(5.f, 0.f))
-					]
-				]
-				.MenuContent()
-				[
-					ButtonVerticalBox
-				]
-			];
+		FName RatioYZ(TEXT("Ratio"));
+		RatioYZ.SetNumber(static_cast<int32>(ERatioMode::PreserveYZ));
+		RatioModes.Add(RatioYZ);
+
+		FName RatioXYZ(TEXT("Ratio"));
+		RatioXYZ.SetNumber(static_cast<int32>(ERatioMode::PreserveXYZ));
+		RatioModes.Add(RatioXYZ);
 	}
 
-	// Assign Name Widget
-	HeaderRow.NameContent()[StructPropertyHandle->CreatePropertyNameWidget()];
+	const TSharedPtr<SWidget> PreserveRatioWidget = SNew(SBox)
+		.MinDesiredWidth(60.f)
+		.Visibility(this, &FAvaVectorPropertyTypeCustomization::GetRatioWidgetVisibility)
+		[
+			SNew(SComboBox<FName>)
+			.ComboBoxStyle(&FAppStyle::Get().GetWidgetStyle<FComboBoxStyle>(TEXT("ComboBox")))
+			.OptionsSource(&RatioModes)
+			.HasDownArrow(false)
+			.InitiallySelectedItem(GetRatioCurrentItem())
+			.ToolTipText(LOCTEXT("PreserveRatioTooltip", "Select the ratio mode to lock for the component axis"))
+			.ContentPadding(0.f)
+			.OnGenerateWidget(this, &FAvaVectorPropertyTypeCustomization::OnGenerateRatioWidget)
+			.OnSelectionChanged(this, &FAvaVectorPropertyTypeCustomization::OnRatioSelectionChanged)
+			.Content()
+			[
+				OnGenerateRatioWidget(NAME_None)
+			]
+		];
 
-	// fill available space
-	HeaderRow.ValueWidget.HorizontalAlignment = EHorizontalAlignment::HAlign_Fill;
-	HeaderRow.ValueWidget.VerticalAlignment = EVerticalAlignment::VAlign_Fill;
-	RatioMode = ERatioMode::None;
-	
 	if (bIsVector3d)
 	{
 		bPixelSizeProperty = VectorPropertyHandle->GetProperty()->GetFName() == GET_MEMBER_NAME_CHECKED(UAvaShape3DDynMeshBase, PixelSize3D);
@@ -176,23 +106,11 @@ void FAvaVectorPropertyTypeCustomization::CustomizeHeader(TSharedRef<IPropertyHa
 		YPropertyHandle = StructPropertyHandle->GetChildHandle(GET_MEMBER_NAME_CHECKED(FVector, Y));
 		ZPropertyHandle = StructPropertyHandle->GetChildHandle(GET_MEMBER_NAME_CHECKED(FVector, Z));
 
-		if (bPreserveRatio)
-		{
-			if (const FString& Value = VectorPropertyHandle->GetMetaDataProperty()->GetMetaData(FName("VectorRatioMode")); !Value.IsEmpty())
-			{
-				RatioMode = static_cast<ERatioMode>(FCString::Atoi(*Value));
-			}
-			else
-			{
-				RatioMode = ERatioMode::None;
-			}
-		}
-		
 		if (StructPropertyHandle->HasMetaData("ClampMin"))
 		{
 			MinVectorClamp = FVector(StructPropertyHandle->GetFloatMetaData("ClampMin"));
 		}
-		
+
 		if (StructPropertyHandle->HasMetaData("ClampMax"))
 		{
 			MaxVectorClamp = FVector(StructPropertyHandle->GetFloatMetaData("ClampMax"));
@@ -208,7 +126,7 @@ void FAvaVectorPropertyTypeCustomization::CustomizeHeader(TSharedRef<IPropertyHa
 		{
 			SpinDelta = (MaxVectorClamp->X - MinVectorClamp->X) / 100.f;
 		}
-		
+
 		HeaderRow.ValueContent()
 		[
 			SNew(SHorizontalBox)
@@ -228,7 +146,7 @@ void FAvaVectorPropertyTypeCustomization::CustomizeHeader(TSharedRef<IPropertyHa
 				.X(this, &FAvaVectorPropertyTypeCustomization::GetVectorComponent,  static_cast<uint8>(0)) // X
 				.Y(this, &FAvaVectorPropertyTypeCustomization::GetVectorComponent,  static_cast<uint8>(1)) // Y
 				.Z(this, &FAvaVectorPropertyTypeCustomization::GetVectorComponent,  static_cast<uint8>(2)) // Z
-				.bColorAxisLabels(bPreserveRatio)
+				.bColorAxisLabels(true)
 				.MinVector(MinVectorClamp)
 				.MaxVector(MaxVectorClamp)
 				.MinSliderVector(MinVectorClamp)
@@ -253,23 +171,11 @@ void FAvaVectorPropertyTypeCustomization::CustomizeHeader(TSharedRef<IPropertyHa
 		XPropertyHandle = StructPropertyHandle->GetChildHandle(GET_MEMBER_NAME_CHECKED(FVector2D, X));
 		YPropertyHandle = StructPropertyHandle->GetChildHandle(GET_MEMBER_NAME_CHECKED(FVector2D, Y));
 
-		if (bPreserveRatio)
-		{
-			if (const FString& Value = VectorPropertyHandle->GetMetaDataProperty()->GetMetaData(FName("VectorRatioMode")); !Value.IsEmpty())
-			{
-				RatioMode = static_cast<ERatioMode>(FCString::Atoi(*Value));
-			}
-			else
-			{
-				RatioMode = ERatioMode::None;
-			}
-		}
-		
 		if (StructPropertyHandle->HasMetaData("ClampMin"))
 		{
 			MinVector2DClamp = FVector2D(StructPropertyHandle->GetFloatMetaData("ClampMin"));
 		}
-		
+
 		if (StructPropertyHandle->HasMetaData("ClampMax"))
 		{
 			MaxVector2DClamp = FVector2D(StructPropertyHandle->GetFloatMetaData("ClampMax"));
@@ -304,7 +210,7 @@ void FAvaVectorPropertyTypeCustomization::CustomizeHeader(TSharedRef<IPropertyHa
 				.Font(IDetailLayoutBuilder::GetDetailFont())
 				.X(this, &FAvaVectorPropertyTypeCustomization::GetVectorComponent, static_cast<uint8>(0)) // X
 				.Y(this, &FAvaVectorPropertyTypeCustomization::GetVectorComponent,  static_cast<uint8>(1)) // Y
-				.bColorAxisLabels(bPreserveRatio)
+				.bColorAxisLabels(true)
 				.MinVector(MinVector2DClamp)
 				.MaxVector(MaxVector2DClamp)
 				.MinSliderVector(MinVector2DClamp)
@@ -328,71 +234,132 @@ void FAvaVectorPropertyTypeCustomization::CustomizeChildren(TSharedRef<IProperty
 {
 }
 
-FReply FAvaVectorPropertyTypeCustomization::OnComboButtonClicked(const ERatioMode NewMode)
-{
-	RatioMode = NewMode;
-	const FString MetaData = FString::FromInt(static_cast<int32>(RatioMode));
-	VectorPropertyHandle->GetMetaDataProperty()->SetMetaData(FName("VectorRatioMode"), *MetaData);
-	ComboButton->SetIsOpen(false);
-	return FReply::Handled();
-}
-
-const FSlateBrush* FAvaVectorPropertyTypeCustomization::GetComboButtonBrush(const ERatioMode Mode) const
+const FSlateBrush* FAvaVectorPropertyTypeCustomization::GetRatioModeBrush(const ERatioMode InMode) const
 {
-	if (Mode == ERatioMode::None)
+	if (InMode == ERatioMode::None)
 	{
 		return FAvaEditorStyle::Get().GetBrush("Icons.Unlock");
 	}
-	else if (Mode == ERatioMode::PreserveXY || Mode == ERatioMode::PreserveXZ || Mode == ERatioMode::PreserveYZ)
+
+	if (InMode == ERatioMode::PreserveXY || InMode == ERatioMode::PreserveXZ || InMode == ERatioMode::PreserveYZ)
 	{
 		return FAvaEditorStyle::Get().GetBrush("Icons.Lock2d");
 	}
-	else
-	{
-		return FAvaEditorStyle::Get().GetBrush("Icons.Lock3d");
-	}
+
+	return FAvaEditorStyle::Get().GetBrush("Icons.Lock3d");
 }
 
-FText FAvaVectorPropertyTypeCustomization::GetComboButtonText(const ERatioMode Mode) const
+FText FAvaVectorPropertyTypeCustomization::GetRatioModeDisplayText(const ERatioMode InMode) const
 {
-	switch (Mode)
+	FText DisplayText = FText::GetEmpty();
+
+	switch (InMode)
 	{
-		case(ERatioMode::PreserveXY):
-			return FText::FromString("XY ");
+		case ERatioMode::PreserveXY:
+			DisplayText = FText::FromString("XY ");
 			break;
-		case(ERatioMode::PreserveYZ):
-			return FText::FromString("YZ ");
+		case ERatioMode::PreserveYZ:
+			DisplayText = FText::FromString("YZ ");
 			break;
-		case(ERatioMode::PreserveXZ):
-			return FText::FromString("XZ ");
+		case ERatioMode::PreserveXZ:
+			DisplayText = FText::FromString("XZ ");
 			break;
-		case(ERatioMode::PreserveXYZ):
-			return FText::FromString("XYZ");
+		case ERatioMode::PreserveXYZ:
+			DisplayText = FText::FromString("XYZ");
 			break;
 		default:
-			return FText::FromString("Free");
+			DisplayText = FText::FromString("Free");
 			break;
 	}
+
+	return DisplayText;
+}
+
+const FSlateBrush* FAvaVectorPropertyTypeCustomization::GetCurrentRatioModeBrush() const
+{
+	return GetRatioModeBrush(GetRatioModeMetadata());
+}
+
+FText FAvaVectorPropertyTypeCustomization::GetCurrentRatioModeDisplayText() const
+{
+	return GetRatioModeDisplayText(GetRatioModeMetadata());
+}
+
+ERatioMode FAvaVectorPropertyTypeCustomization::GetRatioModeMetadata() const
+{
+	ERatioMode RatioMode = ERatioMode::None;
+
+	if (VectorPropertyHandle.IsValid() && VectorPropertyHandle->IsValidHandle())
+	{
+		const FString& MetadataValue = VectorPropertyHandle->GetMetaData(PropertyMetadata);
+
+		if (MetadataValue.Contains(TEXT("X")))
+		{
+			RatioMode |= ERatioMode::X;
+		}
+
+		if (MetadataValue.Contains(TEXT("Y")))
+		{
+			RatioMode |= ERatioMode::Y;
+		}
+
+		if (bIsVector3d && MetadataValue.Contains(TEXT("Z")))
+		{
+			RatioMode |= ERatioMode::Z;
+		}
+	}
+
+	return RatioMode;
 }
 
-const FSlateBrush* FAvaVectorPropertyTypeCustomization::GetCurrentComboButtonBrush() const
+void FAvaVectorPropertyTypeCustomization::SetRatioModeMetadata(ERatioMode InMode) const
 {
-	return GetComboButtonBrush(RatioMode);
+	if (VectorPropertyHandle.IsValid() && VectorPropertyHandle->IsValidHandle())
+	{
+		FProperty* VectorProperty = VectorPropertyHandle->GetProperty();
+
+		FString NewMetadataValue = TEXT("");
+
+		if ((InMode & ERatioMode::X) != ERatioMode::None)
+		{
+			NewMetadataValue += TEXT("X");
+		}
+
+		if ((InMode & ERatioMode::Y) != ERatioMode::None)
+		{
+			NewMetadataValue += TEXT("Y");
+		}
+
+		if (bIsVector3d && (InMode & ERatioMode::Z) != ERatioMode::None)
+		{
+			NewMetadataValue += TEXT("Z");
+		}
+
+		VectorProperty->SetMetaData(PropertyMetadata, *NewMetadataValue);
+	}
 }
 
-FText FAvaVectorPropertyTypeCustomization::GetCurrentComboButtonText() const
+EVisibility FAvaVectorPropertyTypeCustomization::GetRatioWidgetVisibility() const
 {
-	return GetComboButtonText(RatioMode);
+	if (VectorPropertyHandle.IsValid() && VectorPropertyHandle->IsValidHandle())
+	{
+		// Only show preserve ratio widget if AllowPreserveRatio is set
+		return VectorPropertyHandle->HasMetaData(PropertyMetadata)
+			? EVisibility::Visible
+			: EVisibility::Collapsed;
+	}
+
+	return EVisibility::Collapsed;
 }
 
-TOptional<double> FAvaVectorPropertyTypeCustomization::GetVectorComponent(const uint8 Component) const
+TOptional<double> FAvaVectorPropertyTypeCustomization::GetVectorComponent(const uint8 InComponent) const
 {
 	if (!VectorPropertyHandle.IsValid() || SelectedObjectNum == 0)
 	{
 		return TOptional<double>();
 	}
 	double OutValue = 0.f;
-	switch(Component)
+	switch(InComponent)
 	{
 		case 0:
 			if (XPropertyHandle->GetValue(OutValue) != FPropertyAccess::Success)
@@ -414,31 +381,29 @@ TOptional<double> FAvaVectorPropertyTypeCustomization::GetVectorComponent(const
 		break;
 		default:
 			return TOptional<double>();
-		break;
 	}
+
 	// handle specific case
 	if (bPixelSizeProperty)
 	{
 		return MeshSizeToPixelSize(OutValue);
 	}
-	else
-	{
-		return OutValue;
-	}
+
+	return OutValue;
 }
 
-void FAvaVectorPropertyTypeCustomization::SetVectorComponent(double NewValue, const uint8 Component)
+void FAvaVectorPropertyTypeCustomization::SetVectorComponent(double InNewValue, const uint8 InComponent)
 {
 	if (bMovingSlider)
 	{
-		SetVectorComponent(NewValue, ETextCommit::Default, Component);
+		SetVectorComponent(InNewValue, ETextCommit::Default, InComponent);
 	}
 }
 
-void FAvaVectorPropertyTypeCustomization::SetVectorComponent(double NewValue, ETextCommit::Type CommitType, const uint8 Component)
+void FAvaVectorPropertyTypeCustomization::SetVectorComponent(double InNewValue, ETextCommit::Type InCommitType, const uint8 InComponent)
 {
-	const bool bFinalCommit = CommitType == ETextCommit::OnEnter || CommitType == ETextCommit::OnUserMovedFocus;
-	if (!bFinalCommit && CommitType != ETextCommit::Default)
+	const bool bFinalCommit = InCommitType == ETextCommit::OnEnter || InCommitType == ETextCommit::OnUserMovedFocus;
+	if (!bFinalCommit && InCommitType != ETextCommit::Default)
 	{
 		return;
 	}
@@ -457,7 +422,7 @@ void FAvaVectorPropertyTypeCustomization::SetVectorComponent(double NewValue, ET
 		return;
 	}
 	// handle interactive debounce to avoid slow behaviour
-	LastComponentValueSet = Component;
+	LastComponentValueSet = InComponent;
 	if (DebounceValueSet > 0 && !bFinalCommit)
 	{
 		DebounceValueSet--;
@@ -467,7 +432,7 @@ void FAvaVectorPropertyTypeCustomization::SetVectorComponent(double NewValue, ET
 	// handle specific case for pixel size property
 	if (bPixelSizeProperty)
 	{
-		NewValue = PixelSizeToMeshSize(NewValue);
+		InNewValue = PixelSizeToMeshSize(InNewValue);
 	}
 	if (!bMovingSlider)
 	{
@@ -475,7 +440,7 @@ void FAvaVectorPropertyTypeCustomization::SetVectorComponent(double NewValue, ET
 	}
 	// update objects value for property, we handle transaction ourselves to batch properties changes together
 	const EPropertyValueSetFlags::Type Flags = bMovingSlider ? EPropertyValueSetFlags::InteractiveChange : EPropertyValueSetFlags::NotTransactable;
-	SetComponentValue(NewValue, Component, Flags);
+	SetComponentValue(InNewValue, InComponent, Flags);
 	if (!bMovingSlider)
 	{
 		GEditor->EndTransaction();
@@ -492,14 +457,14 @@ void FAvaVectorPropertyTypeCustomization::OnBeginSliderMovement()
 	GEditor->BeginTransaction(VectorPropertyHandle->GetPropertyDisplayName());
 }
 
-void FAvaVectorPropertyTypeCustomization::OnEndSliderMovement(double NewValue)
+void FAvaVectorPropertyTypeCustomization::OnEndSliderMovement(double InNewValue)
 {
 	DebounceValueSet = 0;
 	bMovingSlider = false;
 	// set final value like enter pressed
 	if (LastComponentValueSet != INVALID_COMPONENT_IDX)
 	{
-		SetVectorComponent(NewValue, ETextCommit::OnEnter, LastComponentValueSet);
+		SetVectorComponent(InNewValue, ETextCommit::OnEnter, LastComponentValueSet);
 	}
 	// end started transactions during process
 	while(GEditor->IsTransactionActive())
@@ -509,13 +474,88 @@ void FAvaVectorPropertyTypeCustomization::OnEndSliderMovement(double NewValue)
 	ResetVectorValuesForRatio();
 }
 
+TSharedRef<SWidget> FAvaVectorPropertyTypeCustomization::OnGenerateRatioWidget(FName InRatioMode)
+{
+	TSharedPtr<SImage> ImageWidget;
+	TSharedPtr<STextBlock> TextWidget;
+
+	const FVector2d ImageSize(16.f);
+
+	if (InRatioMode.IsNone())
+	{
+		ImageWidget = SNew(SImage)
+			.ColorAndOpacity(FAppStyle::GetSlateColor("SelectionColor"))
+			.DesiredSizeOverride(ImageSize)
+			.Image(this, &FAvaVectorPropertyTypeCustomization::GetCurrentRatioModeBrush);
+
+		TextWidget = SNew(STextBlock)
+			.Justification(ETextJustify::Right)
+			.Font(FAppStyle::Get().GetFontStyle("SmallFont"))
+			.Text(this, &FAvaVectorPropertyTypeCustomization::GetCurrentRatioModeDisplayText);
+	}
+	else
+	{
+		const ERatioMode RatioMode = static_cast<ERatioMode>(InRatioMode.GetNumber());
+
+		ImageWidget = SNew(SImage)
+			.ColorAndOpacity(FAppStyle::GetSlateColor("SelectionColor"))
+			.DesiredSizeOverride(ImageSize)
+			.Image(GetRatioModeBrush(RatioMode));
+
+		TextWidget = SNew(STextBlock)
+			.Justification(ETextJustify::Center)
+			.Font(FAppStyle::Get().GetFontStyle("SmallFont"))
+			.Text(GetRatioModeDisplayText(RatioMode));
+	}
+
+	return SNew(SHorizontalBox)
+		.Visibility(EVisibility::Visible)
+		+ SHorizontalBox::Slot()
+		.AutoWidth()
+		.Padding(0.f)
+		[
+			SNew(SScaleBox)
+			.Visibility(EVisibility::HitTestInvisible)
+			.Stretch(EStretch::UserSpecified)
+			.UserSpecifiedScale(1.f)
+			[
+				ImageWidget.ToSharedRef()
+			]
+		]
+		+ SHorizontalBox::Slot()
+		.FillWidth(1.f)
+		.Padding(8.f, 0.f, 0.f, 0.f)
+		.HAlign(EHorizontalAlignment::HAlign_Fill)
+		.VAlign(EVerticalAlignment::VAlign_Center)
+		[
+			TextWidget.ToSharedRef()
+		];
+}
+
+void FAvaVectorPropertyTypeCustomization::OnRatioSelectionChanged(FName InRatioMode, ESelectInfo::Type InSelectInfo) const
+{
+	if (!InRatioMode.IsNone())
+	{
+		const ERatioMode RatioMode = static_cast<ERatioMode>(InRatioMode.GetNumber());
+		SetRatioModeMetadata(RatioMode);
+	}
+}
+
+FName FAvaVectorPropertyTypeCustomization::GetRatioCurrentItem() const
+{
+	const ERatioMode CurrentRatioMode = GetRatioModeMetadata();
+	FName RatioName(TEXT("Ratio"));
+	RatioName.SetNumber(static_cast<int32>(CurrentRatioMode));
+	return RatioName;
+}
+
 bool FAvaVectorPropertyTypeCustomization::CanEditValue() const
 {
 	if (!VectorPropertyHandle.IsValid())
 	{
 		return false;
 	}
-	
+
 	if (bPixelSizeProperty && SelectedObjectNum > 0 && ViewportClient.IsValid())
 	{
 		if (FAvaViewportUtils::IsValidViewportSize(ViewportClient.Pin()->GetVirtualViewportSize()))
@@ -525,7 +565,7 @@ bool FAvaVectorPropertyTypeCustomization::CanEditValue() const
 
 		return false;
 	}
-	
+
 	return VectorPropertyHandle->IsEditable();
 }
 
@@ -533,11 +573,11 @@ void FAvaVectorPropertyTypeCustomization::InitVectorValuesForRatio()
 {
 	Begin2DValues.Empty();
 	Begin3DValues.Empty();
-	
+
 	TArray<FString> OutValues;
 	VectorPropertyHandle->GetPerObjectValues(OutValues);
 	SelectedObjectNum = VectorPropertyHandle->GetNumPerObjectValues();
-	
+
 	if (SelectedObjectNum > 0)
 	{
 		for (const FString& Val : OutValues)
@@ -578,11 +618,13 @@ void FAvaVectorPropertyTypeCustomization::ResetVectorValuesForRatio()
 	Begin3DValues.Empty();
 }
 
-void FAvaVectorPropertyTypeCustomization::SetComponentValue(const double NewValue, const uint8 Component, const EPropertyValueSetFlags::Type Flags)
+void FAvaVectorPropertyTypeCustomization::SetComponentValue(const double InNewValue, const uint8 InComponent, const EPropertyValueSetFlags::Type InFlags)
 {
+	const ERatioMode RatioMode = GetRatioModeMetadata();
+
 	// check if we are preserving ratio for current component change
 	bool bPreserveRatio = false;
-	switch (Component)
+	switch (InComponent)
 	{
 		case 0:
 			bPreserveRatio = ((RatioMode & ERatioMode::X) != ERatioMode::None);
@@ -615,35 +657,35 @@ void FAvaVectorPropertyTypeCustomization::SetComponentValue(const double NewValu
 			continue;
 		}
 		// compute clamped ratio for value change
-		const double ClampedRatio = GetClampedRatioValueChange(ObjIdx, NewValue, Component, PreserveRatios);
+		const double ClampedRatio = GetClampedRatioValueChange(ObjIdx, InNewValue, InComponent, PreserveRatios);
 		// loop for each component (X,Y,Z)
 		for (uint8 ComponentIdx = 0; ComponentIdx < MaxComponentCount; ComponentIdx++)
 		{
 			// only assign value to specific component, skip others
-			if (!bPreserveRatio && ComponentIdx != Component)
+			if (!bPreserveRatio && ComponentIdx != InComponent)
 			{
 				continue;
 			}
 			// compute new component value
-			const double NewComponentValue = GetClampedComponentValue(ObjIdx, NewValue, ClampedRatio, ComponentIdx, Component);
+			const double NewComponentValue = GetClampedComponentValue(ObjIdx, InNewValue, ClampedRatio, ComponentIdx, InComponent);
 			switch(ComponentIdx)
 			{
 				case 0:
-					if (PreserveRatios[0] || ComponentIdx == Component)
+					if (PreserveRatios[0] || ComponentIdx == InComponent)
 					{
-						XPropertyHandle->SetPerObjectValue(ObjIdx, FString::SanitizeFloat(NewComponentValue), Flags);
+						XPropertyHandle->SetPerObjectValue(ObjIdx, FString::SanitizeFloat(NewComponentValue), InFlags);
 					}
 					break;
 				case 1:
-					if (PreserveRatios[1] || ComponentIdx == Component)
+					if (PreserveRatios[1] || ComponentIdx == InComponent)
 					{
-						YPropertyHandle->SetPerObjectValue(ObjIdx, FString::SanitizeFloat(NewComponentValue), Flags);
+						YPropertyHandle->SetPerObjectValue(ObjIdx, FString::SanitizeFloat(NewComponentValue), InFlags);
 					}
 					break;
 				case 2:
-					if (PreserveRatios[2] || ComponentIdx == Component)
+					if (PreserveRatios[2] || ComponentIdx == InComponent)
 					{
-						ZPropertyHandle->SetPerObjectValue(ObjIdx, FString::SanitizeFloat(NewComponentValue), Flags);
+						ZPropertyHandle->SetPerObjectValue(ObjIdx, FString::SanitizeFloat(NewComponentValue), InFlags);
 					}
 					break;
 				default:;
@@ -652,23 +694,23 @@ void FAvaVectorPropertyTypeCustomization::SetComponentValue(const double NewValu
 	}
 }
 
-double FAvaVectorPropertyTypeCustomization::GetClampedRatioValueChange(const int32 ObjectIdx, const double NewValue, const uint8 Component, const TArray<bool>& PreserveRatios) const
+double FAvaVectorPropertyTypeCustomization::GetClampedRatioValueChange(const int32 InObjectIdx, const double InNewValue, const uint8 InComponent, const TArray<bool>& InPreserveRatios) const
 {
 	double Ratio = 1;
 	// get pre change value for this component
 	if (bIsVector3d)
 	{
-		FVector BeginValue = Begin3DValues[ObjectIdx].GetValue();
-		if (BeginValue[Component] != 0)
+		FVector BeginValue = Begin3DValues[InObjectIdx].GetValue();
+		if (BeginValue[InComponent] != 0)
 		{
-			Ratio = NewValue / BeginValue[Component];
+			Ratio = InNewValue / BeginValue[InComponent];
 		}
 		// apply min/max clamp
 		if (MinVectorClamp.IsSet() || MaxVectorClamp.IsSet())
 		{
 			for (int32 ComponentIdx = 0; ComponentIdx < 3; ComponentIdx++)
 			{
-				if (PreserveRatios[ComponentIdx] || ComponentIdx == Component)
+				if (InPreserveRatios[ComponentIdx] || ComponentIdx == InComponent)
 				{
 					const double EndValue = BeginValue[ComponentIdx] * Ratio;
 					if (MinVectorClamp.IsSet())
@@ -688,22 +730,22 @@ double FAvaVectorPropertyTypeCustomization::GetClampedRatioValueChange(const int
 						}
 					}
 				}
-			}	
+			}
 		}
 	}
 	else
 	{
-		FVector2D BeginValue = Begin2DValues[ObjectIdx].GetValue();
-		if (BeginValue[Component] != 0)
+		FVector2D BeginValue = Begin2DValues[InObjectIdx].GetValue();
+		if (BeginValue[InComponent] != 0)
 		{
-			Ratio = NewValue / BeginValue[Component];
+			Ratio = InNewValue / BeginValue[InComponent];
 		}
 		// apply min/max clamp
 		if (MinVector2DClamp.IsSet() || MaxVector2DClamp.IsSet())
 		{
 			for (int32 ComponentIdx = 0; ComponentIdx < 2; ComponentIdx++)
 			{
-				if (PreserveRatios[ComponentIdx] || ComponentIdx == Component)
+				if (InPreserveRatios[ComponentIdx] || ComponentIdx == InComponent)
 				{
 					const double EndValue = BeginValue[ComponentIdx] * Ratio;
 					if (MinVector2DClamp.IsSet())
@@ -729,67 +771,67 @@ double FAvaVectorPropertyTypeCustomization::GetClampedRatioValueChange(const int
 	return Ratio;
 }
 
-double FAvaVectorPropertyTypeCustomization::GetClampedComponentValue(const int32 ObjectIdx, double NewValue, const double Ratio, const uint8 ComponentIdx, const uint8 OriginalComponent)
+double FAvaVectorPropertyTypeCustomization::GetClampedComponentValue(const int32 InObjectIdx, double InNewValue, const double InRatio, const uint8 InComponentIdx, const uint8 InOriginalComponent)
 {
-	const double OldValue = (bIsVector3d ? Begin3DValues[ObjectIdx].GetValue()[ComponentIdx] : Begin2DValues[ObjectIdx].GetValue()[ComponentIdx]);
-	const double SliderOriginalValue = (bIsVector3d ? Begin3DValues[ObjectIdx].GetValue()[OriginalComponent] : Begin2DValues[ObjectIdx].GetValue()[OriginalComponent]);
+	const double OldValue = (bIsVector3d ? Begin3DValues[InObjectIdx].GetValue()[InComponentIdx] : Begin2DValues[InObjectIdx].GetValue()[InComponentIdx]);
+	const double SliderOriginalValue = (bIsVector3d ? Begin3DValues[InObjectIdx].GetValue()[InOriginalComponent] : Begin2DValues[InObjectIdx].GetValue()[InOriginalComponent]);
 	if (SliderOriginalValue == 0 && OldValue == 0)
 	{
 		if (bIsVector3d)
 		{
 			if (MinVectorClamp.IsSet())
 			{
-				NewValue = FMath::Max(NewValue, MinVectorClamp.GetValue()[ComponentIdx]);
+				InNewValue = FMath::Max(InNewValue, MinVectorClamp.GetValue()[InComponentIdx]);
 			}
 			if (MaxVectorClamp.IsSet())
 			{
-				NewValue = FMath::Min(NewValue, MaxVectorClamp.GetValue()[ComponentIdx]);
+				InNewValue = FMath::Min(InNewValue, MaxVectorClamp.GetValue()[InComponentIdx]);
 			}
 		}
 		else
 		{
 			if (MinVector2DClamp.IsSet())
 			{
-				NewValue = FMath::Max(NewValue, MinVector2DClamp.GetValue()[ComponentIdx]);
+				InNewValue = FMath::Max(InNewValue, MinVector2DClamp.GetValue()[InComponentIdx]);
 			}
 			if (MaxVector2DClamp.IsSet())
 			{
-				NewValue = FMath::Min(NewValue, MaxVector2DClamp.GetValue()[ComponentIdx]);
+				InNewValue = FMath::Min(InNewValue, MaxVector2DClamp.GetValue()[InComponentIdx]);
 			}
 		}
-		return NewValue;
+		return InNewValue;
 	}
-	return OldValue * Ratio;
+	return OldValue * InRatio;
 }
 
-double FAvaVectorPropertyTypeCustomization::MeshSizeToPixelSize(double MeshSize) const
+double FAvaVectorPropertyTypeCustomization::MeshSizeToPixelSize(double InMeshSize) const
 {
 	if (ViewportClient.IsValid())
 	{
 		double PixelSize;
 
-		if (FAvaEditorViewportUtils::MeshSizeToPixelSize(ViewportClient.Pin().ToSharedRef(), MeshSize, PixelSize))
+		if (FAvaEditorViewportUtils::MeshSizeToPixelSize(ViewportClient.Pin().ToSharedRef(), InMeshSize, PixelSize))
 		{
 			return PixelSize;
 		}
 	}
 
-	return MeshSize;
+	return InMeshSize;
 }
 
-double FAvaVectorPropertyTypeCustomization::PixelSizeToMeshSize(double PixelSize) const
+double FAvaVectorPropertyTypeCustomization::PixelSizeToMeshSize(double InPixelSize) const
 {
 	if (ViewportClient.IsValid())
 	{
 		double MeshSize;
 
-		if (FAvaEditorViewportUtils::PixelSizeToMeshSize(ViewportClient.Pin().ToSharedRef(), PixelSize, MeshSize))
+		if (FAvaEditorViewportUtils::PixelSizeToMeshSize(ViewportClient.Pin().ToSharedRef(), InPixelSize, MeshSize))
 		{
 			return MeshSize;
 		}
 	}
 
-	return PixelSize;
+	return InPixelSize;
 }
 
 #undef LOCTEXT_NAMESPACE
diff --git a/Engine/Plugins/Experimental/Avalanche/Source/AvalancheEditor/Private/DetailView/Customizations/AvaVectorPropertyTypeCustomization.h b/Engine/Plugins/Experimental/Avalanche/Source/AvalancheEditor/Private/DetailView/Customizations/AvaVectorPropertyTypeCustomization.h
--- a/Engine/Plugins/Experimental/Avalanche/Source/AvalancheEditor/Private/DetailView/Customizations/AvaVectorPropertyTypeCustomization.h
+++ b/Engine/Plugins/Experimental/Avalanche/Source/AvalancheEditor/Private/DetailView/Customizations/AvaVectorPropertyTypeCustomization.h
@@ -10,29 +10,19 @@
 #include "AvaVectorPropertyTypeCustomization.generated.h"
 
 class IAvaViewportClient;
+class IDetailChildrenBuilder;
 class SButton;
 class SComboButton;
+struct EVisibility;
 
+/** Used to detect if the type customization can be applied to a property */
 class FAvaVectorPropertyTypeIdentifier : public IPropertyTypeIdentifier
 {
 public:
 	virtual bool IsPropertyTypeCustomized(const IPropertyHandle& InPropertyHandle) const override
 	{
-		static const TArray<FName> PropertyMetaTags = { 
-			TEXT("AllowPreserveRatio"), 
-			TEXT("VectorRatioMode"), 
-			TEXT("Delta")
-		};
-
-		for (const FName& MetaTag : PropertyMetaTags)
-		{
-			if (InPropertyHandle.HasMetaData(MetaTag))
-			{
-				return true;
-			}
-		}
-
-		return false;
+		static const FName PropertyMetaTag = TEXT("MotionDesignVectorWidget");
+		return InPropertyHandle.HasMetaData(PropertyMetaTag);
 	}
 };
 
@@ -57,6 +47,7 @@ public:
 	using SNumericVectorInputBox2D = SNumericVectorInputBox<double, UE::Math::TVector2<double>, 2>;
 	using SNumericVectorInputBox3D = SNumericVectorInputBox<FVector::FReal, UE::Math::TVector<FVector::FReal>, 3>;
 
+	static constexpr const TCHAR* PropertyMetadata = TEXT("AllowPreserveRatio");
 	static constexpr uint8 MULTI_OBJECT_DEBOUNCE  = 3;
 	static constexpr uint8 SINGLE_OBJECT_DEBOUNCE = 2;
 	static constexpr uint8 INVALID_COMPONENT_IDX  = 5;
@@ -70,70 +61,81 @@ public:
 	{
 	}
 
-	// BEGIN IPropertyTypeCustomization interface
+	//~ Begin IPropertyTypeCustomization
 	virtual void CustomizeHeader(TSharedRef<IPropertyHandle> StructPropertyHandle,
 		FDetailWidgetRow& HeaderRow,
 		IPropertyTypeCustomizationUtils& StructCustomizationUtils) override;
 	virtual void CustomizeChildren(TSharedRef<IPropertyHandle> StructPropertyHandle,
 		IDetailChildrenBuilder& StructBuilder,
 		IPropertyTypeCustomizationUtils& StructCustomizationUtils) override;
-	// END IPropertyTypeCustomization interface
-	
+	//~ End IPropertyTypeCustomization
+
 protected:
+	TOptional<double> GetVectorComponent(const uint8 InComponent) const;
+	void SetVectorComponent(double InNewValue, const uint8 InComponent);
+	void SetVectorComponent(double InNewValue, ETextCommit::Type InCommitType, const uint8 InComponent);
+
+	void OnBeginSliderMovement();
+	void OnEndSliderMovement(double InNewValue);
+
+	TSharedRef<SWidget> OnGenerateRatioWidget(FName InRatioMode);
+	void OnRatioSelectionChanged(FName InRatioMode, ESelectInfo::Type InSelectInfo) const;
+	FName GetRatioCurrentItem() const;
+	EVisibility GetRatioWidgetVisibility() const;
+
+	const FSlateBrush* GetRatioModeBrush(const ERatioMode InMode) const;
+	FText GetRatioModeDisplayText(const ERatioMode InMode) const;
+
+	const FSlateBrush* GetCurrentRatioModeBrush() const;
+	FText GetCurrentRatioModeDisplayText() const;
+
+	ERatioMode GetRatioModeMetadata() const;
+	void SetRatioModeMetadata(ERatioMode InMode) const;
+
+	bool CanEditValue() const;
+	void InitVectorValuesForRatio();
+	void ResetVectorValuesForRatio();
+
+	void SetComponentValue(const double InNewValue, const uint8 InComponent, const EPropertyValueSetFlags::Type InFlags);
+
+	/** get the correct clamped ratio if a component hit min/max value */
+	double GetClampedRatioValueChange(const int32 InObjectIdx, const double InNewValue, const uint8 InComponent, const TArray<bool>& InPreserveRatios) const;
+
+	/** get the new clamped value if a component original value is zero since ratio * 0 = 0  */
+	double GetClampedComponentValue(const int32 InObjectIdx, double InNewValue, const double InRatio, const uint8 InComponentIdx, const uint8 InOriginalComponent);
+
+	/** special case for the pixel property only available in editor */
+	double MeshSizeToPixelSize(double InMeshSize) const;
+
+	double PixelSizeToMeshSize(double InPixelSize) const;
+
+private:
 	TWeakPtr<IAvaViewportClient> ViewportClient;
 
-	// used to close the dropdown menu
-	TSharedPtr<SComboButton> ComboButton;
-	
 	TSharedPtr<IPropertyHandle> VectorPropertyHandle;
 	TSharedPtr<IPropertyHandle> XPropertyHandle;
 	TSharedPtr<IPropertyHandle> YPropertyHandle;
 	TSharedPtr<IPropertyHandle> ZPropertyHandle;
 
-	// optional begin values to compute ratios change
+	/** optional begin values to compute ratios change */
 	TArray<TOptional<FVector>> Begin3DValues;
 	TArray<TOptional<FVector2D>> Begin2DValues;
-	
+
 	int32 SelectedObjectNum = 0;
 	uint8 DebounceValueSet = 0;
 	uint8 LastComponentValueSet = INDEX_NONE;
 	bool bMovingSlider = false;
 	bool bIsVector3d = false;
-	ERatioMode RatioMode = ERatioMode::None;
-	// specific case to handle that needs conversion
+
+	/** specific case to handle that needs conversion */
 	bool bPixelSizeProperty = false;
 
-	// optional clamp values
+	/** optional clamp values */
 	TOptional<FVector> MinVectorClamp;
 	TOptional<FVector> MaxVectorClamp;
 	TOptional<FVector2D> MinVector2DClamp;
 	TOptional<FVector2D> MaxVector2DClamp;
 
-	TOptional<double> GetVectorComponent(const uint8 Component) const;
-	void SetVectorComponent(double NewValue, const uint8 Component);
-	void SetVectorComponent(double NewValue, ETextCommit::Type CommitType, const uint8 Component);
-
-	void OnBeginSliderMovement();
-	void OnEndSliderMovement(double NewValue);
-
-	FReply OnComboButtonClicked(const ERatioMode NewMode);
-	const FSlateBrush* GetComboButtonBrush(const ERatioMode Mode) const;
-	FText GetComboButtonText(const ERatioMode Mode) const;
-	const FSlateBrush* GetCurrentComboButtonBrush() const;
-	FText GetCurrentComboButtonText() const;
-
-	bool CanEditValue() const;
-	void InitVectorValuesForRatio();
-	void ResetVectorValuesForRatio();
-
-	void SetComponentValue(const double NewValue, const uint8 Component, const EPropertyValueSetFlags::Type Flags);
-
-	// get the correct clamped ratio if a component hit min/max value
-	double GetClampedRatioValueChange(const int32 ObjectIdx, const double NewValue, const uint8 Component, const TArray<bool>& PreserveRatios) const;
-	// get the new clamped value if a component original value is zero since ratio * 0 = 0
-	double GetClampedComponentValue(const int32 ObjectIdx, double NewValue, const double Ratio, const uint8 ComponentIdx, const uint8 OriginalComponent);
-
-	// special case for the pixel property only available in editor
-	double MeshSizeToPixelSize(double MeshSize) const;
-	double PixelSizeToMeshSize(double PixelSize) const;
+	/** Ratio modes available for the property dropdown */
+	TArray<FName> RatioModes;
 };
diff --git a/Engine/Plugins/Experimental/Avalanche/Source/AvalancheShapes/Public/AvaShapeUVParameters.h b/Engine/Plugins/Experimental/Avalanche/Source/AvalancheShapes/Public/AvaShapeUVParameters.h
--- a/Engine/Plugins/Experimental/Avalanche/Source/AvalancheShapes/Public/AvaShapeUVParameters.h
+++ b/Engine/Plugins/Experimental/Avalanche/Source/AvalancheShapes/Public/AvaShapeUVParameters.h
@@ -76,7 +76,7 @@ protected:
 	UPROPERTY(EditInstanceOnly, BlueprintReadWrite, Category="Shape", meta=(ClampMin="0.0", ClampMax="1.0", EditCondition="AnchorPreset == EAvaAnchors::Custom", AllowPrivateAccess = "true"))
 	FVector2D Anchor;
 
-	UPROPERTY(EditInstanceOnly, BlueprintReadWrite, Category="Shape", meta=(ClampMin="0.01", AllowPreserveRatio, Delta="0.05", AllowPrivateAccess = "true"))
+	UPROPERTY(EditInstanceOnly, BlueprintReadWrite, Category="Shape", meta=(ClampMin="0.01", MotionDesignVectorWidget, AllowPreserveRatio="XY", Delta="0.05", AllowPrivateAccess = "true"))
 	FVector2D Scale;
 
 	UPROPERTY(EditInstanceOnly, BlueprintReadWrite, Category="Shape", meta=(Delta="0.05", AllowPrivateAccess = "true"))
diff --git a/Engine/Plugins/Experimental/Avalanche/Source/AvalancheShapes/Public/DynamicMeshes/AvaShape2DDynMeshBase.h b/Engine/Plugins/Experimental/Avalanche/Source/AvalancheShapes/Public/DynamicMeshes/AvaShape2DDynMeshBase.h
--- a/Engine/Plugins/Experimental/Avalanche/Source/AvalancheShapes/Public/DynamicMeshes/AvaShape2DDynMeshBase.h
+++ b/Engine/Plugins/Experimental/Avalanche/Source/AvalancheShapes/Public/DynamicMeshes/AvaShape2DDynMeshBase.h
@@ -108,11 +108,11 @@ protected:
 	bool bDoNotRecenterVertices;
 
 	// pixel size of the mesh, will only be available in editor
-	UPROPERTY(EditInstanceOnly, BlueprintReadWrite, Transient, Setter, Getter, Category="Shape", meta=(ClampMin="0.0", DisplayName="Pixel Size", DisplayAfter="SizeType", AllowPreserveRatio, EditCondition="bAllowEditSize && SizeType == ESizeType::Pixel", EditConditionHides, AllowPrivateAccess="true"))
+	UPROPERTY(EditInstanceOnly, BlueprintReadWrite, Transient, Setter, Getter, Category="Shape", meta=(ClampMin="0.0", DisplayName="Pixel Size", DisplayAfter="SizeType", MotionDesignVectorWidget, AllowPreserveRatio="XY", EditCondition="bAllowEditSize && SizeType == ESizeType::Pixel", EditConditionHides, AllowPrivateAccess="true"))
 	FVector2D PixelSize2D = FVector2D::ZeroVector;
 
 	// total size in 2D from 0 to mesh size and not origin
-	UPROPERTY(EditInstanceOnly, BlueprintReadWrite, Setter, Getter, Category="Shape", meta=(ClampMin="0.0", DisplayName="Mesh Size", DisplayAfter="SizeType", AllowPreserveRatio, Units="Centimeters", EditCondition="bAllowEditSize && SizeType == ESizeType::UnrealUnit", EditConditionHides, AllowPrivateAccess="true"))
+	UPROPERTY(EditInstanceOnly, BlueprintReadWrite, Setter, Getter, Category="Shape", meta=(ClampMin="0.0", DisplayName="Mesh Size", DisplayAfter="SizeType", MotionDesignVectorWidget, AllowPreserveRatio="XY", Units="Centimeters", EditCondition="bAllowEditSize && SizeType == ESizeType::UnrealUnit", EditConditionHides, AllowPrivateAccess="true"))
 	FVector2D Size2D = FVector2D::ZeroVector;
 
 	UPROPERTY(Transient)
diff --git a/Engine/Plugins/Experimental/Avalanche/Source/AvalancheShapes/Public/DynamicMeshes/AvaShape3DDynMeshBase.h b/Engine/Plugins/Experimental/Avalanche/Source/AvalancheShapes/Public/DynamicMeshes/AvaShape3DDynMeshBase.h
--- a/Engine/Plugins/Experimental/Avalanche/Source/AvalancheShapes/Public/DynamicMeshes/AvaShape3DDynMeshBase.h
+++ b/Engine/Plugins/Experimental/Avalanche/Source/AvalancheShapes/Public/DynamicMeshes/AvaShape3DDynMeshBase.h
@@ -83,13 +83,13 @@ protected:
 	void AddTriangle(FAvaShapeMesh& InMesh, int32 A, int32 B, int32 C);
 
 	// pixel size of the mesh, will only be available in editor
-	UPROPERTY(EditInstanceOnly, BlueprintReadWrite, Transient, Setter, Getter, Category="Shape", meta=(ClampMin="0.0", DisplayName="Pixel Size", DisplayAfter="SizeType", AllowPreserveRatio, EditCondition="bAllowEditSize && SizeType == ESizeType::Pixel", EditConditionHides, AllowPrivateAccess="true"))
+	UPROPERTY(EditInstanceOnly, BlueprintReadWrite, Transient, Setter, Getter, Category="Shape", meta=(ClampMin="0.0", DisplayName="Pixel Size", DisplayAfter="SizeType", MotionDesignVectorWidget, AllowPreserveRatio="XYZ", EditCondition="bAllowEditSize && SizeType == ESizeType::Pixel", EditConditionHides, AllowPrivateAccess="true"))
 	FVector PixelSize3D = FVector::ZeroVector;
 
 	/*
 	 * Corresponds to the total size from 0 to mesh size
 	 */
-	UPROPERTY(EditInstanceOnly, BlueprintReadWrite, Setter, Getter, Category="Shape", meta=(ClampMin="0.0", DisplayName="Mesh Size", DisplayAfter="SizeType", AllowPreserveRatio, Units="Centimeters", EditCondition="bAllowEditSize && SizeType == ESizeType::UnrealUnit", EditConditionHides, AllowPrivateAccess="true"))
+	UPROPERTY(EditInstanceOnly, BlueprintReadWrite, Setter, Getter, Category="Shape", meta=(ClampMin="0.0", DisplayName="Mesh Size", DisplayAfter="SizeType", MotionDesignVectorWidget, AllowPreserveRatio="XYZ", Units="Centimeters", EditCondition="bAllowEditSize && SizeType == ESizeType::UnrealUnit", EditConditionHides, AllowPrivateAccess="true"))
 	FVector Size3D = FVector::ZeroVector;
 
 	UPROPERTY(Transient)
diff --git a/Engine/Plugins/Experimental/ClonerEffector/Source/ClonerEffector/Public/Cloner/Extensions/CEClonerCollisionExtension.h b/Engine/Plugins/Experimental/ClonerEffector/Source/ClonerEffector/Public/Cloner/Extensions/CEClonerCollisionExtension.h
--- a/Engine/Plugins/Experimental/ClonerEffector/Source/ClonerEffector/Public/Cloner/Extensions/CEClonerCollisionExtension.h
+++ b/Engine/Plugins/Experimental/ClonerEffector/Source/ClonerEffector/Public/Cloner/Extensions/CEClonerCollisionExtension.h
@@ -147,7 +147,7 @@ protected:
 	int32 CollisionGridResolution = 32;
 
 	/** Size of the neighbor grid to detect collision between particles, particles outside this grid will not have collisions */
-	UPROPERTY(EditInstanceOnly, BlueprintReadWrite, Setter, Getter, Category="Collisions", meta=(ClampMin="0", AllowPreserveRatio, EditCondition="bParticleCollisionEnabled", EditConditionHides))
+	UPROPERTY(EditInstanceOnly, BlueprintReadWrite, Setter, Getter, Category="Collisions", meta=(ClampMin="0", MotionDesignVectorWidget, AllowPreserveRatio="XYZ", EditCondition="bParticleCollisionEnabled", EditConditionHides))
 	FVector CollisionGridSize = FVector(5000.f);
 
 	/** Collision radius calculation mode */
diff --git a/Engine/Plugins/Experimental/ClonerEffector/Source/ClonerEffector/Public/Cloner/Extensions/CEClonerRangeExtension.h b/Engine/Plugins/Experimental/ClonerEffector/Source/ClonerEffector/Public/Cloner/Extensions/CEClonerRangeExtension.h
--- a/Engine/Plugins/Experimental/ClonerEffector/Source/ClonerEffector/Public/Cloner/Extensions/CEClonerRangeExtension.h
+++ b/Engine/Plugins/Experimental/ClonerEffector/Source/ClonerEffector/Public/Cloner/Extensions/CEClonerRangeExtension.h
@@ -143,10 +143,10 @@ protected:
 	UPROPERTY(EditInstanceOnly, BlueprintReadWrite, Setter="SetRangeScaleUniform", Getter="GetRangeScaleUniform", DisplayName="ScaleUniformEnabled", Category="Range", meta=(EditCondition="bRangeEnabled", EditConditionHides))
 	bool bRangeScaleUniform = true;
 
-	UPROPERTY(EditInstanceOnly, BlueprintReadWrite, Setter="SetRangeScaleMin", Getter="GetRangeScaleMin", DisplayName="ScaleMin", Category="Range", meta=(AllowPreserveRatio, Delta="0.0001", ClampMin="0", EditCondition="bRangeEnabled && !bRangeScaleUniform", EditConditionHides))
+	UPROPERTY(EditInstanceOnly, BlueprintReadWrite, Setter="SetRangeScaleMin", Getter="GetRangeScaleMin", DisplayName="ScaleMin", Category="Range", meta=(MotionDesignVectorWidget, AllowPreserveRatio="XYZ", Delta="0.0001", ClampMin="0", EditCondition="bRangeEnabled && !bRangeScaleUniform", EditConditionHides))
 	FVector RangeScaleMin = FVector::OneVector;
 
-	UPROPERTY(EditInstanceOnly, BlueprintReadWrite, Setter="SetRangeScaleMax", Getter="GetRangeScaleMax", DisplayName="ScaleMax", Category="Range", meta=(AllowPreserveRatio, Delta="0.0001", ClampMin="0", EditCondition="bRangeEnabled && !bRangeScaleUniform", EditConditionHides))
+	UPROPERTY(EditInstanceOnly, BlueprintReadWrite, Setter="SetRangeScaleMax", Getter="GetRangeScaleMax", DisplayName="ScaleMax", Category="Range", meta=(MotionDesignVectorWidget, AllowPreserveRatio="XYZ", Delta="0.0001", ClampMin="0", EditCondition="bRangeEnabled && !bRangeScaleUniform", EditConditionHides))
 	FVector RangeScaleMax = FVector::OneVector;
 
 	UPROPERTY(EditInstanceOnly, BlueprintReadWrite, Setter="SetRangeScaleUniformMin", Getter="GetRangeScaleUniformMin", DisplayName="ScaleMin", Category="Range", meta=(Delta="0.0001", ClampMin="0", EditCondition="bRangeEnabled && bRangeScaleUniform", EditConditionHides))
diff --git a/Engine/Plugins/Experimental/ClonerEffector/Source/ClonerEffector/Public/Cloner/Extensions/CEClonerStepExtension.h b/Engine/Plugins/Experimental/ClonerEffector/Source/ClonerEffector/Public/Cloner/Extensions/CEClonerStepExtension.h
--- a/Engine/Plugins/Experimental/ClonerEffector/Source/ClonerEffector/Public/Cloner/Extensions/CEClonerStepExtension.h
+++ b/Engine/Plugins/Experimental/ClonerEffector/Source/ClonerEffector/Public/Cloner/Extensions/CEClonerStepExtension.h
@@ -83,7 +83,7 @@ protected:
 	FRotator DeltaStepRotation = FRotator(0.f);
 
 	/** Amount of scale difference between one step and the next one */
-	UPROPERTY(EditInstanceOnly, BlueprintReadWrite, Setter="SetDeltaStepScale", Getter="GetDeltaStepScale", DisplayName="Scale", Category="Step", meta=(AllowPreserveRatio, Delta="0.0001", EditCondition="bDeltaStepEnabled", EditConditionHides))
+	UPROPERTY(EditInstanceOnly, BlueprintReadWrite, Setter="SetDeltaStepScale", Getter="GetDeltaStepScale", DisplayName="Scale", Category="Step", meta=(MotionDesignVectorWidget, AllowPreserveRatio="XYZ", Delta="0.0001", EditCondition="bDeltaStepEnabled", EditConditionHides))
 	FVector DeltaStepScale = FVector(0.f);
 
 private:
diff --git a/Engine/Plugins/Experimental/ClonerEffector/Source/ClonerEffector/Public/Cloner/Layouts/CEClonerCircleLayout.h b/Engine/Plugins/Experimental/ClonerEffector/Source/ClonerEffector/Public/Cloner/Layouts/CEClonerCircleLayout.h
--- a/Engine/Plugins/Experimental/ClonerEffector/Source/ClonerEffector/Public/Cloner/Layouts/CEClonerCircleLayout.h
+++ b/Engine/Plugins/Experimental/ClonerEffector/Source/ClonerEffector/Public/Cloner/Layouts/CEClonerCircleLayout.h
@@ -126,7 +126,7 @@ protected:
 	UPROPERTY(EditInstanceOnly, BlueprintReadWrite, Setter="SetRotation", Getter="GetRotation", Category="Layout", meta=(EditCondition="Plane == ECEClonerPlane::Custom", EditConditionHides))
 	FRotator Rotation = FRotator(0.f, 0.f, 0.f);
 
-	UPROPERTY(EditInstanceOnly, BlueprintReadWrite, Setter="SetScale", Getter="GetScale", Category="Layout", meta=(ClampMin="0", AllowPreserveRatio, Delta="0.01"))
+	UPROPERTY(EditInstanceOnly, BlueprintReadWrite, Setter="SetScale", Getter="GetScale", Category="Layout", meta=(ClampMin="0", MotionDesignVectorWidget, AllowPreserveRatio="XYZ", Delta="0.01"))
 	FVector Scale = FVector(1.f, 1.f, 1.f);
 
 private:
diff --git a/Engine/Plugins/Experimental/ClonerEffector/Source/ClonerEffector/Public/Cloner/Layouts/CEClonerCylinderLayout.h b/Engine/Plugins/Experimental/ClonerEffector/Source/ClonerEffector/Public/Cloner/Layouts/CEClonerCylinderLayout.h
--- a/Engine/Plugins/Experimental/ClonerEffector/Source/ClonerEffector/Public/Cloner/Layouts/CEClonerCylinderLayout.h
+++ b/Engine/Plugins/Experimental/ClonerEffector/Source/ClonerEffector/Public/Cloner/Layouts/CEClonerCylinderLayout.h
@@ -150,7 +150,7 @@ protected:
 	UPROPERTY(EditInstanceOnly, BlueprintReadWrite, Setter="SetRotation", Getter="GetRotation", Category="Layout", meta=(EditCondition="Plane == ECEClonerPlane::Custom", EditConditionHides))
 	FRotator Rotation = FRotator(0.f, 0.f, 0.f);
 
-	UPROPERTY(EditInstanceOnly, BlueprintReadWrite, Setter="SetScale", Getter="GetScale", Category="Layout", meta=(ClampMin="0", AllowPreserveRatio, Delta="0.01"))
+	UPROPERTY(EditInstanceOnly, BlueprintReadWrite, Setter="SetScale", Getter="GetScale", Category="Layout", meta=(ClampMin="0", MotionDesignVectorWidget, AllowPreserveRatio="XYZ", Delta="0.01"))
 	FVector Scale = FVector(1.f, 1.f, 1.f);
 
 private:
diff --git a/Engine/Plugins/Experimental/ClonerEffector/Source/ClonerEffector/Public/Cloner/Layouts/CEClonerSphereRandomLayout.h b/Engine/Plugins/Experimental/ClonerEffector/Source/ClonerEffector/Public/Cloner/Layouts/CEClonerSphereRandomLayout.h
--- a/Engine/Plugins/Experimental/ClonerEffector/Source/ClonerEffector/Public/Cloner/Layouts/CEClonerSphereRandomLayout.h
+++ b/Engine/Plugins/Experimental/ClonerEffector/Source/ClonerEffector/Public/Cloner/Layouts/CEClonerSphereRandomLayout.h
@@ -125,7 +125,7 @@ protected:
 	UPROPERTY(EditInstanceOnly, BlueprintReadWrite, Setter="SetRotation", Getter="GetRotation", Category="Layout")
 	FRotator Rotation = FRotator(0.f, 0.f, 0.f);
 
-	UPROPERTY(EditInstanceOnly, BlueprintReadWrite, Setter="SetScale", Getter="GetScale", Category="Layout", meta=(ClampMin="0", AllowPreserveRatio, Delta="0.01"))
+	UPROPERTY(EditInstanceOnly, BlueprintReadWrite, Setter="SetScale", Getter="GetScale", Category="Layout", meta=(ClampMin="0", MotionDesignVectorWidget, AllowPreserveRatio="XYZ", Delta="0.01"))
 	FVector Scale = FVector(1.f, 1.f, 1.f);
 
 private:
diff --git a/Engine/Plugins/Experimental/ClonerEffector/Source/ClonerEffector/Public/Cloner/Layouts/CEClonerSphereUniformLayout.h b/Engine/Plugins/Experimental/ClonerEffector/Source/ClonerEffector/Public/Cloner/Layouts/CEClonerSphereUniformLayout.h
--- a/Engine/Plugins/Experimental/ClonerEffector/Source/ClonerEffector/Public/Cloner/Layouts/CEClonerSphereUniformLayout.h
+++ b/Engine/Plugins/Experimental/ClonerEffector/Source/ClonerEffector/Public/Cloner/Layouts/CEClonerSphereUniformLayout.h
@@ -101,7 +101,7 @@ protected:
 	UPROPERTY(EditInstanceOnly, BlueprintReadWrite, Setter="SetRotation", Getter="GetRotation", Category="Layout")
 	FRotator Rotation = FRotator(0.f, 0.f, 0.f);
 
-	UPROPERTY(EditInstanceOnly, BlueprintReadWrite, Setter="SetScale", Getter="GetScale", Category="Layout", meta=(ClampMin="0", AllowPreserveRatio, Delta="0.01"))
+	UPROPERTY(EditInstanceOnly, BlueprintReadWrite, Setter="SetScale", Getter="GetScale", Category="Layout", meta=(ClampMin="0", MotionDesignVectorWidget, AllowPreserveRatio="XYZ", Delta="0.01"))
 	FVector Scale = FVector(1.f, 1.f, 1.f);
 
 private:
diff --git a/Engine/Plugins/Experimental/ClonerEffector/Source/ClonerEffector/Public/Effector/Modes/CEEffectorNoiseMode.h b/Engine/Plugins/Experimental/ClonerEffector/Source/ClonerEffector/Public/Effector/Modes/CEEffectorNoiseMode.h
--- a/Engine/Plugins/Experimental/ClonerEffector/Source/ClonerEffector/Public/Effector/Modes/CEEffectorNoiseMode.h
+++ b/Engine/Plugins/Experimental/ClonerEffector/Source/ClonerEffector/Public/Effector/Modes/CEEffectorNoiseMode.h
@@ -84,7 +84,7 @@ protected:
 	FRotator RotationStrength = FRotator::ZeroRotator;
 
 	/** Amplitude of the noise field for scale */
-	UPROPERTY(EditInstanceOnly, BlueprintReadWrite, Setter, Getter, Category="Mode", meta=(ClampMin="0", AllowPreserveRatio, Delta="0.01"))
+	UPROPERTY(EditInstanceOnly, BlueprintReadWrite, Setter, Getter, Category="Mode", meta=(ClampMin="0", MotionDesignVectorWidget, AllowPreserveRatio="XYZ", Delta="0.01"))
 	FVector ScaleStrength = FVector::OneVector;
 
 	/** Panning to offset the noise field sampling */
diff --git a/Engine/Plugins/Experimental/ClonerEffector/Source/ClonerEffector/Public/Effector/Modes/CEEffectorOffsetMode.h b/Engine/Plugins/Experimental/ClonerEffector/Source/ClonerEffector/Public/Effector/Modes/CEEffectorOffsetMode.h
--- a/Engine/Plugins/Experimental/ClonerEffector/Source/ClonerEffector/Public/Effector/Modes/CEEffectorOffsetMode.h
+++ b/Engine/Plugins/Experimental/ClonerEffector/Source/ClonerEffector/Public/Effector/Modes/CEEffectorOffsetMode.h
@@ -66,7 +66,7 @@ protected:
 	FRotator Rotation = FRotator::ZeroRotator;
 
 	/** Scale applied on affected clones */
-	UPROPERTY(EditInstanceOnly, BlueprintReadWrite, Setter, Getter, Category="Mode", meta=(ClampMin="0", AllowPreserveRatio, Delta="0.01"))
+	UPROPERTY(EditInstanceOnly, BlueprintReadWrite, Setter, Getter, Category="Mode", meta=(ClampMin="0", MotionDesignVectorWidget, AllowPreserveRatio="XYZ", Delta="0.01"))
 	FVector Scale = FVector::OneVector;
 
 private:
diff --git a/Engine/Plugins/Experimental/ClonerEffector/Source/ClonerEffector/Public/Effector/Modes/CEEffectorPushMode.h b/Engine/Plugins/Experimental/ClonerEffector/Source/ClonerEffector/Public/Effector/Modes/CEEffectorPushMode.h
--- a/Engine/Plugins/Experimental/ClonerEffector/Source/ClonerEffector/Public/Effector/Modes/CEEffectorPushMode.h
+++ b/Engine/Plugins/Experimental/ClonerEffector/Source/ClonerEffector/Public/Effector/Modes/CEEffectorPushMode.h
@@ -49,7 +49,7 @@ protected:
 	//~ End UObject
 
 	/** Strength of the push effect */
-	UPROPERTY(EditInstanceOnly, BlueprintReadWrite, Setter, Getter, Category="Mode", meta=(AllowPreserveRatio))
+	UPROPERTY(EditInstanceOnly, BlueprintReadWrite, Setter, Getter, Category="Mode", meta=(MotionDesignVectorWidget, AllowPreserveRatio="XYZ"))
 	FVector PushStrength = FVector(100.f);
 
 	/** Relative direction computed for the push effect on each clone */
diff --git a/Engine/Plugins/Experimental/ClonerEffector/Source/ClonerEffector/Public/Effector/Types/CEEffectorBoxType.h b/Engine/Plugins/Experimental/ClonerEffector/Source/ClonerEffector/Public/Effector/Types/CEEffectorBoxType.h
--- a/Engine/Plugins/Experimental/ClonerEffector/Source/ClonerEffector/Public/Effector/Types/CEEffectorBoxType.h
+++ b/Engine/Plugins/Experimental/ClonerEffector/Source/ClonerEffector/Public/Effector/Types/CEEffectorBoxType.h
@@ -54,11 +54,11 @@ protected:
 	//~ End UObject
 
 	/** Inner extent of box, all clones inside will be affected with a maximum weight */
-	UPROPERTY(EditInstanceOnly, BlueprintReadWrite, Setter, Getter, Category="Shape", meta=(ClampMin="0", AllowPreserveRatio))
+	UPROPERTY(EditInstanceOnly, BlueprintReadWrite, Setter, Getter, Category="Shape", meta=(ClampMin="0", MotionDesignVectorWidget, AllowPreserveRatio="XYZ"))
 	FVector InnerExtent = FVector(50.f);
 
 	/** Outer extent of box, all clones outside will not be affected */
-	UPROPERTY(EditInstanceOnly, BlueprintReadWrite, Setter, Getter, Category="Shape", meta=(ClampMin="0", AllowPreserveRatio))
+	UPROPERTY(EditInstanceOnly, BlueprintReadWrite, Setter, Getter, Category="Shape", meta=(ClampMin="0", MotionDesignVectorWidget, AllowPreserveRatio="XYZ"))
 	FVector OuterExtent = FVector(200.f);
 
 private:
