Motion Design Transition: Fixed issue with importing tree in Advanced Mode causing the Tree Details to not refresh to the new imported tree.

diff --git a/Engine/Plugins/Experimental/Avalanche/Source/AvalancheTransitionEditor/Private/TabFactories/AvaTransitionTreeDetailsTabFactory.cpp b/Engine/Plugins/Experimental/Avalanche/Source/AvalancheTransitionEditor/Private/TabFactories/AvaTransitionTreeDetailsTabFactory.cpp
--- a/Engine/Plugins/Experimental/Avalanche/Source/AvalancheTransitionEditor/Private/TabFactories/AvaTransitionTreeDetailsTabFactory.cpp
+++ b/Engine/Plugins/Experimental/Avalanche/Source/AvalancheTransitionEditor/Private/TabFactories/AvaTransitionTreeDetailsTabFactory.cpp
@@ -1,14 +1,7 @@
 // Copyright Epic Games, Inc. All Rights Reserved.
 
 #include "AvaTransitionTreeDetailsTabFactory.h"
-#include "AvaTransitionTreeEditorData.h"
-#include "AvaTypeSharedPointer.h"
-#include "Customizations/AvaTransitionTreeEditorDataCustomization.h"
-#include "IDetailsView.h"
-#include "Modules/ModuleManager.h"
-#include "PropertyEditorModule.h"
 #include "ViewModels/AvaTransitionEditorViewModel.h"
-#include "ViewModels/AvaTransitionViewModelSharedData.h"
 #include "Widgets/SNullWidget.h"
 
 #define LOCTEXT_NAMESPACE "AvaTransitionTreeDetailsTabFactory"
@@ -40,30 +33,7 @@ TSharedRef<SWidget> FAvaTransitionTreeDetailsTabFactory::CreateTabBody(const FWo
 		return SNullWidget::NullWidget;
 	}
 
-	UAvaTransitionTreeEditorData* EditorData = EditorViewModel->GetEditorData();
-	if (!ensure(EditorData))
-	{
-		return SNullWidget::NullWidget;
-	}
-
-	FDetailsViewArgs DetailsViewArgs;
-	DetailsViewArgs.NameAreaSettings = FDetailsViewArgs::HideNameArea;
-
-	FPropertyEditorModule& PropertyEditorModule = FModuleManager::GetModuleChecked<FPropertyEditorModule>("PropertyEditor");
-	TSharedRef<IDetailsView> DetailsView = PropertyEditorModule.CreateDetailView(DetailsViewArgs);
-
-	DetailsView->RegisterInstancedCustomPropertyLayout(UStateTreeEditorData::StaticClass()
-		, FOnGetDetailCustomizationInstance::CreateStatic(&FAvaTransitionTreeEditorDataCustomization::MakeInstance));
-
-	// Read-only
-	if (EditorViewModel->GetSharedData()->IsReadOnly())
-	{
-		DetailsView->SetIsPropertyEditingEnabledDelegate(FIsPropertyEditingEnabled::CreateLambda([]{ return false; }));
-	}
-
-	DetailsView->SetObject(EditorData);
-
-	return DetailsView;
+	return EditorViewModel->GetTreeDetails();
 }
 
 #undef LOCTEXT_NAMESPACE
diff --git a/Engine/Plugins/Experimental/Avalanche/Source/AvalancheTransitionEditor/Private/TabFactories/AvaTransitionTreeTabFactory.cpp b/Engine/Plugins/Experimental/Avalanche/Source/AvalancheTransitionEditor/Private/TabFactories/AvaTransitionTreeTabFactory.cpp
--- a/Engine/Plugins/Experimental/Avalanche/Source/AvalancheTransitionEditor/Private/TabFactories/AvaTransitionTreeTabFactory.cpp
+++ b/Engine/Plugins/Experimental/Avalanche/Source/AvalancheTransitionEditor/Private/TabFactories/AvaTransitionTreeTabFactory.cpp
@@ -3,7 +3,6 @@
 #include "AvaTransitionTreeTabFactory.h"
 #include "AvaTransitionTree.h"
 #include "AvaTypeSharedPointer.h"
-#include "Extensions/IAvaTransitionWidgetExtension.h"
 #include "Styling/SlateIconFinder.h"
 #include "ViewModels/AvaTransitionEditorViewModel.h"
 #include "Widgets/SNullWidget.h"
@@ -31,13 +30,13 @@ TSharedRef<SWidget> FAvaTransitionTreeTabFactory::CreateTabBody(const FWorkflowT
 		return SNullWidget::NullWidget;
 	}
 
-	TSharedPtr<IAvaTransitionWidgetExtension> WidgetExtension = UE::AvaCore::CastSharedPtr<IAvaTransitionWidgetExtension>(Editor->GetEditorViewModel());
-	if (!ensure(WidgetExtension.IsValid()))
+	TSharedPtr<FAvaTransitionEditorViewModel> EditorViewModel = Editor->GetEditorViewModel();
+	if (!ensure(EditorViewModel.IsValid()))
 	{
 		return SNullWidget::NullWidget;
 	}
 
-	return WidgetExtension->CreateWidget();
+	return EditorViewModel->GetTreeWidget();
 }
 
 #undef LOCTEXT_NAMESPACE
diff --git a/Engine/Plugins/Experimental/Avalanche/Source/AvalancheTransitionEditor/Private/ViewModels/AvaTransitionEditorViewModel.cpp b/Engine/Plugins/Experimental/Avalanche/Source/AvalancheTransitionEditor/Private/ViewModels/AvaTransitionEditorViewModel.cpp
--- a/Engine/Plugins/Experimental/Avalanche/Source/AvalancheTransitionEditor/Private/ViewModels/AvaTransitionEditorViewModel.cpp
+++ b/Engine/Plugins/Experimental/Avalanche/Source/AvalancheTransitionEditor/Private/ViewModels/AvaTransitionEditorViewModel.cpp
@@ -21,6 +21,7 @@
 #include "StateTreeDelegates.h"
 #include "StateTreeEditorSettings.h"
 #include "TabFactories/AvaTransitionCompilerResultsTabFactory.h"
+#include "Views/SAvaTransitionTreeDetails.h"
 #include "Views/SAvaTransitionTreeView.h"
 #include "Widgets/Layout/SScrollBar.h"
 #include "Widgets/Layout/SScrollBox.h"
@@ -131,6 +132,51 @@ TSharedRef<FAvaTransitionSelection> FAvaTransitionEditorViewModel::GetSelection(
 	return GetSharedData()->GetSelection();
 }
 
+TSharedRef<SWidget> FAvaTransitionEditorViewModel::GetTreeWidget()
+{
+	if (!TreeView.IsValid())
+	{
+		TreeView = SNew(SAvaTransitionTreeView, SharedThis(this));
+	}
+
+	return SNew(SVerticalBox)
+		+ SVerticalBox::Slot()
+		.AutoHeight()
+		[
+			Toolbar->GenerateTreeToolbarWidget()
+		]
+		+ SVerticalBox::Slot()
+		.FillHeight(1.f)
+		[
+			SNew(SHorizontalBox)
+			+ SHorizontalBox::Slot()
+			.FillWidth(1.f)
+			[
+				SNew(SScrollBox)
+				.Orientation(Orient_Horizontal)
+				+ SScrollBox::Slot()
+				.FillSize(1.f)
+				[
+					TreeView.ToSharedRef()
+				]
+			]
+			+ SHorizontalBox::Slot()
+			.AutoWidth()
+			[
+				TreeView->GetVerticalScrollbar()
+			]
+		];
+}
+
+TSharedRef<SWidget> FAvaTransitionEditorViewModel::GetTreeDetails()
+{
+	if (!TreeDetails.IsValid())
+	{
+		TreeDetails = SNew(SAvaTransitionTreeDetails, SharedThis(this));
+	}
+	return TreeDetails.ToSharedRef();
+}
+
 void FAvaTransitionEditorViewModel::OnInitialize()
 {
 	FAvaTransitionViewModel::OnInitialize();
@@ -152,13 +198,17 @@ void FAvaTransitionEditorViewModel::OnInitialize()
 	{
 		ViewModelSharedData->SetReadOnly(Editor->IsReadOnly());
 	}
-
-	TreeView = SNew(SAvaTransitionTreeView, This);
 }
 
 void FAvaTransitionEditorViewModel::PostRefresh()
 {
 	FAvaTransitionViewModel::PostRefresh();
+
+	if (TreeDetails.IsValid())
+	{
+		TreeDetails->Refresh();
+	}
+
 	RefreshTreeView();
 	UpdateTree();
 }
@@ -199,38 +249,6 @@ void FAvaTransitionEditorViewModel::PostUndo(bool bInSuccess)
 	Refresh();
 }
 
-TSharedRef<SWidget> FAvaTransitionEditorViewModel::CreateWidget()
-{
-	check(TreeView.IsValid());
-	return SNew(SVerticalBox)
-		+ SVerticalBox::Slot()
-		.AutoHeight()
-		[
-			Toolbar->GenerateTreeToolbarWidget()
-		]
-		+ SVerticalBox::Slot()
-		.FillHeight(1.f)
-		[
-			SNew(SHorizontalBox)
-			+ SHorizontalBox::Slot()
-			.FillWidth(1.f)
-			[
-				SNew(SScrollBox)
-				.Orientation(Orient_Horizontal)
-				+ SScrollBox::Slot()
-				.FillSize(1.f)
-				[
-					TreeView.ToSharedRef()
-				]
-			]
-			+ SHorizontalBox::Slot()
-			.AutoWidth()
-			[
-				TreeView->GetVerticalScrollbar()
-			]
-		];
-}
-
 UObject* FAvaTransitionEditorViewModel::GetObject() const
 {
 	return GetEditorData();
diff --git a/Engine/Plugins/Experimental/Avalanche/Source/AvalancheTransitionEditor/Private/ViewModels/AvaTransitionEditorViewModel.h b/Engine/Plugins/Experimental/Avalanche/Source/AvalancheTransitionEditor/Private/ViewModels/AvaTransitionEditorViewModel.h
--- a/Engine/Plugins/Experimental/Avalanche/Source/AvalancheTransitionEditor/Private/ViewModels/AvaTransitionEditorViewModel.h
+++ b/Engine/Plugins/Experimental/Avalanche/Source/AvalancheTransitionEditor/Private/ViewModels/AvaTransitionEditorViewModel.h
@@ -17,6 +17,7 @@ class FAvaTransitionTreeContextMenu;
 class FUICommandList;
 class IAvaTransitionSelectableExtension;
 class IMessageToken;
+class SAvaTransitionTreeDetails;
 class SAvaTransitionTreeView;
 class UAvaTransitionTree;
 class UAvaTransitionTreeEditorData;
@@ -25,10 +26,10 @@ class UStateTreeState;
 class UToolMenu;
 
 /** View Model for the Tree Editor Data */
-class FAvaTransitionEditorViewModel : public FAvaTransitionViewModel, public FSelfRegisteringEditorUndoClient, public IAvaTransitionWidgetExtension, public IAvaTransitionObjectExtension
+class FAvaTransitionEditorViewModel : public FAvaTransitionViewModel, public FSelfRegisteringEditorUndoClient, public IAvaTransitionObjectExtension
 {
 public:
-	UE_AVA_INHERITS(FAvaTransitionEditorViewModel, FAvaTransitionViewModel, IAvaTransitionWidgetExtension, IAvaTransitionObjectExtension)
+	UE_AVA_INHERITS(FAvaTransitionEditorViewModel, FAvaTransitionViewModel, IAvaTransitionObjectExtension)
 
 	explicit FAvaTransitionEditorViewModel(UAvaTransitionTree* InTransitionTree, const TSharedPtr<FAvaTransitionEditor>& InEditor);
 
@@ -71,6 +72,10 @@ public:
 		return ContextMenu;
 	}
 
+	TSharedRef<SWidget> GetTreeWidget();
+
+	TSharedRef<SWidget> GetTreeDetails();
+
 	//~ Begin FAvaTransitionViewModel
 	virtual void OnInitialize() override;
 	virtual void PostRefresh() override;
@@ -82,10 +87,6 @@ public:
 	virtual void PostUndo(bool bInSuccess) override;
 	//~ End FEditorUndoClient
 
-	//~ Begin IAvaTransitionWidgetExtension
-	virtual TSharedRef<SWidget> CreateWidget() override;
-	//~ End IAvaTransitionWidgetExtension
-
 	//~ Begin IAvaTransitionObjectExtension
 	virtual UObject* GetObject() const override;
 	virtual void OnPropertiesChanged(const FPropertyChangedEvent& InPropertyChangedEvent) override {}
@@ -120,5 +121,7 @@ private:
 
 	TSharedPtr<SAvaTransitionTreeView> TreeView;
 
+	TSharedPtr<SAvaTransitionTreeDetails> TreeDetails;
+
 	TArray<TSharedRef<FAvaTransitionActions>> Actions;
 };
diff --git a/Engine/Plugins/Experimental/Avalanche/Source/AvalancheTransitionEditor/Private/Views/SAvaTransitionTreeDetails.cpp b/Engine/Plugins/Experimental/Avalanche/Source/AvalancheTransitionEditor/Private/Views/SAvaTransitionTreeDetails.cpp
new file mode 100644
--- /dev/null
+++ b/Engine/Plugins/Experimental/Avalanche/Source/AvalancheTransitionEditor/Private/Views/SAvaTransitionTreeDetails.cpp
@@ -0,0 +1,55 @@
+// Copyright Epic Games, Inc. All Rights Reserved.
+
+#include "SAvaTransitionTreeDetails.h"
+#include "AvaTransitionTreeEditorData.h"
+#include "Customizations/AvaTransitionTreeEditorDataCustomization.h"
+#include "IDetailsView.h"
+#include "Modules/ModuleManager.h"
+#include "PropertyEditorModule.h"
+#include "ViewModels/AvaTransitionEditorViewModel.h"
+#include "ViewModels/AvaTransitionViewModelSharedData.h"
+
+void SAvaTransitionTreeDetails::Construct(const FArguments& InArgs, const TSharedRef<FAvaTransitionEditorViewModel>& InEditorViewModel)
+{
+	EditorViewModelWeak = InEditorViewModel;
+
+	FDetailsViewArgs DetailsViewArgs;
+	DetailsViewArgs.NameAreaSettings = FDetailsViewArgs::HideNameArea;
+
+	FPropertyEditorModule& PropertyEditorModule = FModuleManager::GetModuleChecked<FPropertyEditorModule>("PropertyEditor");
+
+	DetailsView = PropertyEditorModule.CreateDetailView(DetailsViewArgs);
+	DetailsView->RegisterInstancedCustomPropertyLayout(UStateTreeEditorData::StaticClass()
+		, FOnGetDetailCustomizationInstance::CreateStatic(&FAvaTransitionTreeEditorDataCustomization::MakeInstance));
+
+	// Read-only
+	if (InEditorViewModel->GetSharedData()->IsReadOnly())
+	{
+		DetailsView->SetIsPropertyEditingEnabledDelegate(FIsPropertyEditingEnabled::CreateLambda([]{ return false; }));
+	}
+
+	ChildSlot
+	[
+		DetailsView.ToSharedRef()
+	];
+
+	Refresh();
+}
+
+void SAvaTransitionTreeDetails::Refresh()
+{
+	TSharedPtr<FAvaTransitionEditorViewModel> EditorViewModel = EditorViewModelWeak.Pin();
+	if (!EditorViewModel.IsValid())
+	{
+		return;
+	}
+
+	UAvaTransitionTreeEditorData* EditorData = EditorViewModel->GetEditorData();
+	if (!EditorData)
+	{
+		return;
+	}
+
+	check(DetailsView.IsValid());
+	DetailsView->SetObject(EditorData);
+}
diff --git a/Engine/Plugins/Experimental/Avalanche/Source/AvalancheTransitionEditor/Private/Views/SAvaTransitionTreeDetails.h b/Engine/Plugins/Experimental/Avalanche/Source/AvalancheTransitionEditor/Private/Views/SAvaTransitionTreeDetails.h
new file mode 100644
--- /dev/null
+++ b/Engine/Plugins/Experimental/Avalanche/Source/AvalancheTransitionEditor/Private/Views/SAvaTransitionTreeDetails.h
@@ -0,0 +1,24 @@
+// Copyright Epic Games, Inc. All Rights Reserved.
+
+#pragma once
+
+#include "Widgets/SCompoundWidget.h"
+
+class FAvaTransitionEditorViewModel;
+class IDetailsView;
+
+class SAvaTransitionTreeDetails : public SCompoundWidget
+{
+public:
+	SLATE_BEGIN_ARGS(SAvaTransitionTreeDetails){}
+	SLATE_END_ARGS()
+
+	void Construct(const FArguments& InArgs, const TSharedRef<FAvaTransitionEditorViewModel>& InEditorViewModel);
+
+	void Refresh();
+
+private:
+	TWeakPtr<FAvaTransitionEditorViewModel> EditorViewModelWeak;
+
+	TSharedPtr<IDetailsView> DetailsView;
+};
