More non-unity compile error fixes

diff --git a/Engine/Plugins/Experimental/Avalanche/Source/AvalancheSequencer/Private/AvaSequencer.h b/Engine/Plugins/Experimental/Avalanche/Source/AvalancheSequencer/Private/AvaSequencer.h
--- a/Engine/Plugins/Experimental/Avalanche/Source/AvalancheSequencer/Private/AvaSequencer.h
+++ b/Engine/Plugins/Experimental/Avalanche/Source/AvalancheSequencer/Private/AvaSequencer.h
@@ -7,6 +7,7 @@
 #include "Framework/Commands/UIAction.h"
 #include "IAvaSequencer.h"
 #include "IAvaSequencerProvider.h"
+#include "MovieSceneSequenceID.h"
 
 class AActor;
 class FAvaEaseCurveTool;
diff --git a/Engine/Plugins/Experimental/Avalanche/Source/AvalancheSequencer/Private/SequenceTree/AvaSequenceItem.h b/Engine/Plugins/Experimental/Avalanche/Source/AvalancheSequencer/Private/SequenceTree/AvaSequenceItem.h
--- a/Engine/Plugins/Experimental/Avalanche/Source/AvalancheSequencer/Private/SequenceTree/AvaSequenceItem.h
+++ b/Engine/Plugins/Experimental/Avalanche/Source/AvalancheSequencer/Private/SequenceTree/AvaSequenceItem.h
@@ -3,6 +3,8 @@
 #pragma once
 
 #include "SequenceTree/IAvaSequenceItem.h"
+#include "UObject/WeakObjectPtr.h"
+#include "UObject/WeakObjectPtrTemplates.h"
 
 class FAvaSequencer;
 
diff --git a/Engine/Source/Editor/Sequencer/Private/SSequencerTimePanel.h b/Engine/Source/Editor/Sequencer/Private/SSequencerTimePanel.h
--- a/Engine/Source/Editor/Sequencer/Private/SSequencerTimePanel.h
+++ b/Engine/Source/Editor/Sequencer/Private/SSequencerTimePanel.h
@@ -10,6 +10,7 @@
 #include "Widgets/SBoxPanel.h"
 #include "MovieSceneFwd.h"
 #include "FrameNumberDisplayFormat.h"
+#include "Misc/FrameRate.h"
 
 class FSequencer;
 class UMovieScene;
