[Motion Design] Rundown Server API: Importing rundown from json.

4 new commands:
- create rundown can be created transient in game mode.
- delete rundown (restricted to transient only in game mode)
- import rundown from json, either from a json file on the server or json data from client.
- export rundown to json, either to a json file on the server or json data to client.

Tests:
- Using the rundown test client (see QAMotionDesign project)
- Tested create, import, export and delete rundown commands in editor mode and game build.
- Tested (for regression) import/export in content browser.

diff --git a/Engine/Plugins/Experimental/Avalanche/Source/AvalancheMedia/Internal/Rundown/AvaRundownMessages.h b/Engine/Plugins/Experimental/Avalanche/Source/AvalancheMedia/Internal/Rundown/AvaRundownMessages.h
--- a/Engine/Plugins/Experimental/Avalanche/Source/AvalancheMedia/Internal/Rundown/AvaRundownMessages.h
+++ b/Engine/Plugins/Experimental/Avalanche/Source/AvalancheMedia/Internal/Rundown/AvaRundownMessages.h
@@ -249,6 +249,112 @@ public:
 	FString Rundown;
 };
 
+/**
+ * Request to create a new rundown asset.
+ *
+ * The full package name is going to be: [PackagePath]/[AssetName] 
+ * The full asset path is going to be: [PackagePath]/[AssetName].[AssetName]
+ * For all other requests, the rundown reference is the full asset path.
+ */
+USTRUCT()
+struct FAvaRundownCreateRundown : public FAvaRundownMsgBase
+{
+	GENERATED_BODY()
+
+public:
+	/** Package path (excluding the package name) */
+	UPROPERTY()
+	FString PackagePath;
+
+	/** Asset Name. */
+	UPROPERTY()
+	FString AssetName;
+
+	/**
+	 * Create the rundown as a transient object.
+	 * @remark For game builds, the created rundown will always be transient, regardless of this flag. 
+	 */
+	UPROPERTY()
+	bool bTransient = true;
+};
+
+/**
+ * Request a previously created rundown to be deleted or at least no longer managed (if transient only). 
+ */
+USTRUCT()
+struct FAvaRundownDeleteRundown : public FAvaRundownMsgBase
+{
+	GENERATED_BODY()
+
+public:
+	/** Rundown asset path: [PackagePath]/[AssetName].[AssetName] */
+	UPROPERTY()
+	FString Rundown;
+};
+
+/**
+ * Import rundown from json data or file.
+ */
+USTRUCT()
+struct FAvaRundownImportRundown : public FAvaRundownMsgBase
+{
+	GENERATED_BODY()
+
+public:
+	/** Rundown asset path: [PackagePath]/[AssetName].[AssetName] */
+	UPROPERTY()
+	FString Rundown;
+
+	/**
+	 * If specified, this is a server local path to a json file from which the rundown will be imported.
+	 */
+	UPROPERTY()
+	FString RundownFile;
+
+	/**
+	 * If specified, json data containing the rundown to import.
+	 */
+	UPROPERTY()
+	FString RundownData;
+};
+
+/**
+ * Export a rundown to json data or file.
+ * This command is supported in game build.
+ */
+USTRUCT()
+struct FAvaRundownExportRundown : public FAvaRundownMsgBase
+{
+	GENERATED_BODY()
+
+public:
+	/** Rundown asset path: [PackagePath]/[AssetName].[AssetName] */
+	UPROPERTY()
+	FString Rundown;
+
+	/** Optional path to a server local file where the rundown will be saved. */
+	UPROPERTY()
+	FString RundownFile;
+};
+
+/**
+ * Server reply to FAvaRundownExportRundown containing the exported rundown.
+ */
+USTRUCT()
+struct FAvaRundownExportedRundown : public FAvaRundownMsgBase
+{
+	GENERATED_BODY()
+
+public:
+	/** Rundown asset path: [PackagePath]/[AssetName].[AssetName] */
+	UPROPERTY()
+	FString Rundown;
+
+	/** Exported rundown in json format. */
+	UPROPERTY()
+	FString RundownData;
+};
+
 /**
  * Request that the given rundown be saved to disk.
  * The rundown asset must have been loaded, either by an edit command
diff --git a/Engine/Plugins/Experimental/Avalanche/Source/AvalancheMedia/Internal/Rundown/AvaRundownSerializationUtils.h b/Engine/Plugins/Experimental/Avalanche/Source/AvalancheMedia/Internal/Rundown/AvaRundownSerializationUtils.h
new file mode 100644
--- /dev/null
+++ b/Engine/Plugins/Experimental/Avalanche/Source/AvalancheMedia/Internal/Rundown/AvaRundownSerializationUtils.h
@@ -0,0 +1,18 @@
+// Copyright Epic Games, Inc. All Rights Reserved.
+
+#pragma once
+
+#include "HAL/Platform.h"
+
+class FArchive;
+class FText;
+class UAvaRundown;
+
+namespace UE::AvaMedia::RundownSerializationUtils
+{
+	AVALANCHEMEDIA_API bool SaveRundownToJson(const UAvaRundown* InRundown, FArchive& InArchive, FText& OutErrorMessage);
+	AVALANCHEMEDIA_API bool SaveRundownToJson(const UAvaRundown* InRundown, const TCHAR* InFilepath, FText& OutErrorMessage);
+
+	AVALANCHEMEDIA_API bool LoadRundownFromJson(UAvaRundown* InRundown, FArchive& InArchive, FText& OutErrorMessage);
+	AVALANCHEMEDIA_API bool LoadRundownFromJson(UAvaRundown* InRundown, const TCHAR* InFilepath, FText& OutErrorMessage);
+}
\ No newline at end of file
diff --git a/Engine/Plugins/Experimental/Avalanche/Source/AvalancheMedia/Private/Rundown/AvaRundownSerializationUtils.cpp b/Engine/Plugins/Experimental/Avalanche/Source/AvalancheMedia/Private/Rundown/AvaRundownSerializationUtils.cpp
new file mode 100644
--- /dev/null
+++ b/Engine/Plugins/Experimental/Avalanche/Source/AvalancheMedia/Private/Rundown/AvaRundownSerializationUtils.cpp
@@ -0,0 +1,125 @@
+// Copyright Epic Games, Inc. All Rights Reserved.
+
+#include "Rundown/AvaRundownSerializationUtils.h"
+
+#include "Backends/JsonStructDeserializerBackend.h"
+#include "Backends/JsonStructSerializerBackend.h"
+#include "HAL/FileManager.h"
+#include "Rundown/AvaRundown.h"
+#include "StructDeserializer.h"
+#include "StructSerializer.h"
+
+#define LOCTEXT_NAMESPACE "AvaRundownSerializationUtils"
+
+namespace UE::AvaMedia::RundownSerializationUtils::Private
+{
+	// Don't serialize the transient properties.
+	static auto TransientPropertyFilter = [](const FProperty* InCurrentProp, const FProperty* InParentProp)
+	{
+		const bool bIsTransient = InCurrentProp && InCurrentProp->HasAnyPropertyFlags(CPF_Transient); 
+		return !bIsTransient; 
+	};
+
+	struct FRundownSerializerPolicies : public FStructSerializerPolicies
+	{
+		FRundownSerializerPolicies()
+		{
+			PropertyFilter = TransientPropertyFilter;
+		}
+	};
+
+	struct FRundownDeserializerPolicies : public FStructDeserializerPolicies
+	{
+		FRundownDeserializerPolicies()
+		{
+			PropertyFilter = TransientPropertyFilter;
+		}
+	};
+}
+
+namespace UE::AvaMedia::RundownSerializationUtils
+{
+	bool SaveRundownToJson(const UAvaRundown* InRundown, FArchive& InArchive, FText& OutErrorMessage)
+	{
+		if (!IsValid(InRundown))
+		{
+			OutErrorMessage = LOCTEXT("SaveRundownJson_InvalidRundown", "Invalid rundown.");
+			return false;
+		}
+
+		// Remark: this is hardcoded to encode in utf16-le.
+		FJsonStructSerializerBackend Backend(InArchive, EStructSerializerBackendFlags::Default);
+		FStructSerializer::Serialize(InRundown, *InRundown->GetClass(), Backend, Private::FRundownSerializerPolicies());
+		return true;
+	}
+
+	bool SaveRundownToJson(const UAvaRundown* InRundown, const TCHAR* InFilepath, FText& OutErrorMessage)
+	{
+		const TUniquePtr<FArchive> FileWriter(IFileManager::Get().CreateFileWriter(InFilepath));
+		if (FileWriter)
+		{
+			const bool bSerialized = SaveRundownToJson(InRundown, *FileWriter, OutErrorMessage);
+			FileWriter->Close();
+			return bSerialized;
+		}
+
+		OutErrorMessage = LOCTEXT("SaveRundownJson_FailedFileWriting", "Failed to open file for writing.");
+		return false;
+	}
+	
+	bool LoadRundownFromJson(UAvaRundown* InRundown, FArchive& InArchive, FText& OutErrorMessage)
+	{
+		if (!IsValid(InRundown))
+		{
+			OutErrorMessage = LOCTEXT("LoadRundownJson_InvalidRundown", "Invalid rundown.");
+			return false;
+		}
+			
+		// Serializing doesn't reset content, it will add to it.
+		// so we need to explicitly make the rundown empty first.
+		if (!InRundown->Empty())
+		{
+			// One reason this could fail is if the rundown is currently playing.
+			if (InRundown->IsPlaying())
+			{
+				OutErrorMessage = LOCTEXT("LoadRundownJson_RundownIsPlaying", "Cannot import on a playing rundown. Stop rundown playback first.");
+			}
+			else
+			{
+				OutErrorMessage = LOCTEXT("LoadRundownJson_FailedClearRundown", "Failed to clear rundown content.");
+			}
+			return false;
+		}
+		
+		FJsonStructDeserializerBackend Backend(InArchive);
+		const bool bLoaded = FStructDeserializer::Deserialize(InRundown, *InRundown->GetClass(), Backend, Private::FRundownDeserializerPolicies());
+		if (bLoaded)
+		{
+			InRundown->PostLoad();
+			InRundown->MarkPackageDirty();
+		}
+		else
+		{
+			OutErrorMessage =
+				FText::Format(LOCTEXT("LoadRundownJson_DeserializerError", "Json Deserializer error: {0}"),
+					FText::FromString(Backend.GetLastErrorMessage()));
+		}
+		return bLoaded;
+	}
+
+	bool LoadRundownFromJson(UAvaRundown* InRundown, const TCHAR* InFilepath, FText& OutErrorMessage)
+	{
+		const TUniquePtr<FArchive> FileReader(IFileManager::Get().CreateFileReader(InFilepath));
+		if (FileReader)
+		{
+			const bool bLoaded = LoadRundownFromJson(InRundown, *FileReader, OutErrorMessage);
+			FileReader->Close();
+			return bLoaded;
+		}
+
+		OutErrorMessage = LOCTEXT("FileNotFound", "File not found");
+		return false;
+	}
+}
+
+#undef LOCTEXT_NAMESPACE
diff --git a/Engine/Plugins/Experimental/Avalanche/Source/AvalancheMedia/Private/Rundown/AvaRundownServer.cpp b/Engine/Plugins/Experimental/Avalanche/Source/AvalancheMedia/Private/Rundown/AvaRundownServer.cpp
--- a/Engine/Plugins/Experimental/Avalanche/Source/AvalancheMedia/Private/Rundown/AvaRundownServer.cpp
+++ b/Engine/Plugins/Experimental/Avalanche/Source/AvalancheMedia/Private/Rundown/AvaRundownServer.cpp
@@ -8,6 +8,7 @@
 #include "Async/Async.h"
 #include "AvaMediaMessageUtils.h"
 #include "AvaMediaRenderTargetUtils.h"
+#include "AvaMediaSerializationUtils.h"
 #include "AvaRemoteControlUtils.h"
 #include "Broadcast/AvaBroadcast.h"
 #include "Broadcast/OutputDevices/AvaBroadcastOutputClassItem.h"
@@ -32,11 +33,13 @@
 #include "Rundown/AvaRundownManagedInstanceCache.h"
 #include "Rundown/AvaRundownPagePlayer.h"
 #include "Rundown/AvaRundownPlaybackUtils.h"
+#include "Rundown/AvaRundownSerializationUtils.h"
 #include "Rundown/AvaRundownServerMediaOutputUtils.h"
 #include "TextureResource.h"
 
 #if WITH_EDITOR
 #include "Editor.h"
+#include "ObjectTools.h"
 #include "ScopedTransaction.h"
 #include "Subsystems/EditorAssetSubsystem.h"
 #endif
@@ -52,7 +55,7 @@ namespace UE::AvaRundownServer::Private
 
 	EAvaRundownServerBuildTargetType GetRundownEngineBuild()
 	{
-		switch(FApp::GetBuildTargetType())
+		switch (FApp::GetBuildTargetType())
 		{
 		case EBuildTargetType::Unknown:
 			return EAvaRundownServerBuildTargetType::Unknown;
@@ -101,6 +104,55 @@ namespace UE::AvaRundownServer::Private
 		
 		return EAvaRundownServerEngineMode::Other;
 	}
+
+	void SanitizeInvalidCharsInline(FString& InOutText, const TCHAR* InvalidChars)
+	{
+		const TCHAR* InvalidChar = InvalidChars ? InvalidChars : TEXT("");
+		while (*InvalidChar)
+		{
+			InOutText.ReplaceCharInline(*InvalidChar, TCHAR('_'), ESearchCase::CaseSensitive);
+			++InvalidChar;
+		}
+	}
+	
+	FString SanitizePackageName(const FString& InPackageName)
+	{
+		// Ensure no backslashes.
+		FString SanitizedName = InPackageName.Replace(TEXT("\\"), TEXT("/"));
+
+		// Replace any other invalid characters with '_'.
+		SanitizeInvalidCharsInline(SanitizedName, INVALID_LONGPACKAGE_CHARACTERS);
+
+		// Coalesce multiple contiguous slashes into a single slash
+		int32 CharIndex = 0;
+		while (CharIndex < SanitizedName.Len())
+		{
+			if (SanitizedName[CharIndex] == TEXT('/'))
+			{
+				int32 SlashCount = 1;
+				while (CharIndex + SlashCount < SanitizedName.Len() &&
+					   SanitizedName[CharIndex + SlashCount] == TEXT('/'))
+				{
+					SlashCount++;
+				}
+
+				if (SlashCount > 1)
+				{
+					SanitizedName.RemoveAt(CharIndex + 1, SlashCount - 1, EAllowShrinking::No);
+				}
+			}
+
+			CharIndex++;
+		}
+
+		// Finally, ensure it begins with "/" since this is an absolute package name.
+		if (!SanitizedName.StartsWith(TEXT("/")))
+		{
+			SanitizedName = FString(TEXT("/")) + SanitizedName;
+		}
+
+		return SanitizedName;
+	}
 	
 	inline FAvaRundownPageInfo GetPageInfo(const UAvaRundown* InRundown, const FAvaRundownPage& InPage)
 	{
@@ -289,6 +341,31 @@ const FMessageAddress& FAvaRundownServer::GetMessageAddress() const
 	return InvalidMessageAddress;
 }
 
+void FAvaRundownServer::AddReferencedObjects(FReferenceCollector& InCollector)
+{
+	if (RundownEditCommandData.CurrentRundown)
+	{
+		InCollector.AddReferencedObject(RundownEditCommandData.CurrentRundown);
+	}
+	if (RundownPlaybackCommandData.CurrentRundown)
+	{
+		InCollector.AddReferencedObject(RundownPlaybackCommandData.CurrentRundown);
+	}
+
+	for (TPair<FSoftObjectPath, TObjectPtr<UAvaRundown>>& ManagedRundown : ManagedRundowns)
+	{
+		if (IsValid(ManagedRundown.Value.Get()))
+		{
+			InCollector.AddReferencedObject(ManagedRundown.Value);
+		}
+	}
+}
+
+FString FAvaRundownServer::GetReferencerName() const
+{
+	return TEXT("FAvaRundownServer");
+}
+
 void FAvaRundownServer::Init(const FString& InAssignedHostName)
 {
 	HostName = InAssignedHostName.IsEmpty() ? FPlatformProcess::ComputerName() : InAssignedHostName;
@@ -298,6 +375,10 @@ void FAvaRundownServer::Init(const FString& InAssignedHostName)
 	.Handling<FAvaRundownGetServerInfo>(this, &FAvaRundownServer::HandleGetRundownServerInfo)
 	.Handling<FAvaRundownGetRundowns>(this, &FAvaRundownServer::HandleGetRundowns)
 	.Handling<FAvaRundownLoadRundown>(this, &FAvaRundownServer::HandleLoadRundown)
+	.Handling<FAvaRundownCreateRundown>(this, &FAvaRundownServer::HandleCreateRundown)
+	.Handling<FAvaRundownDeleteRundown>(this, &FAvaRundownServer::HandleDeleteRundown)
+	.Handling<FAvaRundownImportRundown>(this, &FAvaRundownServer::HandleImportRundown)
+	.Handling<FAvaRundownExportRundown>(this, &FAvaRundownServer::HandleExportRundown)
 	.Handling<FAvaRundownSaveRundown>(this, &FAvaRundownServer::HandleSaveRundown)
 	.Handling<FAvaRundownCreatePage>(this, &FAvaRundownServer::HandleCreatePage)
 	.Handling<FAvaRundownDeletePage>(this, &FAvaRundownServer::HandleDeletePage)
@@ -385,7 +466,7 @@ void FAvaRundownServer::OnPageListChanged(const FAvaRundownPageListChangeParams&
 
 void FAvaRundownServer::OnPagesChanged(const UAvaRundown* InRundown, const FAvaRundownPage& InPage, EAvaRundownPageChanges InChange) const
 {
-	switch(InChange)
+	switch (InChange)
 	{
 	case EAvaRundownPageChanges::AnimationSettings:
 		{
@@ -576,7 +657,7 @@ void FAvaRundownServer::HandleGetRundowns(const FAvaRundownGetRundowns& InMessag
 
 	ReplyMessage->RequestId = InMessage.RequestId;
 	
-	// List all the play list.
+	// List all the rundown assets.
 	if (IAssetRegistry* AssetRegistry = IAssetRegistry::Get())
 	{
 		TArray<FAssetData> Assets;
@@ -590,6 +671,15 @@ void FAvaRundownServer::HandleGetRundowns(const FAvaRundownGetRundowns& InMessag
 		}
 	}
 
+	// Adding the managed rundowns as well, in case they are not listed in the asset registry.
+	for (const TPair<FSoftObjectPath, TObjectPtr<UAvaRundown>>& ManagedRundown : ManagedRundowns)
+	{
+		if (IsValid(ManagedRundown.Value))
+		{
+			ReplyMessage->Rundowns.AddUnique(ManagedRundown.Key.ToString());
+		}
+	}
+
 	SendResponse(ReplyMessage, InContext->GetSender());
 }
 
@@ -626,6 +716,225 @@ void FAvaRundownServer::HandleLoadRundown(const FAvaRundownLoadRundown& InMessag
 		TEXT("Rundown \"%s\" loaded."), *RundownPlaybackCommandData.CurrentRundownPath.ToString());
 }
 
+void FAvaRundownServer::HandleCreateRundown(const FAvaRundownCreateRundown& InMessage, const TSharedRef<IMessageContext, ESPMode::ThreadSafe>& InContext)
+{
+	if (InMessage.PackagePath.IsEmpty() || InMessage.AssetName.IsEmpty())
+	{
+		SendMessage(InContext->GetSender(), InMessage.RequestId, ELogVerbosity::Error,
+			TEXT("\"CreateRundown\" Failed: No rundown path/name specified."));
+		return;
+	}
+
+	using namespace UE::AvaRundownServer;
+	
+	FString PackageName = Private::SanitizePackageName(InMessage.PackagePath + TEXT("/") + InMessage.AssetName);
+
+#if WITH_EDITOR
+	const bool bTransient = InMessage.bTransient;
+#else
+	const bool bTransient = true;
+#endif
+
+	if (bTransient)
+	{
+		static const FString GamePath(TEXT("/Game"));
+		if (PackageName.StartsWith(GamePath))	// ignore case
+		{
+			PackageName = PackageName.Mid(GamePath.Len());
+		}
+
+		static const FString TempPath(TEXT("/Temp"));
+		
+		// Ensure the path begins with /Temp
+		if (!PackageName.StartsWith(TempPath))	// ignore case
+		{
+			PackageName = Private::SanitizePackageName(TempPath + TEXT("/") + PackageName);
+		}
+	}
+
+	if (FindPackage(nullptr, *PackageName))
+	{
+		SendMessage(InContext->GetSender(), InMessage.RequestId, ELogVerbosity::Error,
+			TEXT("\"CreateRundown\" Failed: Requested package \"%s\" already exists."), *PackageName);
+		return;
+	}
+
+	UPackage* const RundownPackage = CreatePackage(*PackageName);
+
+	if (!RundownPackage)
+	{
+		SendMessage(InContext->GetSender(), InMessage.RequestId, ELogVerbosity::Error,
+			TEXT("\"CreateRundown\" Failed: Requested package \"%s\" could not be created."), *PackageName);
+		return;
+	}
+		
+	if (bTransient)
+	{
+		RundownPackage->SetFlags(RF_Transient);
+	}
+
+	constexpr EObjectFlags AssetFlags = RF_Public|RF_Standalone|RF_Transactional;
+	UAvaRundown* Rundown = NewObject<UAvaRundown>(RundownPackage, FName(*InMessage.AssetName), AssetFlags);
+
+	if (!bTransient)
+	{
+		FAssetRegistryModule::AssetCreated(Rundown);
+		RundownPackage->MarkPackageDirty();
+	}
+
+	// The created rundown is added to a managed list to be kept alive by the server as long as it is running.
+	const FSoftObjectPath RundownPath(Rundown);
+
+	if (bTransient)
+	{
+		ManagedRundowns.Add(RundownPath, Rundown);
+	}
+
+	SendMessage(InContext->GetSender(), InMessage.RequestId, ELogVerbosity::Log,
+	TEXT("Rundown \"%s\" Created."), *RundownPath.ToString());
+}
+
+void FAvaRundownServer::HandleDeleteRundown(const FAvaRundownDeleteRundown& InMessage, const TSharedRef<IMessageContext, ESPMode::ThreadSafe>& InContext)
+{
+	// Deleting requires explicit specification of the rundown.
+	if (InMessage.Rundown.IsEmpty())
+	{
+		SendMessage(InContext->GetSender(), InMessage.RequestId, ELogVerbosity::Error,
+			TEXT("\"DeleteRundown\": Rundown asset not specified."));
+		return;
+	}
+	
+	FSoftObjectPath RundownPath(InMessage.Rundown);
+
+	// Only allow rundowns to be deleted if not playing.
+	// We will require an explicit stop command for security reasons.
+	UAvaRundown* Rundown = Cast<UAvaRundown>(RundownPath.ResolveObject());
+	if (Rundown && Rundown->IsPlaying())
+	{
+		SendMessage(InContext->GetSender(), InMessage.RequestId, ELogVerbosity::Error,
+			TEXT("\"DeleteRundown\": Rundown is currently playing. It must be stopped first."));
+		return;
+	}
+
+	if (ManagedRundowns.Remove(RundownPath) > 0)
+	{
+		SendMessage(InContext->GetSender(), InMessage.RequestId, ELogVerbosity::Log,
+			TEXT("\"DeleteRundown\": Rundown \"%s\" removed."), *InMessage.Rundown);
+		return;
+	}
+	
+#if WITH_EDITOR
+	FAssetData RundownAsset;
+	if (FAssetRegistryModule::GetRegistry().TryGetAssetByObjectPath(RundownPath, RundownAsset) == UE::AssetRegistry::EExists::Exists)
+	{
+		TArray<FAssetData> AssetData;
+		AssetData.Add(RundownAsset);
+		ObjectTools::DeleteAssets(AssetData, /*bShowConfirmation*/ false);
+		
+		SendMessage(InContext->GetSender(), InMessage.RequestId, ELogVerbosity::Log,
+			TEXT("\"DeleteRundown\": Rundown \"%s\" deleted."), *InMessage.Rundown);
+		return;
+	}
+#endif
+	
+	SendMessage(InContext->GetSender(), InMessage.RequestId, ELogVerbosity::Error,
+		TEXT("\"DeleteRundown\": Rundown \"%s\" not found."), *InMessage.Rundown);
+}
+
+void FAvaRundownServer::HandleImportRundown(const FAvaRundownImportRundown& InMessage, const TSharedRef<IMessageContext, ESPMode::ThreadSafe>& InContext)
+{
+	UAvaRundown* Rundown = GetOrLoadRundownForEdit(InContext->GetSender(), InMessage.RequestId, InMessage.Rundown);
+	if (!Rundown)
+	{
+		return;	// Response sent by GetOrLoadRundownForEdit.
+	}
+
+	using namespace UE::AvaMedia::RundownSerializationUtils;
+
+	// Load from file
+	if (!InMessage.RundownFile.IsEmpty())
+	{
+		FText ErrorMessage;
+		if (LoadRundownFromJson(Rundown, *InMessage.RundownFile, ErrorMessage))
+		{
+			SendMessage(InContext->GetSender(), InMessage.RequestId, ELogVerbosity::Log,
+				TEXT("\"ImportRundown\": Loaded from file \"%s\"."), *InMessage.RundownFile);
+		}
+		else
+		{
+			SendMessage(InContext->GetSender(), InMessage.RequestId, ELogVerbosity::Error,
+				TEXT("\"ImportRundown\": Failed to load from file \"%s\". Reason: %s"), *InMessage.RundownFile, *ErrorMessage.ToString());
+		}
+		return;
+	}
+
+	// Load from data
+	if (!InMessage.RundownData.IsEmpty())
+	{
+		FText ErrorMessage;
+		FMemoryReaderView Reader(UE::AvaMediaSerializationUtils::JsonValueConversion::ValueToConstBytesView(InMessage.RundownData));
+		if (LoadRundownFromJson(Rundown, Reader, ErrorMessage))
+		{
+			SendMessage(InContext->GetSender(), InMessage.RequestId, ELogVerbosity::Log,
+				TEXT("\"ImportRundown\": Loaded from data."));
+		}
+		else
+		{
+			SendMessage(InContext->GetSender(), InMessage.RequestId, ELogVerbosity::Error,
+				TEXT("\"ImportRundown\": Failed to load from data. Reason: %s"), *ErrorMessage.ToString());
+		}
+		return;
+	}
+
+	SendMessage(InContext->GetSender(), InMessage.RequestId, ELogVerbosity::Error,
+		TEXT("\"ImportRundown\": No data was provided to import from. Either a filename or json data must be provided."));
+}
+
+void FAvaRundownServer::HandleExportRundown(const FAvaRundownExportRundown& InMessage, const TSharedRef<IMessageContext, ESPMode::ThreadSafe>& InContext)
+{
+	const UAvaRundown* Rundown = GetOrLoadRundownForEdit(InContext->GetSender(), InMessage.RequestId, InMessage.Rundown);
+	if (!Rundown)
+	{
+		return; // Response sent by GetOrLoadRundownForEdit.
+	}
+	
+	using namespace UE::AvaMedia::RundownSerializationUtils;
+
+	FText ErrorMessage;
+	
+	// Export to file on the server.
+	if (!InMessage.RundownFile.IsEmpty())
+	{
+		if (SaveRundownToJson(Rundown, *InMessage.RundownFile, ErrorMessage))
+		{
+			SendMessage(InContext->GetSender(), InMessage.RequestId, ELogVerbosity::Log,
+				TEXT("\"ExportRundown\": Rundown exported to \"%s\"."), *InMessage.RundownFile);
+		}
+		else
+		{
+			SendMessage(InContext->GetSender(), InMessage.RequestId, ELogVerbosity::Error,
+				TEXT("\"ExportRundown\": Failed to export rundown to \"%s\". Reason: %s"), *InMessage.RundownFile, *ErrorMessage.ToString());
+		}
+		return;
+	}
+	
+	TArray<uint8> RundownDataAsBytes;
+	FMemoryWriter Writer(RundownDataAsBytes);
+	if (SaveRundownToJson(Rundown, Writer, ErrorMessage))
+	{
+		FAvaRundownExportedRundown* ReplyMessage = FMessageEndpoint::MakeMessage<FAvaRundownExportedRundown>();
+		ReplyMessage->RequestId = InMessage.RequestId;
+		ReplyMessage->Rundown = InMessage.Rundown;
+		UE::AvaMediaSerializationUtils::JsonValueConversion::BytesToString(RundownDataAsBytes, ReplyMessage->RundownData);
+		SendResponse(ReplyMessage, InContext->GetSender());
+	}
+	else
+	{
+		SendMessage(InContext->GetSender(), InMessage.RequestId, ELogVerbosity::Error,
+			TEXT("\"ExportRundown\": Failed to export rundown. Reason: %s"), *ErrorMessage.ToString());
+	}
+}
+
 void FAvaRundownServer::HandleSaveRundown(const FAvaRundownSaveRundown& InMessage, const TSharedRef<IMessageContext, ESPMode::ThreadSafe>& InContext)
 {
 #if WITH_EDITOR
@@ -688,10 +997,10 @@ void FAvaRundownServer::HandleSaveRundown(const FAvaRundownSaveRundown& InMessag
 
 void FAvaRundownServer::HandleGetPages(const FAvaRundownGetPages& InMessage, const TSharedRef<IMessageContext, ESPMode::ThreadSafe>& InContext)
 {
-	UAvaRundown* Rundown = GetOrLoadRundownForEdit(InContext->GetSender(), InMessage.RequestId, InMessage.Rundown);
+	const UAvaRundown* Rundown = GetOrLoadRundownForEdit(InContext->GetSender(), InMessage.RequestId, InMessage.Rundown);
 	if (!Rundown)
 	{
-		return;
+		return; // Response sent by GetOrLoadRundownForEdit.
 	}
 	
 	FAvaRundownPages* ReplyMessage = FMessageEndpoint::MakeMessage<FAvaRundownPages>();
@@ -718,7 +1027,7 @@ void FAvaRundownServer::HandleCreatePage(const FAvaRundownCreatePage& InMessage,
 	UAvaRundown* Rundown = GetOrLoadRundownForEdit(InContext->GetSender(), InMessage.RequestId, InMessage.Rundown);
 	if (!Rundown)
 	{
-		return;
+		return; // Response sent by GetOrLoadRundownForEdit.
 	}
 
 	const FAvaRundownPage& Template = Rundown->GetPage(InMessage.TemplateId);
@@ -744,7 +1053,7 @@ void FAvaRundownServer::HandleDeletePage(const FAvaRundownDeletePage& InMessage,
 	UAvaRundown* Rundown = GetOrLoadRundownForEdit(InContext->GetSender(), InMessage.RequestId, InMessage.Rundown);
 	if (!Rundown)
 	{
-		return;
+		return; // Response sent by GetOrLoadRundownForEdit.
 	}
 
 	const FAvaRundownPage& Page = Rundown->GetPage(InMessage.PageId);
@@ -764,7 +1073,7 @@ void FAvaRundownServer::HandleDeleteTemplate(const FAvaRundownDeleteTemplate& In
 	UAvaRundown* Rundown = GetOrLoadRundownForEdit(InContext->GetSender(), InMessage.RequestId, InMessage.Rundown);
 	if (!Rundown)
 	{
-		return;
+		return; // Response sent by GetOrLoadRundownForEdit.
 	}
 
 	const FAvaRundownPage& Page = Rundown->GetPage(InMessage.PageId);
@@ -796,7 +1105,7 @@ void FAvaRundownServer::HandleCreateTemplate(const FAvaRundownCreateTemplate& In
 	UAvaRundown* Rundown = GetOrLoadRundownForEdit(InContext->GetSender(), InMessage.RequestId, InMessage.Rundown);
 	if (!Rundown)
 	{
-		return;
+		return; // Response sent by GetOrLoadRundownForEdit.
 	}
 	
 	Rundown->AddTemplate();
@@ -810,7 +1119,7 @@ void FAvaRundownServer::HandleChangeTemplateBP(const FAvaRundownChangeTemplateBP
 	UAvaRundown* Rundown = GetOrLoadRundownForEdit(InContext->GetSender(), InMessage.RequestId, InMessage.Rundown);
 	if (!Rundown)
 	{
-		return;
+		return; // Response sent by GetOrLoadRundownForEdit.
 	}
 
 	FAvaRundownPage& Page = Rundown->GetPage(InMessage.TemplateId);
@@ -831,7 +1140,7 @@ void FAvaRundownServer::HandleGetPageDetails(const FAvaRundownGetPageDetails& In
 	UAvaRundown* Rundown = GetOrLoadRundownForEdit(InContext->GetSender(), InMessage.RequestId, InMessage.Rundown);
 	if (!Rundown)
 	{
-		return;
+		return; // Response sent by GetOrLoadRundownForEdit.
 	}
 
 	const FAvaRundownPage& Page = Rundown->GetPage(InMessage.PageId);
@@ -909,7 +1218,7 @@ void FAvaRundownServer::HandleChangePageChannel(const FAvaRundownPageChangeChann
 	UAvaRundown* Rundown = GetOrLoadRundownForEdit(InContext->GetSender(), InMessage.RequestId, InMessage.Rundown);
 	if (!Rundown)
 	{
-		return;
+		return; // Response sent by GetOrLoadRundownForEdit.
 	}
 
 	FAvaRundownPage& Page = Rundown->GetPage(InMessage.PageId);
@@ -1101,7 +1410,7 @@ void FAvaRundownServer::HandleSetCurrentProfile(const FAvaRundownSetCurrentProfi
 	UAvaBroadcast& Broadcast = UAvaBroadcast::Get();
 	const FName ProfileName(InMessage.ProfileName);
 
-	if(Broadcast.IsBroadcastingAnyChannel())
+	if (Broadcast.IsBroadcastingAnyChannel())
 	{
 		LogAndSendMessage(InContext->GetSender(), InMessage.RequestId, ELogVerbosity::Error,
 			TEXT("\"SetCurrentProfile\" Failed. Reason: Channels are currently broadcasting."));
@@ -1323,7 +1632,7 @@ void FAvaRundownServer::HandleChannelEditAction(const FAvaRundownChannelEditActi
 
 	if (InMessage.Action == EAvaRundownChannelEditActions::Remove)
 	{
-		if(!Broadcast.GetCurrentProfile().RemoveChannel(ChannelName))
+		if (!Broadcast.GetCurrentProfile().RemoveChannel(ChannelName))
 		{
 			LogAndSendMessage(InContext->GetSender(), InMessage.RequestId, ELogVerbosity::Error,
 				TEXT("\"ChannelEditAction\" Remove Failed. Reason: Channel \"%s\" didn't exist in profile."), *ChannelName.ToString());
@@ -1366,7 +1675,7 @@ void FAvaRundownServer::HandleRenameChannel(const FAvaRundownRenameChannel& InMe
 		return;
 	}
 
-	if(!Broadcast.RenameChannel(OldChannelName, NewChannelName))
+	if (!Broadcast.RenameChannel(OldChannelName, NewChannelName))
 	{
 		LogAndSendMessage(InContext->GetSender(), InMessage.RequestId, ELogVerbosity::Error,
 			TEXT("\"RenameChannel\" Failed to rename channel \"%s\" to \"%s\" (Unknown reason)."), *OldChannelName.ToString(), *NewChannelName.ToString());
@@ -1717,7 +2026,7 @@ void FAvaRundownServer::ShowStatusCommand(const TArray<FString>& InArgs)
 	UE_LOG(LogAvaRundownServer, Display, TEXT("- Editing PageId: \"%d\""), RundownEditCommandData.ManagedPageId);
 	UE_LOG(LogAvaRundownServer, Display, TEXT("- Playing Rundown: \"%s\""), *RundownPlaybackCommandData.CurrentRundownPath.ToString());
 
-	if (RundownPlaybackCommandData.CurrentRundown.IsValid())
+	if (RundownPlaybackCommandData.CurrentRundown)
 	{
 		TArray<int32> PlayingPages = RundownPlaybackCommandData.CurrentRundown->GetPlayingPageIds();
 		for (const int32 PlayingPageId : PlayingPages)
@@ -1793,7 +2102,7 @@ void FAvaRundownServer::HandlePageActions(const FRequestInfo& InRequestInfo, con
 {
 	using namespace UE::AvaRundownServer::Private;
 	
-	if (!RundownPlaybackCommandData.CurrentRundown.IsValid())
+	if (!RundownPlaybackCommandData.CurrentRundown)
 	{
 		LogAndSendMessage(InRequestInfo.Sender, InRequestInfo.RequestId, ELogVerbosity::Error,
 			TEXT("\"PageAction\" Failed. Reason: no play list currently loaded."));
@@ -1982,7 +2291,7 @@ UAvaRundown* FAvaRundownServer::GetOrLoadRundownForEdit(const FMessageAddress& I
 			Rundown = RundownPlaybackCommandData.CurrentRundown.Get();
 
 			// Update the edit data accordingly.
-			RundownEditCommandData.CurrentRundown.Reset(Rundown);
+			RundownEditCommandData.CurrentRundown = Rundown;
 			RundownEditCommandData.CurrentRundownPath = RundownPlaybackCommandData.CurrentRundownPath;
 		}
 
@@ -2047,7 +2356,7 @@ UAvaRundown* FAvaRundownServer::FRundownCache::GetOrLoadRundown(const FSoftObjec
 				InUnloadCurrentRundownFunction(CurrentRundown.Get());
 			}
 			InNewRundownLoadedFunction(NewRundown.Get());
-			CurrentRundown = NewRundown;
+			CurrentRundown = NewRundown.Get();
 			CurrentRundownPath = InRundownPath;
 		}
 		else
@@ -2131,7 +2440,7 @@ void FAvaRundownServer::FRundownEditCommandData::SaveCurrentRemoteControlPresetT
 		RemoteControlModule.UnregisterEmbeddedPreset(CurrentPresetName);
 	}
 
-	if (!CurrentRundown.IsValid())
+	if (!CurrentRundown)
 	{
 		return;
 	}
@@ -2186,11 +2495,11 @@ FAvaRundownServer::FRundownPlaybackCommandData::~FRundownPlaybackCommandData()
 
 void FAvaRundownServer::FRundownPlaybackCommandData::ClosePlaybackContext()
 {
-	if (CurrentRundown.IsValid())
+	if (CurrentRundown)
 	{
 		// Stop all playing pages.
 		CurrentRundown->ClosePlaybackContext(true);
-		CurrentRundown.Reset();
+		CurrentRundown = nullptr;
 		CurrentRundownPath.Reset();
 	}
 }
diff --git a/Engine/Plugins/Experimental/Avalanche/Source/AvalancheMedia/Private/Rundown/AvaRundownServer.h b/Engine/Plugins/Experimental/Avalanche/Source/AvalancheMedia/Private/Rundown/AvaRundownServer.h
--- a/Engine/Plugins/Experimental/Avalanche/Source/AvalancheMedia/Private/Rundown/AvaRundownServer.h
+++ b/Engine/Plugins/Experimental/Avalanche/Source/AvalancheMedia/Private/Rundown/AvaRundownServer.h
@@ -18,7 +18,7 @@ class UMediaOutput;
  * The intention is to run a web socket transport bridge so the messages can
  * come from external applications.
  */
-class FAvaRundownServer : public TSharedFromThis<FAvaRundownServer>, public IAvaRundownServer
+class FAvaRundownServer : public TSharedFromThis<FAvaRundownServer>, public FGCObject, public IAvaRundownServer
 {
 public:
 	FAvaRundownServer();
@@ -30,6 +30,11 @@ public:
 	virtual TArray<FMessageAddress> GetClientAddresses() const override { return ClientAddresses; }
 	//~ End IAvaRundownServer
 
+	//~ Begin FGCObject
+	virtual void AddReferencedObjects(FReferenceCollector& InCollector) override;
+	virtual FString GetReferencerName() const override;
+	//~ End FGCObject
+
 	void Init(const FString& InAssignedHostName);
 
 	void SetupBroadcastDelegates(UAvaBroadcast* InBroadcast);
@@ -52,6 +57,10 @@ public:
 	void HandleGetRundownServerInfo(const FAvaRundownGetServerInfo& InMessage, const TSharedRef<IMessageContext, ESPMode::ThreadSafe>& InContext);
 	void HandleGetRundowns(const FAvaRundownGetRundowns& InMessage, const TSharedRef<IMessageContext, ESPMode::ThreadSafe>& InContext);
 	void HandleLoadRundown(const FAvaRundownLoadRundown& InMessage, const TSharedRef<IMessageContext, ESPMode::ThreadSafe>& InContext);
+	void HandleCreateRundown(const FAvaRundownCreateRundown& InMessage, const TSharedRef<IMessageContext, ESPMode::ThreadSafe>& InContext);
+	void HandleDeleteRundown(const FAvaRundownDeleteRundown& InMessage, const TSharedRef<IMessageContext, ESPMode::ThreadSafe>& InContext);
+	void HandleImportRundown(const FAvaRundownImportRundown& InMessage, const TSharedRef<IMessageContext, ESPMode::ThreadSafe>& InContext);
+	void HandleExportRundown(const FAvaRundownExportRundown& InMessage, const TSharedRef<IMessageContext, ESPMode::ThreadSafe>& InContext);
 	void HandleSaveRundown(const FAvaRundownSaveRundown& InMessage, const TSharedRef<IMessageContext, ESPMode::ThreadSafe>& InContext);
 	void HandleGetPages(const FAvaRundownGetPages& InMessage, const TSharedRef<IMessageContext, ESPMode::ThreadSafe>& InContext);
 	void HandleCreatePage(const FAvaRundownCreatePage& InMessage,const TSharedRef<IMessageContext, ESPMode::ThreadSafe>& InContext);
@@ -191,7 +200,7 @@ private:
 		FSoftObjectPath CurrentRundownPath;
 		
 		/** Currently loaded/cached rundown object. */
-		TStrongObjectPtr<UAvaRundown> CurrentRundown;
+		TObjectPtr<UAvaRundown> CurrentRundown;
 
 		FDelegateHandle OnPlaybackInstanceStatusChangedDelegateHandle;
 		
@@ -240,4 +249,7 @@ private:
 		void ClosePlaybackContext();
 	};
 	FRundownPlaybackCommandData RundownPlaybackCommandData;
+
+	/** Keep a map of created transient rundowns. */
+	TMap<FSoftObjectPath, TObjectPtr<UAvaRundown>> ManagedRundowns;
 };
diff --git a/Engine/Plugins/Experimental/Avalanche/Source/AvalancheMediaEditor/Private/Rundown/AssetDefinition_AvaRundown.cpp b/Engine/Plugins/Experimental/Avalanche/Source/AvalancheMediaEditor/Private/Rundown/AssetDefinition_AvaRundown.cpp
--- a/Engine/Plugins/Experimental/Avalanche/Source/AvalancheMediaEditor/Private/Rundown/AssetDefinition_AvaRundown.cpp
+++ b/Engine/Plugins/Experimental/Avalanche/Source/AvalancheMediaEditor/Private/Rundown/AssetDefinition_AvaRundown.cpp
@@ -10,6 +10,7 @@
 #include "Rundown/AvaRundown.h"
 #include "Rundown/AvaRundownEditor.h"
 #include "Rundown/AvaRundownEditorUtils.h"
+#include "Rundown/AvaRundownSerializationUtils.h"
 #include "ToolMenu.h"
 #include "ToolMenus.h"
 
@@ -77,10 +78,17 @@ namespace UE::AvaMediaEditor::Rundown::Private
 				}
 
 				using namespace UE::AvaRundownEditor::Utils;
+				using namespace UE::AvaMedia::RundownSerializationUtils;
 				FString ExportFilename = GetExportFilepath(Rundown, TEXT("json file"), TEXT("json"));
 				if (!ExportFilename.IsEmpty())
 				{
-					SaveRundownToJson(Rundown, *ExportFilename);
+					FText ErrorMessage;
+					if (!SaveRundownToJson(Rundown, *ExportFilename, ErrorMessage))
+					{
+						UE_LOG(LogAvaRundown, Error,
+							TEXT("Failed to export rundown \"%s\" to file \"%s\". Reason: %s"),
+							*Rundown->GetFullName(), *ExportFilename, *ErrorMessage.ToString());
+					}
 				}
 			}
 		}
@@ -150,10 +158,17 @@ namespace UE::AvaMediaEditor::Rundown::Private
 				}
 				
 				using namespace UE::AvaRundownEditor::Utils;
+				using namespace UE::AvaMedia::RundownSerializationUtils;
 				FString ImportFilename = GetImportFilepath(TEXT("json file"), TEXT("json"));
 				if (!ImportFilename.IsEmpty())
 				{
-					LoadRundownFromJson(Rundown, *ImportFilename);
+					FText ErrorMessage;
+					if (!LoadRundownFromJson(Rundown, *ImportFilename, ErrorMessage))
+					{
+						UE_LOG(LogAvaRundown, Error,
+							TEXT("Failed to import rundown \"%s\" from file \"%s\". Reason: %s"),
+							*Rundown->GetFullName(), *ImportFilename, *ErrorMessage.ToString());
+					}
 				}
 			}
 		}
diff --git a/Engine/Plugins/Experimental/Avalanche/Source/AvalancheMediaEditor/Private/Rundown/AvaRundownEditorUtils.cpp b/Engine/Plugins/Experimental/Avalanche/Source/AvalancheMediaEditor/Private/Rundown/AvaRundownEditorUtils.cpp
--- a/Engine/Plugins/Experimental/Avalanche/Source/AvalancheMediaEditor/Private/Rundown/AvaRundownEditorUtils.cpp
+++ b/Engine/Plugins/Experimental/Avalanche/Source/AvalancheMediaEditor/Private/Rundown/AvaRundownEditorUtils.cpp
@@ -2,8 +2,6 @@
 
 #include "Rundown/AvaRundownEditorUtils.h"
 
-#include "Backends/JsonStructDeserializerBackend.h"
-#include "Backends/JsonStructSerializerBackend.h"
 #include "Backends/XmlStructSerializerBackend.h"
 #include "ContentBrowserModule.h"
 #include "DesktopPlatformModule.h"
@@ -29,7 +27,7 @@ namespace UE::AvaRundownEditor::Utils::Private
 	static const FString PageEntriesName = TEXT("Pages");
 
 	// Don't serialize the transient properties.
-	static auto TransientPropertyFiler = [](const FProperty* InCurrentProp, const FProperty* InParentProp)
+	static auto TransientPropertyFilter = [](const FProperty* InCurrentProp, const FProperty* InParentProp)
 	{
 		const bool bIsTransient = InCurrentProp && InCurrentProp->HasAnyPropertyFlags(CPF_Transient); 
 		return !bIsTransient; 
@@ -39,7 +37,7 @@ namespace UE::AvaRundownEditor::Utils::Private
 	{
 		FRundownSerializerPolicies()
 		{
-			PropertyFilter = TransientPropertyFiler;
+			PropertyFilter = TransientPropertyFilter;
 		}
 	};
 
@@ -47,7 +45,7 @@ namespace UE::AvaRundownEditor::Utils::Private
 	{
 		FRundownDeserializerPolicies()
 		{
-			PropertyFilter = TransientPropertyFiler;
+			PropertyFilter = TransientPropertyFilter;
 		}
 	};
 	
@@ -353,14 +351,14 @@ bool UE::AvaRundownEditor::Utils::SaveRundownToXml(const UAvaRundown* InRundown,
 	return false;
 }
 
-bool UE::AvaRundownEditor::Utils::SaveRundownToXml(const UAvaRundown* InRundown, const TCHAR* InFilepath)
+bool UE::AvaRundownEditor::Utils::SaveRundownToXml(const UAvaRundown* InRundown, const TCHAR* InFilepath, EXmlSerializationEncoding InXmlEncoding)
 {
 	if (IsValid(InRundown))
 	{
 		const TUniquePtr<FArchive> FileWriter(IFileManager::Get().CreateFileWriter(InFilepath));
 		if (FileWriter)	
 		{
-			const bool bSaved = SaveRundownToXml(InRundown, *FileWriter, EXmlSerializationEncoding::Utf8);
+			const bool bSaved = SaveRundownToXml(InRundown, *FileWriter, InXmlEncoding);
 			FileWriter->Close();
 			return bSaved;
 		}
@@ -368,69 +366,6 @@ bool UE::AvaRundownEditor::Utils::SaveRundownToXml(const UAvaRundown* InRundown,
 	return false;
 }
 
-bool UE::AvaRundownEditor::Utils::SaveRundownToJson(const UAvaRundown* InRundown, FArchive& InArchive)
-{
-	if (IsValid(InRundown))
-	{
-		FJsonStructSerializerBackend Backend(InArchive, EStructSerializerBackendFlags::Default);
-		FStructSerializer::Serialize(InRundown, *InRundown->GetClass(), Backend, Private::FRundownSerializerPolicies());
-		return true;
-	}
-	return false;
-}
-
-bool UE::AvaRundownEditor::Utils::SaveRundownToJson(const UAvaRundown* InRundown, const TCHAR* InFilepath)
-{
-	if (IsValid(InRundown))
-	{
-		const TUniquePtr<FArchive> FileWriter(IFileManager::Get().CreateFileWriter(InFilepath));
-		if (FileWriter)
-		{
-			const bool bSerialized = SaveRundownToJson(InRundown, *FileWriter);
-			FileWriter->Close();
-			return bSerialized;
-		}
-	}
-	return false;
-}
-
-bool UE::AvaRundownEditor::Utils::LoadRundownFromJson(UAvaRundown* InRundown, FArchive& InArchive)
-{
-	if (IsValid(InRundown))
-	{
-		// Serializing doesn't reset content, it will add to it, so we need
-		// to explicitly make the rundown empty first.
-		if (!InRundown->Empty())
-		{
-			return false;
-		}
-		
-		FJsonStructDeserializerBackend Backend(InArchive);
-		const bool bLoaded = FStructDeserializer::Deserialize(InRundown, *InRundown->GetClass(), Backend, Private::FRundownDeserializerPolicies());
-		if (bLoaded)
-		{
-			InRundown->PostLoad();
-		}
-		return bLoaded;
-	}
-	return false;
-}
-
-bool UE::AvaRundownEditor::Utils::LoadRundownFromJson(UAvaRundown* InRundown, const TCHAR* InFilepath)
-{
-	if (IsValid(InRundown))
-	{
-		const TUniquePtr<FArchive> FileReader(IFileManager::Get().CreateFileReader(InFilepath));
-		if (FileReader)
-		{
-			const bool bLoaded = LoadRundownFromJson(InRundown, *FileReader);
-			FileReader->Close();
-			return bLoaded;
-		}
-	}
-	return false;
-}
-
 bool UE::AvaRundownEditor::Utils::CanLoadRundownFromFile(const TCHAR* InFilepath)
 {
 	const FStringView Extension = FPathViews::GetExtension(InFilepath);
diff --git a/Engine/Plugins/Experimental/Avalanche/Source/AvalancheMediaEditor/Private/Rundown/AvaRundownEditorUtils.h b/Engine/Plugins/Experimental/Avalanche/Source/AvalancheMediaEditor/Private/Rundown/AvaRundownEditorUtils.h
--- a/Engine/Plugins/Experimental/Avalanche/Source/AvalancheMediaEditor/Private/Rundown/AvaRundownEditorUtils.h
+++ b/Engine/Plugins/Experimental/Avalanche/Source/AvalancheMediaEditor/Private/Rundown/AvaRundownEditorUtils.h
@@ -30,13 +30,7 @@ namespace UE::AvaRundownEditor::Utils
 	TStrongObjectPtr<UAvaRundown> ExportPagesToRundown(const UAvaRundown* InRundown, const TArray<int32>& InPageIds);
 
 	bool SaveRundownToXml(const UAvaRundown* InRundown, FArchive& InArchive, EXmlSerializationEncoding InXmlEncoding);
-	bool SaveRundownToXml(const UAvaRundown* InRundown, const TCHAR* InFilepath);
-
-	bool SaveRundownToJson(const UAvaRundown* InRundown, FArchive& InArchive);
-	bool SaveRundownToJson(const UAvaRundown* InRundown, const TCHAR* InFilepath);
-
-	bool LoadRundownFromJson(UAvaRundown* InRundown, FArchive& InArchive);
-	bool LoadRundownFromJson(UAvaRundown* InRundown, const TCHAR* InFilepath);
+	bool SaveRundownToXml(const UAvaRundown* InRundown, const TCHAR* InFilepath, EXmlSerializationEncoding InXmlEncoding);
 
 	/**
 	 * Check the file extenstion to see if it is a supported format to use 
diff --git a/Engine/Plugins/Experimental/Avalanche/Source/AvalancheMediaEditor/Private/Rundown/AvaRundownExporter.cpp b/Engine/Plugins/Experimental/Avalanche/Source/AvalancheMediaEditor/Private/Rundown/AvaRundownExporter.cpp
--- a/Engine/Plugins/Experimental/Avalanche/Source/AvalancheMediaEditor/Private/Rundown/AvaRundownExporter.cpp
+++ b/Engine/Plugins/Experimental/Avalanche/Source/AvalancheMediaEditor/Private/Rundown/AvaRundownExporter.cpp
@@ -3,8 +3,10 @@
 #include "Rundown/AvaRundownExporter.h"
 
 #include "AvaMediaSerializationUtils.h"
+#include "Misc/FeedbackContext.h"
 #include "Rundown/AvaRundown.h"
 #include "Rundown/AvaRundownEditorUtils.h"
+#include "Rundown/AvaRundownSerializationUtils.h"
 #include "Serialization/MemoryWriter.h"
 
 UAvaRundownExporter::UAvaRundownExporter()
@@ -31,7 +33,15 @@ bool UAvaRundownExporter::ExportText(const FExportObjectInnerContext* InContext,
 
 	if ( FCString::Stricmp(InType, TEXT("json")) == 0 )
 	{
-		bSavedToBytes = UE::AvaRundownEditor::Utils::SaveRundownToJson(Rundown, Writer);
+		FText ErrorMessage;
+		bSavedToBytes = UE::AvaMedia::RundownSerializationUtils::SaveRundownToJson(Rundown, Writer, ErrorMessage);
+		
+		if (!bSavedToBytes && InWarn)
+		{
+			InWarn->CategorizedLogf(LogAvaRundown.GetCategoryName(), ELogVerbosity::Error,
+				TEXT("Failed to export rundown \"%s\". Reason: %s"),
+				*Rundown->GetFullName(), *ErrorMessage.ToString());
+		}
 	}
 	else if (FCString::Stricmp(InType, TEXT("xml")) == 0 )
 	{
diff --git a/Engine/Plugins/Experimental/Avalanche/Source/AvalancheMediaEditor/Private/Rundown/Pages/Slate/SAvaRundownInstancedPageList.cpp b/Engine/Plugins/Experimental/Avalanche/Source/AvalancheMediaEditor/Private/Rundown/Pages/Slate/SAvaRundownInstancedPageList.cpp
--- a/Engine/Plugins/Experimental/Avalanche/Source/AvalancheMediaEditor/Private/Rundown/Pages/Slate/SAvaRundownInstancedPageList.cpp
+++ b/Engine/Plugins/Experimental/Avalanche/Source/AvalancheMediaEditor/Private/Rundown/Pages/Slate/SAvaRundownInstancedPageList.cpp
@@ -11,6 +11,7 @@
 #include "Rundown/AvaRundownEditorUtils.h"
 #include "Rundown/AvaRundownPage.h"
 #include "Rundown/AvaRundownPlaybackUtils.h"
+#include "Rundown/AvaRundownSerializationUtils.h"
 #include "Rundown/Factories/Filters/AvaRundownFactoriesUtils.h"
 #include "Rundown/Pages/Columns/AvaRundownPageAssetNameColumn.h"
 #include "Rundown/Pages/Columns/AvaRundownPageChannelSelectorColumn.h"
@@ -597,9 +598,11 @@ bool SAvaRundownInstancedPageList::HandleDropExternalFiles(const TArray<FString>
 		if (FPathViews::GetExtension(File).Equals(TEXT("json"), ESearchCase::IgnoreCase))
 		{
 			using namespace UE::AvaRundownEditor::Utils;
+			using namespace UE::AvaMedia::RundownSerializationUtils;
 			const TStrongObjectPtr<UAvaRundown> TmpRundown(NewObject<UAvaRundown>());
+			FText ErrorMessage;
 			
-			if (LoadRundownFromJson(TmpRundown.Get(), *File))
+			if (LoadRundownFromJson(TmpRundown.Get(), *File, ErrorMessage))
 			{
 				if (!TmpRundown->GetInstancedPages().Pages.IsEmpty())
 				{
@@ -621,7 +624,7 @@ bool SAvaRundownInstancedPageList::HandleDropExternalFiles(const TArray<FString>
 			}
 			else
 			{
-				UE_LOG(LogAvaRundown, Error, TEXT("%s is not a valid rundown."), *File);
+				UE_LOG(LogAvaRundown, Error, TEXT("Failed to load rundown from file \"%s\". Reason: %s."), *File, *ErrorMessage.ToString());
 			}
 		}
 		else
diff --git a/Engine/Plugins/Experimental/Avalanche/Source/AvalancheMediaEditor/Private/Rundown/Pages/Slate/SAvaRundownPageList.cpp b/Engine/Plugins/Experimental/Avalanche/Source/AvalancheMediaEditor/Private/Rundown/Pages/Slate/SAvaRundownPageList.cpp
--- a/Engine/Plugins/Experimental/Avalanche/Source/AvalancheMediaEditor/Private/Rundown/Pages/Slate/SAvaRundownPageList.cpp
+++ b/Engine/Plugins/Experimental/Avalanche/Source/AvalancheMediaEditor/Private/Rundown/Pages/Slate/SAvaRundownPageList.cpp
@@ -17,6 +17,7 @@
 #include "Rundown/AvaRundownEditorUtils.h"
 #include "Rundown/AvaRundownPage.h"
 #include "Rundown/AvaRundownPlaybackUtils.h"
+#include "Rundown/AvaRundownSerializationUtils.h"
 #include "Rundown/Factories/Filters/AvaRundownFactoriesUtils.h"
 #include "Rundown/Factories/Filters/IAvaRundownFilterSuggestionFactory.h"
 #include "Rundown/Pages/AvaRundownPageContext.h"
@@ -766,6 +767,7 @@ void SAvaRundownPageList::ExportSelectedPagesToExternalFile(const TCHAR* InType)
 	}
 
 	using namespace UE::AvaRundownEditor::Utils;
+	using namespace UE::AvaMedia::RundownSerializationUtils;
 	if (const TStrongObjectPtr<UAvaRundown> NewRundown = ExportPagesToRundown(SourceRundown, SelectedPageIds))
 	{
 		if (FCString::Stricmp(InType, TEXT("json")) == 0)
@@ -773,7 +775,13 @@ void SAvaRundownPageList::ExportSelectedPagesToExternalFile(const TCHAR* InType)
 			const FString ExportFilename = GetExportFilepath(SourceRundown, TEXT("json file"), TEXT("json"));
 			if (!ExportFilename.IsEmpty())
 			{
-				SaveRundownToJson(NewRundown.Get(), *ExportFilename);
+				FText ErrorMessage;
+				if (!SaveRundownToJson(NewRundown.Get(), *ExportFilename, ErrorMessage))
+				{
+					UE_LOG(LogAvaRundown, Error,
+						TEXT("Failed to export rundown \"%s\" to file \"%s\". Reason: %s"),
+						*NewRundown->GetFullName(), *ExportFilename, *ErrorMessage.ToString());
+				}
 			}
 		}
 		else if (FCString::Stricmp(InType, TEXT("xml")) == 0)
@@ -781,7 +789,7 @@ void SAvaRundownPageList::ExportSelectedPagesToExternalFile(const TCHAR* InType)
 			const FString ExportFilename = GetExportFilepath(SourceRundown, TEXT("xml file"), TEXT("xml"));
 			if (!ExportFilename.IsEmpty())
 			{
-				SaveRundownToXml(NewRundown.Get(), *ExportFilename);
+				SaveRundownToXml(NewRundown.Get(), *ExportFilename, EXmlSerializationEncoding::Utf8);
 			}
 		}
 		else
