Motion Design Transition: Added Parameter mode for a Details-view like mode

diff --git a/Engine/Plugins/Experimental/Avalanche/Source/AvalancheTransitionEditor/AvalancheTransitionEditor.Build.cs b/Engine/Plugins/Experimental/Avalanche/Source/AvalancheTransitionEditor/AvalancheTransitionEditor.Build.cs
--- a/Engine/Plugins/Experimental/Avalanche/Source/AvalancheTransitionEditor/AvalancheTransitionEditor.Build.cs
+++ b/Engine/Plugins/Experimental/Avalanche/Source/AvalancheTransitionEditor/AvalancheTransitionEditor.Build.cs
@@ -36,6 +36,7 @@ public class AvalancheTransitionEditor : ModuleRules
                 "SlateCore",
                 "StateTreeModule",
                 "StructUtils",
+                "StructUtilsEditor",
                 "ToolMenus",
                 "UnrealEd",
             }
diff --git a/Engine/Plugins/Experimental/Avalanche/Source/AvalancheTransitionEditor/Private/AppModes/AvaTransitionAdvancedMode.cpp b/Engine/Plugins/Experimental/Avalanche/Source/AvalancheTransitionEditor/Private/AppModes/AvaTransitionAdvancedMode.cpp
--- a/Engine/Plugins/Experimental/Avalanche/Source/AvalancheTransitionEditor/Private/AppModes/AvaTransitionAdvancedMode.cpp
+++ b/Engine/Plugins/Experimental/Avalanche/Source/AvalancheTransitionEditor/Private/AppModes/AvaTransitionAdvancedMode.cpp
@@ -11,9 +11,9 @@
 FAvaTransitionAdvancedMode::FAvaTransitionAdvancedMode(const TSharedRef<FAvaTransitionEditor>& InEditor)
 	: FAvaTransitionAppMode(InEditor, EAvaTransitionEditorMode::Advanced)
 {
-	WorkspaceMenuCategory = FWorkspaceItem::NewGroup(LOCTEXT("WorkspaceMenuCategory", "Motion Design Transition Advanced"));
+	RegisterDefaultTabFactories();
 
-	TabFactories.RegisterFactory(MakeShared<FAvaTransitionTreeDetailsTabFactory>(InEditor));
+	WorkspaceMenuCategory = FWorkspaceItem::NewGroup(LOCTEXT("WorkspaceMenuCategory", "Motion Design Transition Advanced"));
 
 	TabLayout = FTabManager::NewLayout("AvaTransitionEditor_Advanced_Layout_V0_1")
 	->AddArea
diff --git a/Engine/Plugins/Experimental/Avalanche/Source/AvalancheTransitionEditor/Private/AppModes/AvaTransitionAppMode.cpp b/Engine/Plugins/Experimental/Avalanche/Source/AvalancheTransitionEditor/Private/AppModes/AvaTransitionAppMode.cpp
--- a/Engine/Plugins/Experimental/Avalanche/Source/AvalancheTransitionEditor/Private/AppModes/AvaTransitionAppMode.cpp
+++ b/Engine/Plugins/Experimental/Avalanche/Source/AvalancheTransitionEditor/Private/AppModes/AvaTransitionAppMode.cpp
@@ -4,6 +4,7 @@
 #include "AvaTransitionEditor.h"
 #include "TabFactories/AvaTransitionCompilerResultsTabFactory.h"
 #include "TabFactories/AvaTransitionSelectionDetailsTabFactory.h"
+#include "TabFactories/AvaTransitionTreeDetailsTabFactory.h"
 #include "TabFactories/AvaTransitionTreeTabFactory.h"
 #include "Widgets/Layout/SBorder.h"
 #include "Widgets/Layout/SBox.h"
@@ -15,14 +16,16 @@ namespace UE::AvaTransitionEditor::Private
 {
 	static const TMap<EAvaTransitionEditorMode, FName> GModeNames
 	{
-		{ EAvaTransitionEditorMode::Default , TEXT("Default") },
-		{ EAvaTransitionEditorMode::Advanced, TEXT("Advanced") },
+		{ EAvaTransitionEditorMode::Default  , TEXT("Default") },
+		{ EAvaTransitionEditorMode::Advanced , TEXT("Advanced") },
+		{ EAvaTransitionEditorMode::Parameter, TEXT("Parameter") },
 	};
 
 	static const TMap<FName, FText> GLocalizedModeNames
 	{
-		{ TEXT("Default") , LOCTEXT("Default", "Default") },
-		{ TEXT("Advanced"), LOCTEXT("Advanced", "Advanced") },
+		{ TEXT("Default")  , LOCTEXT("Default", "Default") },
+		{ TEXT("Advanced") , LOCTEXT("Advanced", "Advanced") },
+		{ TEXT("Parameter"), LOCTEXT("Parameter", "Parameter") },
 	};
 
 	FText GetLocalizedMode(FName InModeName)
@@ -31,24 +34,16 @@ namespace UE::AvaTransitionEditor::Private
 	}
 }
 
+FName FAvaTransitionAppMode::StaticGetModeName(EAvaTransitionEditorMode InEditorMode)
+{
+	return UE::AvaTransitionEditor::Private::GModeNames[InEditorMode];
+}
+
 FAvaTransitionAppMode::FAvaTransitionAppMode(const TSharedRef<FAvaTransitionEditor>& InEditor, EAvaTransitionEditorMode InEditorMode)
-	: FApplicationMode(UE::AvaTransitionEditor::Private::GModeNames[InEditorMode], &UE::AvaTransitionEditor::Private::GetLocalizedMode)
+	: FApplicationMode(FAvaTransitionAppMode::StaticGetModeName(InEditorMode), &UE::AvaTransitionEditor::Private::GetLocalizedMode)
 	, EditorWeak(InEditor)
 	, EditorMode(InEditorMode)
 {
-	// Default Factories
-	TabFactories.RegisterFactory(MakeShared<FAvaTransitionTreeTabFactory>(InEditor));
-	TabFactories.RegisterFactory(MakeShared<FAvaTransitionSelectionDetailsTabFactory>(InEditor, InEditorMode));
-	TabFactories.RegisterFactory(MakeShared<FAvaTransitionCompilerResultsTabFactory>(InEditor));
-}
-
-void FAvaTransitionAppMode::RegisterTabFactories(TSharedPtr<FTabManager> InTabManager)
-{
-	TSharedPtr<FAvaTransitionEditor> Editor = EditorWeak.Pin();
-	check(Editor.IsValid());
-	Editor->SetEditorMode(EditorMode);
-	Editor->PushTabFactories(TabFactories);
-	FApplicationMode::RegisterTabFactories(InTabManager);
 }
 
 void FAvaTransitionAppMode::AddToToolbar(const TSharedRef<FExtender>& InToolbarExtender)
@@ -60,6 +55,16 @@ void FAvaTransitionAppMode::AddToToolbar(const TSharedRef<FExtender>& InToolbarE
 		, FToolBarExtensionDelegate::CreateSP(this, &FAvaTransitionAppMode::ExtendToolbar));
 }
 
+void FAvaTransitionAppMode::RegisterDefaultTabFactories()
+{
+	TSharedRef<FAvaTransitionEditor> Editor = EditorWeak.Pin().ToSharedRef();
+
+	TabFactories.RegisterFactory(MakeShared<FAvaTransitionTreeTabFactory>(Editor));
+	TabFactories.RegisterFactory(MakeShared<FAvaTransitionCompilerResultsTabFactory>(Editor));
+	TabFactories.RegisterFactory(MakeShared<FAvaTransitionTreeDetailsTabFactory>(Editor));
+	TabFactories.RegisterFactory(MakeShared<FAvaTransitionSelectionDetailsTabFactory>(Editor, EditorMode));
+}
+
 void FAvaTransitionAppMode::ExtendToolbar(FToolBarBuilder& InToolbarBuilder)
 {
 	if (!EditorWeak.IsValid())
@@ -90,4 +95,13 @@ void FAvaTransitionAppMode::ExtendToolbar(FToolBarBuilder& InToolbarBuilder)
 	);
 }
 
+void FAvaTransitionAppMode::RegisterTabFactories(TSharedPtr<FTabManager> InTabManager)
+{
+	TSharedPtr<FAvaTransitionEditor> Editor = EditorWeak.Pin();
+	check(Editor.IsValid());
+	Editor->SetEditorMode(EditorMode);
+	Editor->PushTabFactories(TabFactories);
+	FApplicationMode::RegisterTabFactories(InTabManager);
+}
+
 #undef LOCTEXT_NAMESPACE
diff --git a/Engine/Plugins/Experimental/Avalanche/Source/AvalancheTransitionEditor/Private/AppModes/AvaTransitionAppMode.h b/Engine/Plugins/Experimental/Avalanche/Source/AvalancheTransitionEditor/Private/AppModes/AvaTransitionAppMode.h
--- a/Engine/Plugins/Experimental/Avalanche/Source/AvalancheTransitionEditor/Private/AppModes/AvaTransitionAppMode.h
+++ b/Engine/Plugins/Experimental/Avalanche/Source/AvalancheTransitionEditor/Private/AppModes/AvaTransitionAppMode.h
@@ -14,17 +14,21 @@ struct FSlateBrush;
 class FAvaTransitionAppMode : public FApplicationMode
 {
 public:
+	static FName StaticGetModeName(EAvaTransitionEditorMode InEditorMode);
+
 	FAvaTransitionAppMode(const TSharedRef<FAvaTransitionEditor>& InEditor, EAvaTransitionEditorMode InEditorMode);
 
 	void AddToToolbar(const TSharedRef<FExtender>& InToolbarExtender);
 
 protected:
+	void RegisterDefaultTabFactories();
+
+	void ExtendToolbar(FToolBarBuilder& InToolbarBuilder);
+
 	//~ Begin FApplicationMode
 	virtual void RegisterTabFactories(TSharedPtr<FTabManager> InTabManager) override final;
 	//~ End FApplicationMode
 
-	void ExtendToolbar(FToolBarBuilder& InToolbarBuilder);
-
 	TWeakPtr<FAvaTransitionEditor> EditorWeak;
 
 	FWorkflowAllowedTabSet TabFactories;
diff --git a/Engine/Plugins/Experimental/Avalanche/Source/AvalancheTransitionEditor/Private/AppModes/AvaTransitionParameterMode.cpp b/Engine/Plugins/Experimental/Avalanche/Source/AvalancheTransitionEditor/Private/AppModes/AvaTransitionParameterMode.cpp
new file mode 100644
--- /dev/null
+++ b/Engine/Plugins/Experimental/Avalanche/Source/AvalancheTransitionEditor/Private/AppModes/AvaTransitionParameterMode.cpp
@@ -0,0 +1,29 @@
+// Copyright Epic Games, Inc. All Rights Reserved.
+
+#include "AvaTransitionParameterMode.h"
+#include "TabFactories/AvaTransitionParameterTabFactory.h"
+
+#define LOCTEXT_NAMESPACE "AvaTransitionParameterMode"
+
+FAvaTransitionParameterMode::FAvaTransitionParameterMode(const TSharedRef<FAvaTransitionEditor>& InEditor)
+	: FAvaTransitionAppMode(InEditor, EAvaTransitionEditorMode::Parameter)
+{
+	WorkspaceMenuCategory = FWorkspaceItem::NewGroup(LOCTEXT("WorkspaceMenuCategory", "Motion Design Transition Parameter"));
+
+	TabFactories.RegisterFactory(MakeShared<FAvaTransitionParameterTabFactory>(InEditor));
+
+	TabLayout = FTabManager::NewLayout("AvaTransitionEditor_Parameter_Layout_V0_1")
+	->AddArea
+	(
+		FTabManager::NewPrimaryArea()
+		->SetOrientation(Orient_Vertical)
+		->Split
+		(
+			FTabManager::NewStack()
+			->SetSizeCoefficient(1.f)
+			->AddTab(FAvaTransitionParameterTabFactory::TabId, ETabState::OpenedTab)
+		)
+	);
+}
+
+#undef LOCTEXT_NAMESPACE
diff --git a/Engine/Plugins/Experimental/Avalanche/Source/AvalancheTransitionEditor/Private/AppModes/AvaTransitionParameterMode.h b/Engine/Plugins/Experimental/Avalanche/Source/AvalancheTransitionEditor/Private/AppModes/AvaTransitionParameterMode.h
new file mode 100644
--- /dev/null
+++ b/Engine/Plugins/Experimental/Avalanche/Source/AvalancheTransitionEditor/Private/AppModes/AvaTransitionParameterMode.h
@@ -0,0 +1,11 @@
+// Copyright Epic Games, Inc. All Rights Reserved.
+
+#pragma once
+
+#include "AvaTransitionAppMode.h"
+
+class FAvaTransitionParameterMode : public FAvaTransitionAppMode
+{
+public:
+	FAvaTransitionParameterMode(const TSharedRef<FAvaTransitionEditor>& InEditor);
+};
diff --git a/Engine/Plugins/Experimental/Avalanche/Source/AvalancheTransitionEditor/Private/AppModes/AvaTransitionSimpleMode.cpp b/Engine/Plugins/Experimental/Avalanche/Source/AvalancheTransitionEditor/Private/AppModes/AvaTransitionSimpleMode.cpp
--- a/Engine/Plugins/Experimental/Avalanche/Source/AvalancheTransitionEditor/Private/AppModes/AvaTransitionSimpleMode.cpp
+++ b/Engine/Plugins/Experimental/Avalanche/Source/AvalancheTransitionEditor/Private/AppModes/AvaTransitionSimpleMode.cpp
@@ -3,6 +3,7 @@
 #include "AvaTransitionSimpleMode.h"
 #include "TabFactories/AvaTransitionCompilerResultsTabFactory.h"
 #include "TabFactories/AvaTransitionSelectionDetailsTabFactory.h"
+#include "TabFactories/AvaTransitionTreeDetailsTabFactory.h"
 #include "TabFactories/AvaTransitionTreeTabFactory.h"
 
 #define LOCTEXT_NAMESPACE "AvaTransitionSimpleMode"
@@ -10,9 +11,11 @@
 FAvaTransitionSimpleMode::FAvaTransitionSimpleMode(const TSharedRef<FAvaTransitionEditor>& InEditor)
 	: FAvaTransitionAppMode(InEditor, EAvaTransitionEditorMode::Default)
 {
+	RegisterDefaultTabFactories();
+
 	WorkspaceMenuCategory = FWorkspaceItem::NewGroup(LOCTEXT("WorkspaceMenuCategory", "Motion Design Transition Simple"));
 
-	TabLayout = FTabManager::NewLayout("AvaTransitionEditor_Simple_Layout_V0_1")
+	TabLayout = FTabManager::NewLayout("AvaTransitionEditor_Simple_Layout_V0_2")
 	->AddArea
 	(
 		FTabManager::NewPrimaryArea()
@@ -22,9 +25,15 @@ FAvaTransitionSimpleMode::FAvaTransitionSimpleMode(const TSharedRef<FAvaTransiti
 			FTabManager::NewSplitter()
 			->SetOrientation(Orient_Horizontal)
 			->Split
+			(
+				FTabManager::NewStack()
+				->SetSizeCoefficient(0.2f)
+				->AddTab(FAvaTransitionTreeDetailsTabFactory::TabId, ETabState::ClosedTab)
+			)
+			->Split
 			(
 				FTabManager::NewSplitter()
-				->SetSizeCoefficient(0.8f)
+				->SetSizeCoefficient(0.6f)
 				->SetOrientation(Orient_Vertical)
 				->Split
 				(
diff --git a/Engine/Plugins/Experimental/Avalanche/Source/AvalancheTransitionEditor/Private/AvaTransitionEditor.cpp b/Engine/Plugins/Experimental/Avalanche/Source/AvalancheTransitionEditor/Private/AvaTransitionEditor.cpp
--- a/Engine/Plugins/Experimental/Avalanche/Source/AvalancheTransitionEditor/Private/AvaTransitionEditor.cpp
+++ b/Engine/Plugins/Experimental/Avalanche/Source/AvalancheTransitionEditor/Private/AvaTransitionEditor.cpp
@@ -2,6 +2,7 @@
 
 #include "AvaTransitionEditor.h"
 #include "AppModes/AvaTransitionAdvancedMode.h"
+#include "AppModes/AvaTransitionParameterMode.h"
 #include "AppModes/AvaTransitionSimpleMode.h"
 #include "AvaTransitionTree.h"
 #include "AvaTransitionTreeEditorData.h"
@@ -63,8 +64,9 @@ void FAvaTransitionEditor::RegisterApplicationModes()
 {
 	TSharedRef<FAvaTransitionEditor> This = SharedThis(this);
 
-	const TArray<TSharedRef<FAvaTransitionAppMode>, TFixedAllocator<2>> AppModes = 
+	const TArray<TSharedRef<FAvaTransitionAppMode>, TFixedAllocator<3>> AppModes = 
 	{
+		MakeShared<FAvaTransitionParameterMode>(This),
 		MakeShared<FAvaTransitionSimpleMode>(This),
 		MakeShared<FAvaTransitionAdvancedMode>(This),
 	};
@@ -76,7 +78,7 @@ void FAvaTransitionEditor::RegisterApplicationModes()
 	}
 
 	// set default mode to the first app mode
-	SetCurrentMode(AppModes[0]->GetModeName());
+	SetCurrentMode(FAvaTransitionAppMode::StaticGetModeName(EAvaTransitionEditorMode::Default));
 }
 
 void FAvaTransitionEditor::ExtendMenus()
diff --git a/Engine/Plugins/Experimental/Avalanche/Source/AvalancheTransitionEditor/Private/AvaTransitionEditorEnums.h b/Engine/Plugins/Experimental/Avalanche/Source/AvalancheTransitionEditor/Private/AvaTransitionEditorEnums.h
--- a/Engine/Plugins/Experimental/Avalanche/Source/AvalancheTransitionEditor/Private/AvaTransitionEditorEnums.h
+++ b/Engine/Plugins/Experimental/Avalanche/Source/AvalancheTransitionEditor/Private/AvaTransitionEditorEnums.h
@@ -6,6 +6,12 @@
 
 enum class EAvaTransitionEditorMode : uint8
 {
+	/** Simplified view of the State Tree. Core features for most Transition Logic Needs */
 	Default,
+
+	/** All the features in State Tree */
 	Advanced,
+
+	/** Only Parameters shown */
+	Parameter,
 };
diff --git a/Engine/Plugins/Experimental/Avalanche/Source/AvalancheTransitionEditor/Private/Customizations/AvaTransitionTreeEditorDataCustomization.cpp b/Engine/Plugins/Experimental/Avalanche/Source/AvalancheTransitionEditor/Private/Customizations/AvaTransitionTreeEditorDataCustomization.cpp
--- a/Engine/Plugins/Experimental/Avalanche/Source/AvalancheTransitionEditor/Private/Customizations/AvaTransitionTreeEditorDataCustomization.cpp
+++ b/Engine/Plugins/Experimental/Avalanche/Source/AvalancheTransitionEditor/Private/Customizations/AvaTransitionTreeEditorDataCustomization.cpp
@@ -1,22 +1,37 @@
 // Copyright Epic Games, Inc. All Rights Reserved.
 
 #include "AvaTransitionTreeEditorDataCustomization.h"
+#include "AvaTransitionEditorEnums.h"
 #include "AvaTransitionTreeEditorData.h"
 #include "DetailCategoryBuilder.h"
 #include "DetailLayoutBuilder.h"
 #include "Modules/ModuleManager.h"
+#include "PropertyBagDetails.h"
 #include "PropertyEditorModule.h"
+#include "StateTreeEditorStyle.h"
+#include "ViewModels/AvaTransitionViewModelSharedData.h"
 
-TSharedRef<IDetailCustomization> FAvaTransitionTreeEditorDataCustomization::MakeInstance()
+TSharedRef<IDetailCustomization> FAvaTransitionTreeEditorDataCustomization::MakeInstance(TWeakPtr<FAvaTransitionViewModelSharedData> InSharedDataWeak)
+{
+	return MakeShared<FAvaTransitionTreeEditorDataCustomization>(InSharedDataWeak);
+}
+
+FAvaTransitionTreeEditorDataCustomization::FAvaTransitionTreeEditorDataCustomization(const TWeakPtr<FAvaTransitionViewModelSharedData>& InSharedDataWeak)
+	: SharedDataWeak(InSharedDataWeak)
 {
-	return MakeShared<FAvaTransitionTreeEditorDataCustomization>();
 }
 
 void FAvaTransitionTreeEditorDataCustomization::CustomizeDetails(IDetailLayoutBuilder& InDetailBuilder)
 {
-	if (TSharedPtr<IDetailCustomization> Customization = GetDefaultCustomization())
+	const EAvaTransitionEditorMode EditorMode = GetEditorMode();
+
+	// Parameters don't need customization as its only purpose is to show the Parameters in a fixed layout way (FAvaTransitionTreeEditorDataCustomization::CustomizeParameters)
+	if (EditorMode != EAvaTransitionEditorMode::Parameter)
 	{
-		Customization->CustomizeDetails(InDetailBuilder);
+		if (TSharedPtr<IDetailCustomization> Customization = GetDefaultCustomization())
+		{
+			Customization->CustomizeDetails(InDetailBuilder);
+		}
 	}
 
 	// Hide property as it's going to show in the Toolbar
@@ -24,6 +39,46 @@ void FAvaTransitionTreeEditorDataCustomization::CustomizeDetails(IDetailLayoutBu
 		, UAvaTransitionTreeEditorData::StaticClass());
 
 	InDetailBuilder.HideProperty(LayerHandle);
+
+	if (EditorMode != EAvaTransitionEditorMode::Advanced)
+	{
+		InDetailBuilder.HideCategory(TEXT("Evaluators"));
+		InDetailBuilder.HideCategory(TEXT("Global Tasks"));
+	}
+
+	if (EditorMode == EAvaTransitionEditorMode::Parameter)
+	{
+		InDetailBuilder.HideCategory(TEXT("Theme"));
+		CustomizeParameters(InDetailBuilder);
+	}
+}
+
+void FAvaTransitionTreeEditorDataCustomization::CustomizeParameters(IDetailLayoutBuilder& InDetailBuilder)
+{
+	TSharedPtr<IPropertyHandle> RootParametersHandle = InDetailBuilder.GetProperty(TEXT("RootParameters"));
+	check(RootParametersHandle);
+	RootParametersHandle->MarkHiddenByCustomization();
+
+	TSharedPtr<IPropertyHandle> PropertyBagParametersHandle = RootParametersHandle->GetChildHandle(TEXT("Parameters")); // FInstancedPropertyBag
+	check(PropertyBagParametersHandle);
+
+	TSharedRef<FPropertyBagInstanceDataDetails> InstanceDetails = MakeShared<FPropertyBagInstanceDataDetails>(PropertyBagParametersHandle
+		, InDetailBuilder.GetPropertyUtilities()
+		, /*bFixedLayout*/true);
+
+	InDetailBuilder.EditCategory(TEXT("Parameters"))
+		.AddCustomBuilder(InstanceDetails);
+}
+
+EAvaTransitionEditorMode FAvaTransitionTreeEditorDataCustomization::GetEditorMode() const
+{
+	if (TSharedPtr<FAvaTransitionViewModelSharedData> SharedData = SharedDataWeak.Pin())
+	{
+		return SharedData->GetEditorMode();
+	}
+
+	// Fallback to Advanced as it is the one with most features
+	return EAvaTransitionEditorMode::Advanced;
 }
 
 TSharedPtr<IDetailCustomization> FAvaTransitionTreeEditorDataCustomization::GetDefaultCustomization() const
diff --git a/Engine/Plugins/Experimental/Avalanche/Source/AvalancheTransitionEditor/Private/Customizations/AvaTransitionTreeEditorDataCustomization.h b/Engine/Plugins/Experimental/Avalanche/Source/AvalancheTransitionEditor/Private/Customizations/AvaTransitionTreeEditorDataCustomization.h
--- a/Engine/Plugins/Experimental/Avalanche/Source/AvalancheTransitionEditor/Private/Customizations/AvaTransitionTreeEditorDataCustomization.h
+++ b/Engine/Plugins/Experimental/Avalanche/Source/AvalancheTransitionEditor/Private/Customizations/AvaTransitionTreeEditorDataCustomization.h
@@ -4,11 +4,16 @@
 
 #include "IDetailCustomization.h"
 
+class FAvaTransitionViewModelSharedData;
+enum class EAvaTransitionEditorMode : uint8;
+
 /** Customization that re-uses the module-registered Customization for UStateTreeEditorData, and tweaks a few settings for it */
 class FAvaTransitionTreeEditorDataCustomization : public IDetailCustomization
 {
 public:
-	static TSharedRef<IDetailCustomization> MakeInstance();
+	static TSharedRef<IDetailCustomization> MakeInstance(TWeakPtr<FAvaTransitionViewModelSharedData> InSharedDataWeak);
+
+	FAvaTransitionTreeEditorDataCustomization(const TWeakPtr<FAvaTransitionViewModelSharedData>& InSharedDataWeak);
 
 protected:
 	//~ Begin IDetailCustomization
@@ -16,5 +21,11 @@ protected:
 	//~ End IDetailCustomization
 
 private:
+	void CustomizeParameters(IDetailLayoutBuilder& InDetailBuilder);
+
+	EAvaTransitionEditorMode GetEditorMode() const;
+
 	TSharedPtr<IDetailCustomization> GetDefaultCustomization() const;
+
+	TWeakPtr<FAvaTransitionViewModelSharedData> SharedDataWeak;
 };
diff --git a/Engine/Plugins/Experimental/Avalanche/Source/AvalancheTransitionEditor/Private/TabFactories/AvaTransitionParameterTabFactory.cpp b/Engine/Plugins/Experimental/Avalanche/Source/AvalancheTransitionEditor/Private/TabFactories/AvaTransitionParameterTabFactory.cpp
new file mode 100644
--- /dev/null
+++ b/Engine/Plugins/Experimental/Avalanche/Source/AvalancheTransitionEditor/Private/TabFactories/AvaTransitionParameterTabFactory.cpp
@@ -0,0 +1,41 @@
+// Copyright Epic Games, Inc. All Rights Reserved.
+
+#include "AvaTransitionParameterTabFactory.h"
+#include "Styling/SlateIconFinder.h"
+#include "ViewModels/AvaTransitionEditorViewModel.h"
+#include "Views/SAvaTransitionParameterDetails.h"
+#include "Widgets/SNullWidget.h"
+
+#define LOCTEXT_NAMESPACE "AvaTransitionParameterTabFactory"
+
+const FName FAvaTransitionParameterTabFactory::TabId(TEXT("AvaTransitionParameter"));
+
+FAvaTransitionParameterTabFactory::FAvaTransitionParameterTabFactory(const TSharedRef<FAvaTransitionEditor>& InEditor)
+	: FAvaTransitionTabFactory(TabId, InEditor)
+{
+	TabIcon             = FSlateIcon(FAppStyle::GetAppStyleSetName(), "LevelEditor.Tabs.Details");
+	TabLabel            = LOCTEXT("TabLabel", "Parameters");
+	ViewMenuTooltip     = LOCTEXT("ViewMenuTooltip", "State Tree Parameters");
+	ViewMenuDescription = LOCTEXT("ViewMenuDescription", "State Tree Parameters");
+	bIsSingleton        = true;
+	ReadOnlyBehavior    = ETabReadOnlyBehavior::Custom;
+}
+
+TSharedRef<SWidget> FAvaTransitionParameterTabFactory::CreateTabBody(const FWorkflowTabSpawnInfo& Info) const
+{
+	TSharedPtr<FAvaTransitionEditor> Editor = GetEditor();
+	if (!ensure(Editor.IsValid()))
+	{
+		return SNullWidget::NullWidget;
+	}
+
+	TSharedPtr<FAvaTransitionEditorViewModel> EditorViewModel = Editor->GetEditorViewModel();
+	if (!ensure(EditorViewModel.IsValid()))
+	{
+		return SNullWidget::NullWidget;
+	}
+
+	return SNew(SAvaTransitionParameterDetails, EditorViewModel.ToSharedRef());
+}
+
+#undef LOCTEXT_NAMESPACE
diff --git a/Engine/Plugins/Experimental/Avalanche/Source/AvalancheTransitionEditor/Private/TabFactories/AvaTransitionParameterTabFactory.h b/Engine/Plugins/Experimental/Avalanche/Source/AvalancheTransitionEditor/Private/TabFactories/AvaTransitionParameterTabFactory.h
new file mode 100644
--- /dev/null
+++ b/Engine/Plugins/Experimental/Avalanche/Source/AvalancheTransitionEditor/Private/TabFactories/AvaTransitionParameterTabFactory.h
@@ -0,0 +1,17 @@
+// Copyright Epic Games, Inc. All Rights Reserved.
+
+#pragma once
+
+#include "AvaTransitionTabFactory.h"
+
+class FAvaTransitionParameterTabFactory : public FAvaTransitionTabFactory
+{
+public:
+	static const FName TabId;
+
+	explicit FAvaTransitionParameterTabFactory(const TSharedRef<FAvaTransitionEditor>& InEditor);
+
+	//~ Begin FWorkflowTabFactory
+	virtual TSharedRef<SWidget> CreateTabBody(const FWorkflowTabSpawnInfo& InInfo) const override;
+	//~ End FWorkflowTabFactory
+};
diff --git a/Engine/Plugins/Experimental/Avalanche/Source/AvalancheTransitionEditor/Private/TabFactories/AvaTransitionTreeDetailsTabFactory.cpp b/Engine/Plugins/Experimental/Avalanche/Source/AvalancheTransitionEditor/Private/TabFactories/AvaTransitionTreeDetailsTabFactory.cpp
--- a/Engine/Plugins/Experimental/Avalanche/Source/AvalancheTransitionEditor/Private/TabFactories/AvaTransitionTreeDetailsTabFactory.cpp
+++ b/Engine/Plugins/Experimental/Avalanche/Source/AvalancheTransitionEditor/Private/TabFactories/AvaTransitionTreeDetailsTabFactory.cpp
@@ -2,6 +2,7 @@
 
 #include "AvaTransitionTreeDetailsTabFactory.h"
 #include "ViewModels/AvaTransitionEditorViewModel.h"
+#include "Views/SAvaTransitionTreeDetails.h"
 #include "Widgets/SNullWidget.h"
 
 #define LOCTEXT_NAMESPACE "AvaTransitionTreeDetailsTabFactory"
@@ -33,7 +34,7 @@ TSharedRef<SWidget> FAvaTransitionTreeDetailsTabFactory::CreateTabBody(const FWo
 		return SNullWidget::NullWidget;
 	}
 
-	return EditorViewModel->GetTreeDetails();
+	return SNew(SAvaTransitionTreeDetails, EditorViewModel.ToSharedRef());
 }
 
 #undef LOCTEXT_NAMESPACE
diff --git a/Engine/Plugins/Experimental/Avalanche/Source/AvalancheTransitionEditor/Private/ViewModels/AvaTransitionEditorViewModel.cpp b/Engine/Plugins/Experimental/Avalanche/Source/AvalancheTransitionEditor/Private/ViewModels/AvaTransitionEditorViewModel.cpp
--- a/Engine/Plugins/Experimental/Avalanche/Source/AvalancheTransitionEditor/Private/ViewModels/AvaTransitionEditorViewModel.cpp
+++ b/Engine/Plugins/Experimental/Avalanche/Source/AvalancheTransitionEditor/Private/ViewModels/AvaTransitionEditorViewModel.cpp
@@ -21,7 +21,6 @@
 #include "StateTreeDelegates.h"
 #include "StateTreeEditorSettings.h"
 #include "TabFactories/AvaTransitionCompilerResultsTabFactory.h"
-#include "Views/SAvaTransitionTreeDetails.h"
 #include "Views/SAvaTransitionTreeView.h"
 #include "Widgets/Layout/SScrollBar.h"
 #include "Widgets/Layout/SScrollBox.h"
@@ -168,15 +167,6 @@ TSharedRef<SWidget> FAvaTransitionEditorViewModel::GetTreeWidget()
 		];
 }
 
-TSharedRef<SWidget> FAvaTransitionEditorViewModel::GetTreeDetails()
-{
-	if (!TreeDetails.IsValid())
-	{
-		TreeDetails = SNew(SAvaTransitionTreeDetails, SharedThis(this));
-	}
-	return TreeDetails.ToSharedRef();
-}
-
 void FAvaTransitionEditorViewModel::OnInitialize()
 {
 	FAvaTransitionViewModel::OnInitialize();
@@ -204,13 +194,10 @@ void FAvaTransitionEditorViewModel::PostRefresh()
 {
 	FAvaTransitionViewModel::PostRefresh();
 
-	if (TreeDetails.IsValid())
-	{
-		TreeDetails->Refresh();
-	}
-
 	RefreshTreeView();
 	UpdateTree();
+
+	OnPostRefreshDelegate.Broadcast();
 }
 
 void FAvaTransitionEditorViewModel::RefreshTreeView()
diff --git a/Engine/Plugins/Experimental/Avalanche/Source/AvalancheTransitionEditor/Private/ViewModels/AvaTransitionEditorViewModel.h b/Engine/Plugins/Experimental/Avalanche/Source/AvalancheTransitionEditor/Private/ViewModels/AvaTransitionEditorViewModel.h
--- a/Engine/Plugins/Experimental/Avalanche/Source/AvalancheTransitionEditor/Private/ViewModels/AvaTransitionEditorViewModel.h
+++ b/Engine/Plugins/Experimental/Avalanche/Source/AvalancheTransitionEditor/Private/ViewModels/AvaTransitionEditorViewModel.h
@@ -17,7 +17,6 @@ class FAvaTransitionTreeContextMenu;
 class FUICommandList;
 class IAvaTransitionSelectableExtension;
 class IMessageToken;
-class SAvaTransitionTreeDetails;
 class SAvaTransitionTreeView;
 class UAvaTransitionTree;
 class UAvaTransitionTreeEditorData;
@@ -31,6 +30,12 @@ class FAvaTransitionEditorViewModel : public FAvaTransitionViewModel, public FSe
 public:
 	UE_AVA_INHERITS(FAvaTransitionEditorViewModel, FAvaTransitionViewModel, IAvaTransitionObjectExtension)
 
+	DECLARE_MULTICAST_DELEGATE(FOnPostRefresh);
+	FOnPostRefresh::RegistrationType& GetOnPostRefresh() const
+	{
+		return OnPostRefreshDelegate;
+	}
+
 	explicit FAvaTransitionEditorViewModel(UAvaTransitionTree* InTransitionTree, const TSharedPtr<FAvaTransitionEditor>& InEditor);
 
 	virtual ~FAvaTransitionEditorViewModel() override;
@@ -74,8 +79,6 @@ public:
 
 	TSharedRef<SWidget> GetTreeWidget();
 
-	TSharedRef<SWidget> GetTreeDetails();
-
 	//~ Begin FAvaTransitionViewModel
 	virtual void OnInitialize() override;
 	virtual void PostRefresh() override;
@@ -121,7 +124,7 @@ private:
 
 	TSharedPtr<SAvaTransitionTreeView> TreeView;
 
-	TSharedPtr<SAvaTransitionTreeDetails> TreeDetails;
-
 	TArray<TSharedRef<FAvaTransitionActions>> Actions;
+
+	mutable FOnPostRefresh OnPostRefreshDelegate;
 };
diff --git a/Engine/Plugins/Experimental/Avalanche/Source/AvalancheTransitionEditor/Private/Views/SAvaTransitionParameterDetails.cpp b/Engine/Plugins/Experimental/Avalanche/Source/AvalancheTransitionEditor/Private/Views/SAvaTransitionParameterDetails.cpp
new file mode 100644
--- /dev/null
+++ b/Engine/Plugins/Experimental/Avalanche/Source/AvalancheTransitionEditor/Private/Views/SAvaTransitionParameterDetails.cpp
@@ -0,0 +1,12 @@
+// Copyright Epic Games, Inc. All Rights Reserved.
+
+#include "SAvaTransitionParameterDetails.h"
+#include "SAvaTransitionTreeDetails.h"
+
+void SAvaTransitionParameterDetails::Construct(const FArguments& InArgs, const TSharedRef<FAvaTransitionEditorViewModel>& InEditorViewModel)
+{
+	ChildSlot
+	[
+		SNew(SAvaTransitionTreeDetails, InEditorViewModel)
+	];
+}
diff --git a/Engine/Plugins/Experimental/Avalanche/Source/AvalancheTransitionEditor/Private/Views/SAvaTransitionParameterDetails.h b/Engine/Plugins/Experimental/Avalanche/Source/AvalancheTransitionEditor/Private/Views/SAvaTransitionParameterDetails.h
new file mode 100644
--- /dev/null
+++ b/Engine/Plugins/Experimental/Avalanche/Source/AvalancheTransitionEditor/Private/Views/SAvaTransitionParameterDetails.h
@@ -0,0 +1,16 @@
+// Copyright Epic Games, Inc. All Rights Reserved.
+
+#pragma once
+
+#include "Widgets/SCompoundWidget.h"
+
+class FAvaTransitionEditorViewModel;
+
+class SAvaTransitionParameterDetails : public SCompoundWidget
+{
+public:
+	SLATE_BEGIN_ARGS(SAvaTransitionParameterDetails){}
+	SLATE_END_ARGS()
+
+	void Construct(const FArguments& InArgs, const TSharedRef<FAvaTransitionEditorViewModel>& InEditorViewModel);
+};
diff --git a/Engine/Plugins/Experimental/Avalanche/Source/AvalancheTransitionEditor/Private/Views/SAvaTransitionTreeDetails.cpp b/Engine/Plugins/Experimental/Avalanche/Source/AvalancheTransitionEditor/Private/Views/SAvaTransitionTreeDetails.cpp
--- a/Engine/Plugins/Experimental/Avalanche/Source/AvalancheTransitionEditor/Private/Views/SAvaTransitionTreeDetails.cpp
+++ b/Engine/Plugins/Experimental/Avalanche/Source/AvalancheTransitionEditor/Private/Views/SAvaTransitionTreeDetails.cpp
@@ -13,17 +13,20 @@ void SAvaTransitionTreeDetails::Construct(const FArguments& InArgs, const TShare
 {
 	EditorViewModelWeak = InEditorViewModel;
 
-	FDetailsViewArgs DetailsViewArgs;
-	DetailsViewArgs.NameAreaSettings = FDetailsViewArgs::HideNameArea;
+	TSharedRef<FAvaTransitionViewModelSharedData> SharedData = InEditorViewModel->GetSharedData();
 
 	FPropertyEditorModule& PropertyEditorModule = FModuleManager::GetModuleChecked<FPropertyEditorModule>("PropertyEditor");
 
+	FDetailsViewArgs DetailsViewArgs;
+	DetailsViewArgs.NameAreaSettings = FDetailsViewArgs::HideNameArea;
+
 	DetailsView = PropertyEditorModule.CreateDetailView(DetailsViewArgs);
+
 	DetailsView->RegisterInstancedCustomPropertyLayout(UStateTreeEditorData::StaticClass()
-		, FOnGetDetailCustomizationInstance::CreateStatic(&FAvaTransitionTreeEditorDataCustomization::MakeInstance));
+		, FOnGetDetailCustomizationInstance::CreateStatic(&FAvaTransitionTreeEditorDataCustomization::MakeInstance, SharedData->AsWeak()));
 
 	// Read-only
-	if (InEditorViewModel->GetSharedData()->IsReadOnly())
+	if (SharedData->IsReadOnly())
 	{
 		DetailsView->SetIsPropertyEditingEnabledDelegate(FIsPropertyEditingEnabled::CreateLambda([]{ return false; }));
 	}
@@ -34,6 +37,17 @@ void SAvaTransitionTreeDetails::Construct(const FArguments& InArgs, const TShare
 	];
 
 	Refresh();
+
+	OnRefreshHandle = InEditorViewModel->GetOnPostRefresh().AddSP(this, &SAvaTransitionTreeDetails::Refresh);
+}
+
+SAvaTransitionTreeDetails::~SAvaTransitionTreeDetails()
+{
+	if (TSharedPtr<FAvaTransitionEditorViewModel> EditorViewModel = EditorViewModelWeak.Pin())
+	{
+		EditorViewModel->GetOnPostRefresh().Remove(OnRefreshHandle);
+		OnRefreshHandle.Reset();
+	}
 }
 
 void SAvaTransitionTreeDetails::Refresh()
@@ -51,5 +65,5 @@ void SAvaTransitionTreeDetails::Refresh()
 	}
 
 	check(DetailsView.IsValid());
-	DetailsView->SetObject(EditorData);
+	DetailsView->SetObject(EditorData, /*bForceRefresh*/true);
 }
diff --git a/Engine/Plugins/Experimental/Avalanche/Source/AvalancheTransitionEditor/Private/Views/SAvaTransitionTreeDetails.h b/Engine/Plugins/Experimental/Avalanche/Source/AvalancheTransitionEditor/Private/Views/SAvaTransitionTreeDetails.h
--- a/Engine/Plugins/Experimental/Avalanche/Source/AvalancheTransitionEditor/Private/Views/SAvaTransitionTreeDetails.h
+++ b/Engine/Plugins/Experimental/Avalanche/Source/AvalancheTransitionEditor/Private/Views/SAvaTransitionTreeDetails.h
@@ -15,10 +15,14 @@ public:
 
 	void Construct(const FArguments& InArgs, const TSharedRef<FAvaTransitionEditorViewModel>& InEditorViewModel);
 
-	void Refresh();
+	virtual ~SAvaTransitionTreeDetails() override;
 
 private:
+	void Refresh();
+
 	TWeakPtr<FAvaTransitionEditorViewModel> EditorViewModelWeak;
 
 	TSharedPtr<IDetailsView> DetailsView;
+
+	FDelegateHandle OnRefreshHandle;
 };
