[Motion Design] Rundown Server: Fix references held by the rundown server that prevents deleting rundown assets.

1 - Rundown Server holding reference preventing assets from being deleted

The rundown server is considered an editor and should hold a references to what is being edited as a principle (similar to FAssetEditorToolkit::FGCEditingObjects).

The rundown server maintains contextual caches for each of the commands. Contexts are split in editing and playback categories, each requiring different contextual resources (i.e. either managed instances for page editing through RC or playback instances for actual playback). The rundown server will try to hold on to those contexts in order to save time on successive requests, but it shouldn't prevent assets from being deleted if it makes sense to do so. Fortunately, FEditorDelegates::OnAssetsPreDelete can be used to have the rundown server flush it's command contexts when an asset is about to be deleted. This works both from the content browser and other commands deleting assets through ObjectTools asset and object deleting functions.

In order to make the editing and playback contexts easier to "flush", some of the code around what was FRundownEditCommandData and FRundownPlaybackCommandData has been reworked. I found it more appropriate to call those "contexts" and move all the common "rundown server to rundown" delegates binding to a separate object: FRundownEntry, that has it's own separate cache and can be shared amongts the command contexts. Given that we have independent command contexts and, in the future, it is to be expected that there will be more (i.e. one editing context per client for instance), we really want to decouple the per-rundown objects from the per-command context ones. The FRundowEntry manages the rundown server's delegates in the rundown, while the command contexts manage the associated resources. I think it is clearer that way.

1.1 - Playback Manager Delegates

Moved the Rundown Server to Playback Manager's delegate binding to one place at init of the rundown server. It is less convoluted (imho) and there will potentially be more delegates to listen to (i.e. OnPlaybackInstanceInvalidated, OnLocalPlaybackAssetRemoved) in the future.

As I found out, it seems the order of "post engine init" is not garanteed, and it seems the AvaMediaEditor module is called before AvaMedia module's post engine init callback. The result is that the rundown server is started before some of the needed objects in ava media module are created. I worked around that by having some internal functions to ensure those objects (playback manager and managed instance cache) are created. The alternative would be to have a AvaMedia module have it's own PostInit event and AvaMediaEditor be a callback to that (might be more future proof?).

Also, as a precaution related to that, I added some functions to IAvaMediaModule to query the availability of the playback manager and instance cache so that shutdown sequences don't trigger check failures. I didn't get any, but it seems more prudent to have logic to avoid it in any case.

1.2 - Playback Context Management improvement:

It is intended for the rundown server and the rundow editor to work togheter (as well as possible) on the same rundown.
Example: stream deck can send some playback commands while the rundown editor is opened.
We want a rundown's playback context to only be closed when it is no longer needed by either editor or server.

The previous issue was if the rundown server would switch from one rundown to another, in the playback context, it would close the rundown's playback context, even if the editor was opened.
To fix that, a delegate has been added to query if the playback context can be closed. So that either the server or editor can keep it alive if needed and closed only when it no longer is.
This solution may be temporary until the refactor to have a separate rundown player (or playback context) object that can be either instanced exclusively or shared, which would be much more powerfull.

1.2.1 - New Playback Context Switch message

Adding a message to give feedback to the control app client about playback context switch. It is an important information to keep track of and having a dedicated message for it makes it more reliable.

Also added more fields to the FAvaRundownAssetsChanged to make it more usefull. I plan to add more event types in the future.

diff --git a/Engine/Plugins/Experimental/Avalanche/Source/AvalancheMedia/Internal/Rundown/AvaRundownMessages.h b/Engine/Plugins/Experimental/Avalanche/Source/AvalancheMedia/Internal/Rundown/AvaRundownMessages.h
--- a/Engine/Plugins/Experimental/Avalanche/Source/AvalancheMedia/Internal/Rundown/AvaRundownMessages.h
+++ b/Engine/Plugins/Experimental/Avalanche/Source/AvalancheMedia/Internal/Rundown/AvaRundownMessages.h
@@ -234,8 +234,10 @@ public:
 /**
  *	Request that the given rundown be loaded for playback.
  *	This will also open an associated playback context.
- *	Only one rundown can be opened for playback at a time. If another rundown
- *	is opened, it will be closed and all currently playing pages stopped.
+ *	Only one rundown can be opened for playback at a time by the rundown server.
+ *	If another rundown is opened, the previous one will be closed and all currently playing pages stopped,
+ *	unless the rundown editor is opened. The rundown editor will keep the playback context alive.
+ *	
  *	If the path is empty, nothing will be done and the server will reply with
  *	a FAvaRundownServerMsg message indicating which rundown is currently loaded.
  */
@@ -245,6 +247,9 @@ struct FAvaRundownLoadRundown : public FAvaRundownMsgBase
 	GENERATED_BODY()
 
 public:
+	/**
+	 * Rundown asset path: [PackagePath]/[AssetName].[AssetName]
+	 */
 	UPROPERTY()
 	FString Rundown;
 };
@@ -375,6 +380,28 @@ public:
 	bool bOnlyIfIsDirty = false;
 };
 
+/**
+ * Rundown specific events broadcast by the server to help status display or related contexts in control applications.
+ */
+USTRUCT()
+struct FAvaRundownPlaybackContextChanged : public FAvaRundownMsgBase
+{
+	GENERATED_BODY()
+
+public:
+	/**
+	 * Previous rundown (can be empty).
+	 */
+	UPROPERTY()
+	FString PreviousRundown;
+	
+	/**
+	 * New current rundown (can be empty).
+	 */
+	UPROPERTY()
+	FString NewRundown;
+};
+
 /**
  * Request the list of pages from the given rundown.
  */
@@ -954,16 +981,50 @@ public:
 	TArray<FAvaRundownChannel> Channels;
 };
 
+/**
+ * Generic asset event
+ */
+UENUM()
+enum class EAvaRundownAssetEvent : uint8
+{
+	Unknown = 0,
+	Added,
+	Removed,
+	//Saved, // todo
+	//Modified // todo
+};
+
+/**
+ * Event broadcast when an asset event occurs on the server.
+ */
 USTRUCT()
 struct FAvaRundownAssetsChanged : public FAvaRundownMsgBase
 {
 	GENERATED_BODY()
 
+	/** Asset name only, without the package path. (Keeping for legacy) */
 	UPROPERTY()
 	FString AssetName;
+
+	/** Full asset path: /PackagePath/PackageName.AssetName */
+	UPROPERTY()
+	FString AssetPath;
+
+	/** Full asset class path. */
+	UPROPERTY()
+	FString AssetClass;
+
+	/** true if the asset is a "playable" asset, i.e. an asset that can be set in a page's asset. */
+	UPROPERTY()
+	bool bIsPlayable = false;
+
+	UPROPERTY()
+	EAvaRundownAssetEvent EventType = EAvaRundownAssetEvent::Unknown;
 };
 
-// Channel actions
+/**
+ * Channel actions
+ */
 UENUM()
 enum class EAvaRundownChannelActions
 {
diff --git a/Engine/Plugins/Experimental/Avalanche/Source/AvalancheMedia/Private/AvaMediaModule.cpp b/Engine/Plugins/Experimental/Avalanche/Source/AvalancheMedia/Private/AvaMediaModule.cpp
--- a/Engine/Plugins/Experimental/Avalanche/Source/AvalancheMedia/Private/AvaMediaModule.cpp
+++ b/Engine/Plugins/Experimental/Avalanche/Source/AvalancheMedia/Private/AvaMediaModule.cpp
@@ -156,6 +156,9 @@ void FAvaMediaModule::StartupModule()
 
 	FCoreDelegates::OnPostEngineInit.AddRaw(this, &FAvaMediaModule::PostEngineInit);
 	FCoreDelegates::OnEnginePreExit.AddRaw(this, &FAvaMediaModule::EnginePreExit);
+#if WITH_EDITOR
+	FEditorDelegates::PrePIEEnded.AddRaw(this, &FAvaMediaModule::PrePIEEnded);
+#endif
 	
 	FString DummyServerName;
 	if (IsPlaybackServerManuallyStarted(DummyServerName) || IsRundownServerManuallyStarted(DummyServerName))
@@ -186,6 +189,9 @@ void FAvaMediaModule::ShutdownModule()
 	
 	FCoreDelegates::OnPostEngineInit.RemoveAll(this);
 	FCoreDelegates::OnEnginePreExit.RemoveAll(this);
+#if WITH_EDITOR
+	FEditorDelegates::PrePIEEnded.RemoveAll(this);
+#endif
 
 	if (IMediaIOCoreModule::IsAvailable())
 	{
@@ -388,12 +394,22 @@ const FAvaInstanceSettings& FAvaMediaModule::GetAvaInstanceSettings() const
 	return UAvaMediaSettings::Get().AvaInstanceSettings;
 }
 
+bool FAvaMediaModule::IsLocalPlaybackManagerAvailable() const
+{
+	return LocalPlaybackManager.IsValid(); 
+}
+
 FAvaPlaybackManager& FAvaMediaModule::GetLocalPlaybackManager() const
 {
 	check(LocalPlaybackManager.IsValid());
 	return *LocalPlaybackManager;
 }
 
+bool FAvaMediaModule::IsManagedInstanceCacheAvailable() const
+{
+	return ManagedInstanceCache.IsValid();
+}
+
 FAvaRundownManagedInstanceCache& FAvaMediaModule::GetManagedInstanceCache() const
 {
 	check(ManagedInstanceCache.IsValid());
@@ -420,11 +436,12 @@ void FAvaMediaModule::PostEngineInit()
 {
 	using namespace UE::AvaMediaModule::Private;
 
+	ConditionalCreateLocalPlaybackManager();
+	ConditionalCreateManagedInstanceCache();
+
 	// This needs to happen late in the loading process, otherwise it fails.
 	const UAvaMediaSettings& Settings = UAvaMediaSettings::Get();
 
-	ManagedInstanceCache = MakeShared<FAvaRundownManagedInstanceCache>();
-
 #if WITH_EDITOR
 	// Initialise the device provider registry.
 	FAvaBroadcastDeviceProviderRegistry::Get();
@@ -487,18 +504,6 @@ void FAvaMediaModule::PostEngineInit()
 		StartPlaybackServerCommand({PlaybackServerName});
 	}
 	
-	LocalPlaybackManager = MakeShared<FAvaPlaybackManager>();
-
-#if WITH_EDITOR
-	// Capture Raw Ptr to avoid keeping ref count on capture
-	// Only Local Playback Manager should handle tear down for PIE End
-	FAvaPlaybackManager* LocalPlaybackManagerRaw = LocalPlaybackManager.Get();
-	FEditorDelegates::PrePIEEnded.AddSPLambda(LocalPlaybackManagerRaw, [LocalPlaybackManagerRaw](const bool)
-	{
-		LocalPlaybackManagerRaw->OnParentWorldBeginTearDown();
-	});
-#endif
-
 	// Playback server required by Http server
 	if (AvaPlaybackServer.IsValid() && Settings.bAutoStartWebServer)
 	{
@@ -519,6 +524,15 @@ void FAvaMediaModule::EnginePreExit()
 	StopAllServices();
 }
 
+void FAvaMediaModule::PrePIEEnded(const bool)
+{
+	// Only Local Playback Manager should handle tear down for PIE End.
+	if (LocalPlaybackManager)
+	{
+		LocalPlaybackManager->OnParentWorldBeginTearDown();
+	}
+}
+
 void FAvaMediaModule::StopAllServices()
 {
 	StopRundownServerCommand({});
@@ -534,6 +548,22 @@ void FAvaMediaModule::StopAllServices()
 	ManagedInstanceCache.Reset();
 }
 
+void FAvaMediaModule::ConditionalCreateLocalPlaybackManager()
+{
+	if (!LocalPlaybackManager)
+	{
+		LocalPlaybackManager = MakeShared<FAvaPlaybackManager>();
+	}
+}
+
+void FAvaMediaModule::ConditionalCreateManagedInstanceCache()
+{
+	if (!ManagedInstanceCache)
+	{
+		ManagedInstanceCache = MakeShared<FAvaRundownManagedInstanceCache>();
+	}
+}
+
 void FAvaMediaModule::StartRundownServerCommand(const TArray<FString>& Args)
 {
 	if (RundownServer)
@@ -541,6 +571,10 @@ void FAvaMediaModule::StartRundownServerCommand(const TArray<FString>& Args)
 		UE_LOG(LogAvaMedia, Log, TEXT("Rundown Server is already started."));
 		return;
 	}
+
+	// Make sure all the components used by the rundown server are created.
+	ConditionalCreateLocalPlaybackManager();
+	ConditionalCreateManagedInstanceCache();
 	
 	RundownServer = MakeShared<FAvaRundownServer>();
 	
diff --git a/Engine/Plugins/Experimental/Avalanche/Source/AvalancheMedia/Private/AvaMediaModule.h b/Engine/Plugins/Experimental/Avalanche/Source/AvalancheMedia/Private/AvaMediaModule.h
--- a/Engine/Plugins/Experimental/Avalanche/Source/AvalancheMedia/Private/AvaMediaModule.h
+++ b/Engine/Plugins/Experimental/Avalanche/Source/AvalancheMedia/Private/AvaMediaModule.h
@@ -42,7 +42,9 @@ public:
 	virtual bool IsGameModeLocalPlaybackServerLaunched() const override;
 	virtual const IAvaBroadcastSettings& GetBroadcastSettings() const override;
 	virtual const FAvaInstanceSettings& GetAvaInstanceSettings() const override;
+	virtual bool IsLocalPlaybackManagerAvailable() const override;
 	virtual FAvaPlaybackManager& GetLocalPlaybackManager() const override;
+	virtual bool IsManagedInstanceCacheAvailable() const override;
 	virtual FAvaRundownManagedInstanceCache& GetManagedInstanceCache() const override;
 	virtual bool IsAvaMediaSyncProviderFeatureAvailable() const override;
 	virtual IAvaMediaSyncProvider* GetAvaMediaSyncProvider() const override;
@@ -68,7 +70,10 @@ public:
 private:
 	void PostEngineInit();
 	void EnginePreExit();
+	void PrePIEEnded(const bool);
 	void StopAllServices();
+	void ConditionalCreateLocalPlaybackManager();
+	void ConditionalCreateManagedInstanceCache();
 	
 	// Command handlers
 	void StartRundownServerCommand(const TArray<FString>& Args);
diff --git a/Engine/Plugins/Experimental/Avalanche/Source/AvalancheMedia/Private/Rundown/AvaRundown.cpp b/Engine/Plugins/Experimental/Avalanche/Source/AvalancheMedia/Private/Rundown/AvaRundown.cpp
--- a/Engine/Plugins/Experimental/Avalanche/Source/AvalancheMedia/Private/Rundown/AvaRundown.cpp
+++ b/Engine/Plugins/Experimental/Avalanche/Source/AvalancheMedia/Private/Rundown/AvaRundown.cpp
@@ -876,6 +876,13 @@ void UAvaRundown::InitializePlaybackContext()
 	}
 }
 
+bool UAvaRundown::CanClosePlaybackContext() const
+{
+	bool bResult = true;
+	OnCanClosePlaybackContext.Broadcast(this, bResult);
+	return bResult;
+}
+
 void UAvaRundown::ClosePlaybackContext(bool bInStopAllPages)
 {
 	if (bInStopAllPages)
diff --git a/Engine/Plugins/Experimental/Avalanche/Source/AvalancheMedia/Private/Rundown/AvaRundownServer.cpp b/Engine/Plugins/Experimental/Avalanche/Source/AvalancheMedia/Private/Rundown/AvaRundownServer.cpp
--- a/Engine/Plugins/Experimental/Avalanche/Source/AvalancheMedia/Private/Rundown/AvaRundownServer.cpp
+++ b/Engine/Plugins/Experimental/Avalanche/Source/AvalancheMedia/Private/Rundown/AvaRundownServer.cpp
@@ -319,6 +319,7 @@ FAvaRundownServer::FAvaRundownServer()
 
 FAvaRundownServer::~FAvaRundownServer()
 {
+	RemovePlaybackDelegates();
 	RemoveBroadcastDelegates(&UAvaBroadcast::Get());
 	RemoveEditorDelegates();
 	
@@ -343,13 +344,12 @@ const FMessageAddress& FAvaRundownServer::GetMessageAddress() const
 
 void FAvaRundownServer::AddReferencedObjects(FReferenceCollector& InCollector)
 {
-	if (RundownEditCommandData.CurrentRundown)
+	for (const TPair<FSoftObjectPath, TWeakPtr<FRundownEntry>>& RundownEntryWeak : LoadedRundownCache)
 	{
-		InCollector.AddReferencedObject(RundownEditCommandData.CurrentRundown);
-	}
-	if (RundownPlaybackCommandData.CurrentRundown)
-	{
-		InCollector.AddReferencedObject(RundownPlaybackCommandData.CurrentRundown);
+		if (const TSharedPtr<FRundownEntry> RundownEntry = RundownEntryWeak.Value.Pin())
+		{
+			RundownEntry->AddReferencedObjects(InCollector);
+		}
 	}
 
 	for (TPair<FSoftObjectPath, TObjectPtr<UAvaRundown>>& ManagedRundown : ManagedRundowns)
@@ -419,6 +419,7 @@ void FAvaRundownServer::Init(const FString& InAssignedHostName)
 		// Subscribe to the server listing requests
 		MessageEndpoint->Subscribe<FAvaRundownPing>();
 
+		SetupPlaybackDelegates();
 		SetupBroadcastDelegates(&UAvaBroadcast::Get());
 		SetupEditorDelegates();
 
@@ -426,6 +427,12 @@ void FAvaRundownServer::Init(const FString& InAssignedHostName)
 	}
 }
 
+void FAvaRundownServer::SetupPlaybackDelegates()
+{
+	FAvaPlaybackManager& Manager = IAvaMediaModule::Get().GetLocalPlaybackManager();
+	Manager.OnPlaybackInstanceStatusChanged.AddSP(this, &FAvaRundownServer::OnPlaybackInstanceStatusChanged);
+}
+
 void FAvaRundownServer::SetupBroadcastDelegates(UAvaBroadcast* InBroadcast)
 {
 	RemoveBroadcastDelegates(InBroadcast);
@@ -437,8 +444,21 @@ void FAvaRundownServer::SetupEditorDelegates()
 {
 	RemoveEditorDelegates();
 	const FAssetRegistryModule& AssetRegistryModule = FModuleManager::LoadModuleChecked<FAssetRegistryModule>("AssetRegistry");
-	AssetRegistryModule.Get().OnAssetAdded().AddSP(this, &FAvaRundownServer::OnAssetAddedOrRemoved);
-	AssetRegistryModule.Get().OnAssetRemoved().AddSP(this, &FAvaRundownServer::OnAssetAddedOrRemoved);
+	AssetRegistryModule.Get().OnAssetAdded().AddSP(this, &FAvaRundownServer::OnAssetAdded);
+	AssetRegistryModule.Get().OnAssetRemoved().AddSP(this, &FAvaRundownServer::OnAssetRemoved);
+#if WITH_EDITOR
+	FEditorDelegates::OnAssetsPreDelete.AddSP(this, &FAvaRundownServer::OnAssetsPreDelete);
+#endif
+}
+
+void FAvaRundownServer::RemovePlaybackDelegates() const
+{
+	const IAvaMediaModule& AvaMediaModule = IAvaMediaModule::Get();
+	if (AvaMediaModule.IsLocalPlaybackManagerAvailable())
+	{
+		FAvaPlaybackManager& Manager = AvaMediaModule.GetLocalPlaybackManager();
+		Manager.OnPlaybackInstanceStatusChanged.RemoveAll(this);
+	}
 }
 
 void FAvaRundownServer::RemoveBroadcastDelegates(UAvaBroadcast* InBroadcast) const
@@ -452,6 +472,9 @@ void FAvaRundownServer::RemoveEditorDelegates() const
 	const FAssetRegistryModule& AssetRegistryModule = FModuleManager::LoadModuleChecked<FAssetRegistryModule>("AssetRegistry");
 	AssetRegistryModule.Get().OnAssetAdded().RemoveAll(this);
 	AssetRegistryModule.Get().OnAssetRemoved().RemoveAll(this);
+#if WITH_EDITOR
+	FEditorDelegates::OnAssetsPreDelete.RemoveAll(this);
+#endif
 }
 
 void FAvaRundownServer::OnPageListChanged(const FAvaRundownPageListChangeParams& InParams) const
@@ -575,19 +598,31 @@ void FAvaRundownServer::OnBroadcastChannelChanged(const FAvaBroadcastOutputChann
 	SendResponse(ReplyMessage, ClientAddresses);
 }
 
-void FAvaRundownServer::OnAssetAddedOrRemoved(const FAssetData& InAssetData) const
+void FAvaRundownServer::OnAssetAdded(const FAssetData& InAssetData) const
 {
-	if (ClientAddresses.IsEmpty())
-	{
-		return;
-	}
-	
-	using namespace UE::AvaRundownServer::Private;
-	if (InAssetData.GetClass() == UAvaRundown::StaticClass() || FAvaPlaybackUtils::IsPlayableAsset(InAssetData))
+	NotifyAssetEvent(InAssetData, EAvaRundownAssetEvent::Added);
+}
+
+void FAvaRundownServer::OnAssetRemoved(const FAssetData& InAssetData) const
+{
+	NotifyAssetEvent(InAssetData, EAvaRundownAssetEvent::Removed);
+}
+
+void FAvaRundownServer::OnAssetsPreDelete(const TArray<UObject*>& InObjects)
+{
+	for (UObject* Object : InObjects)
 	{
-		FAvaRundownAssetsChanged* ReplyMessage = FMessageEndpoint::MakeMessage<FAvaRundownAssetsChanged>();
-		ReplyMessage->AssetName = InAssetData.AssetName.ToString();
-		SendResponse(ReplyMessage, ClientAddresses);
+		if (const UAvaRundown* Rundown = Cast<UAvaRundown>(Object))
+		{
+			// Allow the edited rundown to be deleted.
+			EditCommandContext.ConditionalFlush(SharedThis(this), Rundown);
+
+			// Allow the playback rundown to be deleted, unless it is playing.
+			if (!Rundown->IsPlaying())
+			{
+				PlaybackCommandContext.ConditionalFlush(SharedThis(this), Rundown);
+			}
+		}
 	}
 }
 
@@ -689,20 +724,7 @@ void FAvaRundownServer::HandleLoadRundown(const FAvaRundownLoadRundown& InMessag
 	if (!InMessage.Rundown.IsEmpty())
 	{
 		const FSoftObjectPath NewRundownPath(InMessage.Rundown);
-		const UAvaRundown* Rundown = RundownPlaybackCommandData.GetOrLoadRundown(NewRundownPath,
-			[this](UAvaRundown* InRundown)
-			{
-				RundownPlaybackCommandData.ClosePlaybackContext();
-				RundownPlaybackCommandData.RemoveRundownDelegates(this, InRundown);
-			},
-			[this](UAvaRundown* InRundown)
-			{
-				RundownPlaybackCommandData.SetupRundownDelegates(this, InRundown);
-				if (InRundown)
-				{
-					InRundown->InitializePlaybackContext();
-				}
-			});
+		const UAvaRundown* Rundown = GetOrLoadRundownForContext(NewRundownPath, PlaybackCommandContext);
 		
 		if (!Rundown)
 		{
@@ -713,7 +735,7 @@ void FAvaRundownServer::HandleLoadRundown(const FAvaRundownLoadRundown& InMessag
 	}
 
 	SendMessage(InContext->GetSender(), InMessage.RequestId, ELogVerbosity::Log,
-		TEXT("Rundown \"%s\" loaded."), *RundownPlaybackCommandData.CurrentRundownPath.ToString());
+		TEXT("Rundown \"%s\" loaded."), *PlaybackCommandContext.GetCurrentRundownPath().ToString());
 }
 
 void FAvaRundownServer::HandleCreateRundown(const FAvaRundownCreateRundown& InMessage, const TSharedRef<IMessageContext, ESPMode::ThreadSafe>& InContext)
@@ -818,6 +840,10 @@ void FAvaRundownServer::HandleDeleteRundown(const FAvaRundownDeleteRundown& InMe
 
 	if (ManagedRundowns.Remove(RundownPath) > 0)
 	{
+		// Also, flush command contexts if associated to this rundown.
+		EditCommandContext.ConditionalFlush(SharedThis(this), RundownPath);
+		PlaybackCommandContext.ConditionalFlush(SharedThis(this), RundownPath);
+		
 		SendMessage(InContext->GetSender(), InMessage.RequestId, ELogVerbosity::Log,
 			TEXT("\"DeleteRundown\": Rundown \"%s\" removed."), *InMessage.Rundown);
 		return;
@@ -829,10 +855,18 @@ void FAvaRundownServer::HandleDeleteRundown(const FAvaRundownDeleteRundown& InMe
 	{
 		TArray<FAssetData> AssetData;
 		AssetData.Add(RundownAsset);
-		ObjectTools::DeleteAssets(AssetData, /*bShowConfirmation*/ false);
-		
-		SendMessage(InContext->GetSender(), InMessage.RequestId, ELogVerbosity::Log,
-			TEXT("\"DeleteRundown\": Rundown \"%s\" deleted."), *InMessage.Rundown);
+		int32 NumDeleted = ObjectTools::DeleteAssets(AssetData, /*bShowConfirmation*/ false);
+
+		if (NumDeleted == AssetData.Num())
+		{
+			SendMessage(InContext->GetSender(), InMessage.RequestId, ELogVerbosity::Log,
+				TEXT("\"DeleteRundown\": Rundown \"%s\" deleted."), *InMessage.Rundown);
+		}
+		else
+		{
+			SendMessage(InContext->GetSender(), InMessage.RequestId, ELogVerbosity::Error,
+				TEXT("\"DeleteRundown\": Rundown \"%s\" could not be deleted."), *InMessage.Rundown);
+		}
 		return;
 	}
 #endif
@@ -1164,7 +1198,7 @@ void FAvaRundownServer::HandleGetPageDetails(const FAvaRundownGetPageDetails& In
 		
 		if (ManagedInstance.IsValid())
 		{
-			RundownEditCommandData.SaveCurrentRemoteControlPresetToPage(true);
+			EditCommandContext.SaveCurrentRemoteControlPresetToPage(true);
 
 			// Applying the controller values can break the WYSIWYG of the editor,
 			// in case multiple controllers set the same exposed entity with different values.
@@ -1179,8 +1213,8 @@ void FAvaRundownServer::HandleGetPageDetails(const FAvaRundownGetPageDetails& In
 			FAvaRemoteControlUtils::RegisterRemoteControlPreset(ManagedInstance->GetRemoteControlPreset(), /*bInEnsureUniqueId*/ true);
 
 			// Keep track of what is currently registered.
-			RundownEditCommandData.ManagedPageId = InMessage.PageId;
-			RundownEditCommandData.ManagedInstance = ManagedInstance;
+			EditCommandContext.ManagedPageId = InMessage.PageId;
+			EditCommandContext.ManagedInstance = ManagedInstance;
 		}
 	}
 	
@@ -1189,10 +1223,10 @@ void FAvaRundownServer::HandleGetPageDetails(const FAvaRundownGetPageDetails& In
 	ReplyMessage->Rundown = InMessage.Rundown;
 	ReplyMessage->PageInfo = UE::AvaRundownServer::Private::GetPageInfo(Rundown, Page);
 	ReplyMessage->RemoteControlValues = Page.GetRemoteControlValues();
-	if (InMessage.bLoadRemoteControlPreset && RundownEditCommandData.ManagedInstance.IsValid() && RundownEditCommandData.ManagedInstance->GetRemoteControlPreset())
+	if (InMessage.bLoadRemoteControlPreset && EditCommandContext.ManagedInstance.IsValid() && EditCommandContext.ManagedInstance->GetRemoteControlPreset())
 	{
-		ReplyMessage->RemoteControlPresetName = RundownEditCommandData.ManagedInstance->GetRemoteControlPreset()->GetPresetName().ToString();
-		ReplyMessage->RemoteControlPresetId = RundownEditCommandData.ManagedInstance->GetRemoteControlPreset()->GetPresetId().ToString();
+		ReplyMessage->RemoteControlPresetName = EditCommandContext.ManagedInstance->GetRemoteControlPreset()->GetPresetName().ToString();
+		ReplyMessage->RemoteControlPresetId = EditCommandContext.ManagedInstance->GetRemoteControlPreset()->GetPresetId().ToString();
 	}
 	SendResponse(ReplyMessage, InContext->GetSender());
 }
@@ -1242,11 +1276,11 @@ void FAvaRundownServer::HandleChangePageChannel(const FAvaRundownPageChangeChann
 void FAvaRundownServer::HandleUpdatePageFromRCP(const FAvaRundownUpdatePageFromRCP& InMessage, const TSharedRef<IMessageContext, ESPMode::ThreadSafe>& InContext)
 {
 	// Note that this doesn't save the rundown.
-	RundownEditCommandData.SaveCurrentRemoteControlPresetToPage(InMessage.bUnregister);
+	EditCommandContext.SaveCurrentRemoteControlPresetToPage(InMessage.bUnregister);
 	if (InMessage.bUnregister)
 	{
-		RundownEditCommandData.ManagedPageId = FAvaRundownPage::InvalidPageId;
-		RundownEditCommandData.ManagedInstance.Reset();
+		EditCommandContext.ManagedPageId = FAvaRundownPage::InvalidPageId;
+		EditCommandContext.ManagedInstance.Reset();
 	}
 	SendMessage(InContext->GetSender(), InMessage.RequestId, ELogVerbosity::Log, TEXT("\"UpdatePageFromRCP\" Ok."));
 }
@@ -2022,18 +2056,18 @@ void FAvaRundownServer::ShowStatusCommand(const TArray<FString>& InArgs)
 	}
 	
 	UE_LOG(LogAvaRundownServer, Display, TEXT("Rundown Caches:"));
-	UE_LOG(LogAvaRundownServer, Display, TEXT("- Editing Rundown: \"%s\""), *RundownEditCommandData.CurrentRundownPath.ToString());
-	UE_LOG(LogAvaRundownServer, Display, TEXT("- Editing PageId: \"%d\""), RundownEditCommandData.ManagedPageId);
-	UE_LOG(LogAvaRundownServer, Display, TEXT("- Playing Rundown: \"%s\""), *RundownPlaybackCommandData.CurrentRundownPath.ToString());
+	UE_LOG(LogAvaRundownServer, Display, TEXT("- Editing Rundown: \"%s\""), *EditCommandContext.GetCurrentRundownPath().ToString());
+	UE_LOG(LogAvaRundownServer, Display, TEXT("- Editing PageId: \"%d\""), EditCommandContext.ManagedPageId);
+	UE_LOG(LogAvaRundownServer, Display, TEXT("- Playing Rundown: \"%s\""), *PlaybackCommandContext.GetCurrentRundownPath().ToString());
 
-	if (RundownPlaybackCommandData.CurrentRundown)
+	if (const UAvaRundown* CurrentPlaybackRundown = PlaybackCommandContext.GetCurrentRundown())
 	{
-		TArray<int32> PlayingPages = RundownPlaybackCommandData.CurrentRundown->GetPlayingPageIds();
+		TArray<int32> PlayingPages = CurrentPlaybackRundown->GetPlayingPageIds();
 		for (const int32 PlayingPageId : PlayingPages)
 		{
 			UE_LOG(LogAvaRundownServer, Display, TEXT("- Playing PageId: \"%d\""), PlayingPageId);
 		}
-		TArray<int32> PreviewingPages = RundownPlaybackCommandData.CurrentRundown->GetPreviewingPageIds();
+		TArray<int32> PreviewingPages = CurrentPlaybackRundown->GetPreviewingPageIds();
 		for (const int32 PreviewingPageId : PreviewingPages)
 		{
 			UE_LOG(LogAvaRundownServer, Display, TEXT("- Previewing PageId: \"%d\""), PreviewingPageId);
@@ -2041,6 +2075,37 @@ void FAvaRundownServer::ShowStatusCommand(const TArray<FString>& InArgs)
 	}
 }
 
+void FAvaRundownServer::NotifyPlaybackContextSwitch(const FSoftObjectPath& InPreviousRundownPath, const FSoftObjectPath& InNewRundownPath) const
+{
+	if (ClientAddresses.IsEmpty())
+	{
+		return;
+	}
+
+	FAvaRundownPlaybackContextChanged* Notification = FMessageEndpoint::MakeMessage<FAvaRundownPlaybackContextChanged>();
+	Notification->PreviousRundown = InPreviousRundownPath.ToString();
+	Notification->NewRundown = InNewRundownPath.ToString();
+	SendResponse(Notification, ClientAddresses);
+}
+
+void FAvaRundownServer::NotifyAssetEvent(const FAssetData& InAssetData, const EAvaRundownAssetEvent InEventType) const
+{
+	if (ClientAddresses.IsEmpty())
+	{
+		return;
+	}
+
+	// todo: probably need some event filtering (playable, by class, etc).
+	FAvaRundownAssetsChanged* Message = FMessageEndpoint::MakeMessage<FAvaRundownAssetsChanged>();
+	Message->AssetName = InAssetData.AssetName.ToString();
+	Message->AssetPath = InAssetData.GetSoftObjectPath().ToString();
+	Message->AssetClass = InAssetData.AssetClassPath.ToString();
+	Message->bIsPlayable = FAvaPlaybackUtils::IsPlayableAsset(InAssetData);
+	Message->EventType = InEventType;
+	
+	SendResponse(Message, ClientAddresses);
+}
+
 void FAvaRundownServer::FinishGetChannelImage(const FRequestInfo& InRequestInfo, const TSharedPtr<FChannelImage>& InChannelImage)
 {
 	FAvaRundownChannelImage* ReplyMessage = FMessageEndpoint::MakeMessage<FAvaRundownChannelImage>();
@@ -2101,16 +2166,16 @@ void FAvaRundownServer::HandlePageActions(const FRequestInfo& InRequestInfo, con
 	bool bInIsPreview, FName InPreviewChannelName, EAvaRundownPageActions InAction) const
 {
 	using namespace UE::AvaRundownServer::Private;
+
+	UAvaRundown* Rundown = PlaybackCommandContext.GetCurrentRundown();
 	
-	if (!RundownPlaybackCommandData.CurrentRundown)
+	if (!Rundown)
 	{
 		LogAndSendMessage(InRequestInfo.Sender, InRequestInfo.RequestId, ELogVerbosity::Error,
-			TEXT("\"PageAction\" Failed. Reason: no play list currently loaded."));
+			TEXT("\"PageAction\" Failed. Reason: no rundown currently loaded for playback."));
 		return;
 	}
 
-	UAvaRundown* Rundown = RundownPlaybackCommandData.CurrentRundown.Get();
-
 	{
 		// Validate the pages - the command will be considered a failure (as a whole) if it contains invalid pages.
 		FString InvalidPages;
@@ -2264,16 +2329,7 @@ UAvaRundown* FAvaRundownServer::GetOrLoadRundownForEdit(const FMessageAddress& I
 		// This will not affect the currently loaded rundown for playback.
 		const FSoftObjectPath NewRundownPath(InRundownPath);
 		
-		Rundown = RundownEditCommandData.GetOrLoadRundown(NewRundownPath,
-			[this](UAvaRundown* InRundown)
-			{
-				RundownEditCommandData.SaveCurrentRemoteControlPresetToPage(true);
-				RundownEditCommandData.RemoveRundownDelegates(this, InRundown);
-			},
-			[this](UAvaRundown* InRundown)
-			{
-				RundownEditCommandData.SetupRundownDelegates(this, InRundown);
-			});
+		Rundown = GetOrLoadRundownForContext(NewRundownPath, EditCommandContext);
 
 		if (!Rundown)
 		{
@@ -2283,16 +2339,12 @@ UAvaRundown* FAvaRundownServer::GetOrLoadRundownForEdit(const FMessageAddress& I
 	else
 	{
 		// If the path is not specified, we assume it is using the previously loaded rundown.
-		Rundown = RundownEditCommandData.CurrentRundown.Get();
+		Rundown = EditCommandContext.GetCurrentRundown();
 
 		// Note: for backward compatibility with QA python script, we allow this command to use the current "playback" rundown as fallback.
 		if (!Rundown)
 		{
-			Rundown = RundownPlaybackCommandData.CurrentRundown.Get();
-
-			// Update the edit data accordingly.
-			RundownEditCommandData.CurrentRundown = Rundown;
-			RundownEditCommandData.CurrentRundownPath = RundownPlaybackCommandData.CurrentRundownPath;
+			Rundown = GetOrLoadRundownForContext(PlaybackCommandContext.GetCurrentRundownPath(), EditCommandContext);
 		}
 
 		if (!Rundown)
@@ -2342,83 +2394,171 @@ void FAvaRundownServer::RefreshClientAddresses()
 	}
 }
 
-UAvaRundown* FAvaRundownServer::FRundownCache::GetOrLoadRundown(const FSoftObjectPath& InRundownPath,
-	const FRundownEventFunction InUnloadCurrentRundownFunction,
-	const FRundownEventFunction InNewRundownLoadedFunction)
+FAvaRundownServer::FRundownEntry::FRundownEntry(const TSharedPtr<FAvaRundownServer>& InRundownServer, const FSoftObjectPath& InRundownPath)
+	: RundownServerRaw(InRundownServer.Get())
+{
+	const TStrongObjectPtr<UAvaRundown> LoadedRundown = UE::AvaRundownServer::Private::LoadRundown(InRundownPath);
+	Rundown = LoadedRundown.IsValid() ? LoadedRundown.Get() : nullptr;
+	
+	if (Rundown && InRundownServer)
+	{
+		const TSharedRef<FAvaRundownServer> RundownServerRef = InRundownServer.ToSharedRef();
+		Rundown->GetOnPagesChanged().AddSP(RundownServerRef, &FAvaRundownServer::OnPagesChanged);
+		Rundown->GetOnInstancedPageListChanged().AddSP(RundownServerRef, &FAvaRundownServer::OnPageListChanged);
+		Rundown->GetOnTemplatePageListChanged().AddSP(RundownServerRef, &FAvaRundownServer::OnPageListChanged);
+		Rundown->GetOnCanClosePlaybackContext().AddSP(RundownServerRef, &FAvaRundownServer::OnCanClosePlaybackContext);
+	}
+}
+
+FAvaRundownServer::FRundownEntry::~FRundownEntry()
+{
+	if (Rundown)
+	{
+		Rundown->GetOnPagesChanged().RemoveAll(RundownServerRaw);
+		Rundown->GetOnInstancedPageListChanged().RemoveAll(RundownServerRaw);
+		Rundown->GetOnTemplatePageListChanged().RemoveAll(RundownServerRaw);
+		Rundown->GetOnCanClosePlaybackContext().RemoveAll(RundownServerRaw);
+	}
+}
+
+void FAvaRundownServer::CompactLoadedRundownCache()
 {
-	if (CurrentRundownPath != InRundownPath)
+	for (TMap<FSoftObjectPath, TWeakPtr<FRundownEntry>>::TIterator ContextIterator = LoadedRundownCache.CreateIterator();
+		ContextIterator; ++ContextIterator  )
 	{
-		const TStrongObjectPtr<UAvaRundown> NewRundown = UE::AvaRundownServer::Private::LoadRundown(InRundownPath);
-		if (NewRundown.IsValid())
+		if (!ContextIterator->Value.IsValid())
 		{
-			if (CurrentRundown.Get())
-			{
-				InUnloadCurrentRundownFunction(CurrentRundown.Get());
-			}
-			InNewRundownLoadedFunction(NewRundown.Get());
-			CurrentRundown = NewRundown.Get();
-			CurrentRundownPath = InRundownPath;
+			ContextIterator.RemoveCurrent();
+		}
+	}
+}
+
+TSharedPtr<FAvaRundownServer::FRundownEntry> FAvaRundownServer::GetOrLoadRundown(const FSoftObjectPath& InRundownPath)
+{
+	if (const TWeakPtr<FRundownEntry>* ExistingEntryWeak = LoadedRundownCache.Find(InRundownPath))
+	{
+		if (TSharedPtr<FRundownEntry> ExistingEntry = (*ExistingEntryWeak).Pin())
+		{
+			return ExistingEntry;
+		}
+	}
+
+	TSharedPtr<FRundownEntry> NewEntry = MakeShared<FRundownEntry>(SharedThis(this), InRundownPath);
+	if (NewEntry->IsValid())
+	{
+		CompactLoadedRundownCache();
+		LoadedRundownCache.Add(InRundownPath, NewEntry);
+		return NewEntry;
+	}
+
+	// Failed to load asset.
+	return nullptr;
+}
+
+UAvaRundown* FAvaRundownServer::GetOrLoadRundownForContext(const FSoftObjectPath& InRundownPath, FCommandContext& InContext)
+{
+	if (InRundownPath != InContext.GetCurrentRundownPath())
+	{
+		if (const TSharedPtr<FRundownEntry> NewRundownEntry = GetOrLoadRundown(InRundownPath))
+		{
+			InContext.SetCurrentRundown(SharedThis(this), InRundownPath, NewRundownEntry);
 		}
 		else
 		{
+			// Indicates failure of loading new rundown asset.
+			// Context is not modified.
 			return nullptr;
 		}
 	}
-	return CurrentRundown.Get();
+	return InContext.GetCurrentRundown();
 }
 
-void FAvaRundownServer::FRundownCache::SetupRundownDelegates(FAvaRundownServer* InRundownServer, UAvaRundown* InRundown)
+UAvaRundownPagePlayer* FAvaRundownServer::FindPagePlayerForInstance(const FAvaPlaybackInstance& InPlaybackInstance) const
 {
-	RemoveRundownDelegates(InRundownServer, InRundown);
+	// We don't know the channel, nor the rundown, but can at least get the page Id from the instance user data.
+	const int32 PageId = UAvaRundownPagePlayer::GetPageIdFromInstanceUserData(InPlaybackInstance.GetInstanceUserData());
 
-	if (!InRundown)
+	// Search in any of the currently loaded rundowns in the server. 
+	for (const TPair<FSoftObjectPath, TWeakPtr<FRundownEntry>>& RundownEntryWeak : LoadedRundownCache)
 	{
-		return;
-	}
-
-	FAvaPlaybackManager& Manager = IAvaMediaModule::Get().GetLocalPlaybackManager();
-	TWeakObjectPtr<UAvaRundown> RundownWeak(InRundown);
-	TWeakPtr<FAvaRundownServer> RundownServerWeak = InRundownServer->AsShared();
-	OnPlaybackInstanceStatusChangedDelegateHandle = 
-		Manager.OnPlaybackInstanceStatusChanged.AddLambda([RundownWeak, RundownServerWeak](const FAvaPlaybackInstance& InPlaybackInstance)
-		{
-			UAvaRundown* Rundown = RundownWeak.Get();
-			const TSharedPtr<FAvaRundownServer> RundownServer = RundownServerWeak.Pin();
-			if (IsValid(Rundown) && RundownServer.IsValid())
+		const TSharedPtr<FRundownEntry> RundownEntry = RundownEntryWeak.Value.Pin();
+		if (!RundownEntry)
+		{
+			continue;
+		}
+		
+		const UAvaRundown* Rundown = RundownEntry->GetRundown();
+		if (!Rundown)
+		{
+			continue;
+		}
+		
+		// if we have a pageId, we can skip any rundown that doesn't have that page.
+		if (PageId != FAvaRundownPage::InvalidPageId && !Rundown->GetPage(PageId).IsValidPage())
+		{
+			continue;
+		}
+		
+		// We don't know the channel (could be preview), so we need to check all page players.
+		for (const TObjectPtr<UAvaRundownPagePlayer>& PagePlayer : Rundown->GetPagePlayers())
+		{
+			// If we have a pageId, we can skip any players for other pages.
+			if (PageId != FAvaRundownPage::InvalidPageId && PagePlayer->PageId != PageId)
 			{
-				const int32 PageId = UAvaRundownPagePlayer::GetPageIdFromInstanceUserData(InPlaybackInstance.GetInstanceUserData());
-				const FAvaRundownPage Page = Rundown->GetPage(PageId);
-				if (Page.IsValidPage())
-				{
-					RundownServer->PageStatusChanged(Rundown, Page);
-				}
+				continue;
 			}
-		});
 
-	InRundown->GetOnPagesChanged().AddRaw(InRundownServer, &FAvaRundownServer::OnPagesChanged);
-	InRundown->GetOnInstancedPageListChanged().AddRaw(InRundownServer, &FAvaRundownServer::OnPageListChanged);
-	InRundown->GetOnTemplatePageListChanged().AddRaw(InRundownServer, &FAvaRundownServer::OnPageListChanged);
+			// Using the instanceId to identify the correct instance.
+			if (PagePlayer->FindInstancePlayerByInstanceId(InPlaybackInstance.GetInstanceId()))
+			{
+				return PagePlayer.Get();
+			}
+		}
+	}
+	return nullptr;
 }
 
-void FAvaRundownServer::FRundownCache::RemoveRundownDelegates(const FAvaRundownServer* InRundownServer, UAvaRundown* InRundown) const
+void FAvaRundownServer::OnPlaybackInstanceStatusChanged(const FAvaPlaybackInstance& InPlaybackInstance)
 {
-	FAvaPlaybackManager& Manager = IAvaMediaModule::Get().GetLocalPlaybackManager();
-	Manager.OnPlaybackInstanceStatusChanged.Remove(OnPlaybackInstanceStatusChangedDelegateHandle);
+	const UAvaRundownPagePlayer* PagePlayer = FindPagePlayerForInstance(InPlaybackInstance);
+	if (!PagePlayer)
+	{
+		return;
+	}
+	UAvaRundown* Rundown = PagePlayer->GetRundown();
+	if (!Rundown)
+	{
+		return;
+	}
+	
+	const FAvaRundownPage& Page = Rundown->GetPage(PagePlayer->PageId);
+	if (Page.IsValidPage())
+	{
+		PageStatusChanged(Rundown, Page);
+	}
+}
 
-	if (InRundown)
+void FAvaRundownServer::OnCanClosePlaybackContext(const UAvaRundown* InRundown, bool& bOutResult) const
+{
+	const UAvaRundown* CurrentPlaybackRundown = PlaybackCommandContext.GetCurrentRundown();
+	if (CurrentPlaybackRundown != nullptr && CurrentPlaybackRundown == InRundown)
 	{
-		InRundown->GetOnPagesChanged().RemoveAll(InRundownServer);
-		InRundown->GetOnInstancedPageListChanged().RemoveAll(InRundownServer);
-		InRundown->GetOnTemplatePageListChanged().RemoveAll(InRundownServer);
+		bOutResult = false;
 	}
 }
 
-FAvaRundownServer::FRundownEditCommandData::~FRundownEditCommandData()
+FAvaRundownServer::FEditCommandContext::~FEditCommandContext()
 {
 	SaveCurrentRemoteControlPresetToPage(true);
 }
 
-void FAvaRundownServer::FRundownEditCommandData::SaveCurrentRemoteControlPresetToPage(bool bInUnregister)
+void FAvaRundownServer::FEditCommandContext::SetCurrentRundown(const TSharedPtr<FAvaRundownServer>& InRundownServer, const FSoftObjectPath& InRundownPath, const TSharedPtr<FRundownEntry>& InRundownEntry)
+{
+	SaveCurrentRemoteControlPresetToPage(true);
+	FCommandContext::SetCurrentRundown(InRundownServer, InRundownPath, InRundownEntry);
+}
+
+void FAvaRundownServer::FEditCommandContext::SaveCurrentRemoteControlPresetToPage(bool bInUnregister)
 {
 	if (!ManagedInstance.IsValid() || !ManagedInstance->GetRemoteControlPreset())
 	{
@@ -2440,6 +2580,8 @@ void FAvaRundownServer::FRundownEditCommandData::SaveCurrentRemoteControlPresetT
 		RemoteControlModule.UnregisterEmbeddedPreset(CurrentPresetName);
 	}
 
+	UAvaRundown* CurrentRundown = GetCurrentRundown();
+
 	if (!CurrentRundown)
 	{
 		return;
@@ -2488,19 +2630,49 @@ void FAvaRundownServer::FRundownEditCommandData::SaveCurrentRemoteControlPresetT
 	}
 }
 
-FAvaRundownServer::FRundownPlaybackCommandData::~FRundownPlaybackCommandData()
+FAvaRundownServer::FPlaybackCommandContext::~FPlaybackCommandContext()
+{
+	TSharedPtr<FRundownEntry> PreviousRundownEntry = CurrentRundownEntry; // Prevent GC for current scope.
+	UAvaRundown* PreviousRundown = GetCurrentRundown();
+
+	// Reset current rundown so it doesn't prevent closing playback context.
+	CurrentRundownEntry.Reset();
+	
+	ClosePlaybackContext(PreviousRundown);
+}
+
+void FAvaRundownServer::FPlaybackCommandContext::SetCurrentRundown(const TSharedPtr<FAvaRundownServer>& InRundownServer, const FSoftObjectPath& InRundownPath, const TSharedPtr<FRundownEntry>& InRundownEntry)
+{
+	TSharedPtr<FRundownEntry> PreviousRundownEntry = CurrentRundownEntry;	// Prevent GC for current scope.
+	UAvaRundown* PreviousRundown = GetCurrentRundown();
+
+	// Notify clients that the current playback context is switching.
+	if (InRundownServer)
+	{
+		InRundownServer->NotifyPlaybackContextSwitch(GetCurrentRundownPath(), InRundownPath);
+	}
+
+	FCommandContext::SetCurrentRundown(InRundownServer, InRundownPath, InRundownEntry);
+	
+	ClosePlaybackContext(PreviousRundown);
+	
+	// Initialize new playback context
+	InitializePlaybackContext();
+}
+
+void FAvaRundownServer::FPlaybackCommandContext::InitializePlaybackContext()
 {
-	ClosePlaybackContext();
+	if (UAvaRundown* CurrentRundown = GetCurrentRundown())
+	{
+		CurrentRundown->InitializePlaybackContext();
+	}
 }
 
-void FAvaRundownServer::FRundownPlaybackCommandData::ClosePlaybackContext()
+void FAvaRundownServer::FPlaybackCommandContext::ClosePlaybackContext(UAvaRundown* InRundownToClose)
 {
-	if (CurrentRundown)
+	if (InRundownToClose && InRundownToClose->CanClosePlaybackContext())
 	{
-		// Stop all playing pages.
-		CurrentRundown->ClosePlaybackContext(true);
-		CurrentRundown = nullptr;
-		CurrentRundownPath.Reset();
+		InRundownToClose->ClosePlaybackContext(/*bInStopAllPages*/ true);
 	}
 }
 
diff --git a/Engine/Plugins/Experimental/Avalanche/Source/AvalancheMedia/Private/Rundown/AvaRundownServer.h b/Engine/Plugins/Experimental/Avalanche/Source/AvalancheMedia/Private/Rundown/AvaRundownServer.h
--- a/Engine/Plugins/Experimental/Avalanche/Source/AvalancheMedia/Private/Rundown/AvaRundownServer.h
+++ b/Engine/Plugins/Experimental/Avalanche/Source/AvalancheMedia/Private/Rundown/AvaRundownServer.h
@@ -4,12 +4,14 @@
 
 #include "Broadcast/AvaBroadcastProfile.h"
 #include "MessageEndpoint.h"
+#include "Playback/AvaPlaybackManager.h"
 #include "Rundown/AvaRundown.h"
 #include "Rundown/AvaRundownMessages.h"
 #include "Rundown/IAvaRundownServer.h"
 #include "UObject/SoftObjectPath.h"
 #include "UObject/StrongObjectPtr.h"
 
+class FAvaPlaybackInstance;
 class FAvaRundownManagedInstance;
 class UMediaOutput;
 
@@ -22,7 +24,7 @@ class FAvaRundownServer : public TSharedFromThis<FAvaRundownServer>, public FGCO
 {
 public:
 	FAvaRundownServer();
-	virtual ~FAvaRundownServer();
+	virtual ~FAvaRundownServer() override;
 
 	//~ Begin IAvaRundownServer
 	virtual const FString& GetName() const override { return HostName; }
@@ -37,8 +39,10 @@ public:
 
 	void Init(const FString& InAssignedHostName);
 
+	void SetupPlaybackDelegates();
 	void SetupBroadcastDelegates(UAvaBroadcast* InBroadcast);
 	void SetupEditorDelegates();
+	void RemovePlaybackDelegates() const;
 	void RemoveBroadcastDelegates(UAvaBroadcast* InBroadcast) const;
 	void RemoveEditorDelegates() const;
 	
@@ -50,7 +54,11 @@ public:
 	void PageAnimSettingsChanged(const UAvaRundown* InRundown, const FAvaRundownPage& InPage) const;
 	void OnBroadcastChannelListChanged(const FAvaBroadcastProfile& InProfile) const;
 	void OnBroadcastChannelChanged(const FAvaBroadcastOutputChannel& InChannel, EAvaBroadcastChannelChange InChange) const;
-	void OnAssetAddedOrRemoved(const FAssetData& InAssetData) const;
+	void OnAssetAdded(const FAssetData& InAssetData) const;
+	void OnAssetRemoved(const FAssetData& InAssetData) const;
+	void OnAssetsPreDelete(const TArray<UObject*>& InObjects);
+	void OnPlaybackInstanceStatusChanged(const FAvaPlaybackInstance& InPlaybackInstance);
+	void OnCanClosePlaybackContext(const UAvaRundown* InRundown, bool& bOutResult) const;
 	
 	// Rundown message handlers
 	void HandleRundownPing(const FAvaRundownPing& InMessage, const TSharedRef<IMessageContext, ESPMode::ThreadSafe>& InContext);
@@ -106,6 +114,11 @@ public:
 	
 	void ShowStatusCommand(const TArray<FString>& InArgs);
 
+	/** Broadcast a rundown playback context switch to all connected clients. */
+	void NotifyPlaybackContextSwitch(const FSoftObjectPath& InPreviousRundownPath, const FSoftObjectPath& InNewRundownPath) const;
+	
+	void NotifyAssetEvent(const FAssetData& InAssetData, const EAvaRundownAssetEvent InEventType) const;
+
 protected:
 	struct FRequestInfo
 	{
@@ -119,6 +132,8 @@ protected:
 	void HandlePageActions(const FRequestInfo& InRequestInfo, const TArray<int32>& InPageIds,
 		bool bInIsPreview, FName InPreviewChannelName, EAvaRundownPageActions InAction) const;
 	
+	UAvaRundownPagePlayer* FindPagePlayerForInstance(const FAvaPlaybackInstance& InPlaybackInstance) const;
+	
 	/**
 	 * Helper function to retrieve the appropriate rundown for editing commands.
 	 */
@@ -191,45 +206,125 @@ private:
 	TArray<TSharedPtr<FChannelImage>> AvailableChannelImages;
 	
 	/**
-	 * Keeps a current rundown cached for commands operating on rundown.
-	 * There is only one "current" rundown at a given time.
+	 * Manages rundown delegates binding with the rundown server handlers.
 	 */
-	struct FRundownCache
+	class FRundownEntry
 	{
-		/** Currently loaded/cached rundown's path. */
-		FSoftObjectPath CurrentRundownPath;
+	public:
+		FRundownEntry(const TSharedPtr<FAvaRundownServer>& InRundownServer, const FSoftObjectPath& InRundownPath);
+		~FRundownEntry();
 		
-		/** Currently loaded/cached rundown object. */
-		TObjectPtr<UAvaRundown> CurrentRundown;
+		bool IsValid() const { return Rundown != nullptr; }
+		UAvaRundown* GetRundown() const { return Rundown;}
+		
+		void AddReferencedObjects(FReferenceCollector& InCollector)
+		{
+			InCollector.AddReferencedObject(Rundown);
+		}
+		
+	private:
+		TObjectPtr<UAvaRundown> Rundown;
+		
+		// Keep a raw ptr for unregistering delegates only.
+		FAvaRundownServer* RundownServerRaw = nullptr;
+	};
 
-		FDelegateHandle OnPlaybackInstanceStatusChangedDelegateHandle;
+	/**
+	 * Cache of loaded rundowns currently referenced by the command contexts.
+	 * There is one entry per loaded asset (shared by command contexts).
+	 * Key is the rundown's asset path.
+	 */
+	TMap<FSoftObjectPath, TWeakPtr<FRundownEntry>> LoadedRundownCache;
+
+	/** Remove stale rundown entries. */
+	void CompactLoadedRundownCache();
+	
+	/**
+	 * Returns requested rundown specified by InRundownPath.
+	 * Will load it if necessary or return the cached one.
+	 * If the new rundown fails to load, the returned value is nullptr.
+	 */
+	TSharedPtr<FRundownEntry> GetOrLoadRundown(const FSoftObjectPath& InRundownPath);
+	
+	/**	 
+	 * Associates a rundown entry with it's contextual resources needed to execute commands.
+	 * There is only one "current" rundown per context. Changing the rundown will flush previous
+	 * resources and allocate new ones of the new rundown.
+	 */
+	struct FCommandContext
+	{
+		virtual ~FCommandContext() = default;
+
+		const FSoftObjectPath& GetCurrentRundownPath() const
+		{
+			return CurrentRundownPath;
+		}
 		
-		using FRundownEventFunction = TFunctionRef<void(UAvaRundown*)>;
+		UAvaRundown* GetCurrentRundown() const
+		{
+			return CurrentRundownEntry.IsValid() ? CurrentRundownEntry->GetRundown() : nullptr;
+		}
+
+		void Flush(const TSharedPtr<FAvaRundownServer>& InRundownServer)
+		{
+			SetCurrentRundown(InRundownServer, FSoftObjectPath(), nullptr);
+		}
+ 
+		void ConditionalFlush(const TSharedPtr<FAvaRundownServer>& InRundownServer, const FSoftObjectPath& InRundownPath)
+		{
+			if (GetCurrentRundownPath() == InRundownPath)
+			{
+				Flush(InRundownServer);
+			}
+		}
 		
-		/**
-		 * Returns requested rundown specified by InRundownPath. Will load it if necessary or returned the cached one if it is the same.
-		 * If the new rundown fails to load, the return value is nullptr, and the previous rundown will remain loaded.
-		 * @param InRundownPath	Requested rundown path. 
-		 * @param InUnloadCurrentRundownFunction	Function called, with the previously cached rundown, when a new rundown is loaded
-		 *											and previous rundown needs to be unloaded.
-		 * @param InNewRundownLoadedFunction	Function called, with the new rundown, when a new rundown is loaded.
-		 */
-		UAvaRundown* GetOrLoadRundown(const FSoftObjectPath& InRundownPath, 
-			const FRundownEventFunction InUnloadCurrentRundownFunction,
-			const FRundownEventFunction InNewRundownLoadedFunction);
+		void ConditionalFlush(const TSharedPtr<FAvaRundownServer>& InRundownServer, const UAvaRundown* InRundown)
+		{
+			if (GetCurrentRundown() == InRundown)
+			{
+				Flush(InRundownServer);
+			}
+		}
+		
+		/** Set a new current rundown for the context. Derived classes will handle context switching implementation. */
+		virtual void SetCurrentRundown(const TSharedPtr<FAvaRundownServer>& InRundownServer, const FSoftObjectPath& InRundownPath, const TSharedPtr<FRundownEntry>& InRundownEntry)
+		{
+			CurrentRundownPath = InRundownPath;
+			CurrentRundownEntry = InRundownEntry;
+		}
+		
+	protected:
+		/** Currently loaded/cached rundown's path. */
+		FSoftObjectPath CurrentRundownPath;
 
-		void SetupRundownDelegates(FAvaRundownServer* InRundownServer, UAvaRundown* InRundown);
-		void RemoveRundownDelegates(const FAvaRundownServer* InRundownServer, UAvaRundown* InRundown) const;
+		/** Currently loaded/cached rundown object. */
+		TSharedPtr<FRundownEntry> CurrentRundownEntry;
 	};
+
+	/**
+	 * Returns requested rundown specified by InRundownPath. Will load it if necessary or returned the cached one if it is the same.
+	 * If the new rundown fails to load, the return value is nullptr, and the previous rundown will remain loaded.
+	 * Only one rundown entry can be loaded per context (for now).
+	 * 
+	 * @param InRundownPath	Requested rundown path.
+	 * @param InContext Command context (either edit or playback).
+	 */
+	UAvaRundown* GetOrLoadRundownForContext(const FSoftObjectPath& InRundownPath, FCommandContext& InContext);
 	
-	/** Cached data for page editing commands (GetPages and GetPageDetails). */
-	struct FRundownEditCommandData : public FRundownCache
+	/**
+	 * Context for page editing commands (i.e. GetPages, GetPageDetails, etc). 
+	 */
+	struct FEditCommandContext : public FCommandContext
 	{
 		/** PageId of the current managed ava asset. */
 		int32 ManagedPageId = FAvaRundownPage::InvalidPageId;
 		TSharedPtr<FAvaRundownManagedInstance> ManagedInstance;
 
-		~FRundownEditCommandData();
+		virtual ~FEditCommandContext() override;
+
+		//~ Begin FRundownContext
+		virtual void SetCurrentRundown(const TSharedPtr<FAvaRundownServer>& InRundownServer, const FSoftObjectPath& InRundownPath, const TSharedPtr<FRundownEntry>& InRundownEntry) override;
+		//~ End FRundownContext
 		
 		/**
 		 * Checks if previous RCP was registered.
@@ -239,16 +334,27 @@ private:
 		 */
 		void SaveCurrentRemoteControlPresetToPage(bool bInUnregister);
 	};
-	FRundownEditCommandData RundownEditCommandData;
 
-	/** Cached data for playback commands (LoadRundown and PageAction). */
-	struct FRundownPlaybackCommandData: public FRundownCache
+	// TODO: it is likely we will need an edit command context per client connection (i.e. move to FClientInfo).
+	FEditCommandContext EditCommandContext;
+
+	/**
+	 * Context for playback commands (i.e. LoadRundown, PageAction, etc). 
+	 */
+	struct FPlaybackCommandContext : public FCommandContext
 	{
-		~FRundownPlaybackCommandData();
+		virtual ~FPlaybackCommandContext() override;
+
+		//~ Begin FRundownContext
+		virtual void SetCurrentRundown(const TSharedPtr<FAvaRundownServer>& InRundownServer, const FSoftObjectPath& InRundownPath, const TSharedPtr<FRundownEntry>& InRundownEntry) override;
+		//~ End FRundownContext
 
-		void ClosePlaybackContext();
+		void InitializePlaybackContext();
+		static void ClosePlaybackContext(UAvaRundown* InRundownToClose);
 	};
-	FRundownPlaybackCommandData RundownPlaybackCommandData;
+
+	// TODO: Will likely need to split playback context between preview (per client) and program (per rundown).
+	FPlaybackCommandContext PlaybackCommandContext;
 
 	/** Keep a map of created transient rundowns. */
 	TMap<FSoftObjectPath, TObjectPtr<UAvaRundown>> ManagedRundowns;
diff --git a/Engine/Plugins/Experimental/Avalanche/Source/AvalancheMedia/Public/IAvaMediaModule.h b/Engine/Plugins/Experimental/Avalanche/Source/AvalancheMedia/Public/IAvaMediaModule.h
--- a/Engine/Plugins/Experimental/Avalanche/Source/AvalancheMedia/Public/IAvaMediaModule.h
+++ b/Engine/Plugins/Experimental/Avalanche/Source/AvalancheMedia/Public/IAvaMediaModule.h
@@ -116,12 +116,22 @@ public:
 	 * call GetAvaInstanceSettings() in the deferred call instead.
 	 */
 	virtual const FAvaInstanceSettings& GetAvaInstanceSettings() const = 0;
-	
+
+	/**
+	 *	Returns true if the local playback manager is (still) available.
+	 */
+	virtual bool IsLocalPlaybackManagerAvailable() const = 0;
+
 	/**
 	 *	This is the backend for playing Motion Design assets locally.
 	 */
 	virtual FAvaPlaybackManager& GetLocalPlaybackManager() const = 0;
 
+	/**
+	 * Returns true if the managed instance cache is (still) available.
+	 */
+	virtual bool IsManagedInstanceCacheAvailable() const = 0;
+	
 	/**
 	 *	Access the "managed" Motion Design Asset Instance cache.
 	 */
diff --git a/Engine/Plugins/Experimental/Avalanche/Source/AvalancheMedia/Public/Rundown/AvaRundown.h b/Engine/Plugins/Experimental/Avalanche/Source/AvalancheMedia/Public/Rundown/AvaRundown.h
--- a/Engine/Plugins/Experimental/Avalanche/Source/AvalancheMedia/Public/Rundown/AvaRundown.h
+++ b/Engine/Plugins/Experimental/Avalanche/Source/AvalancheMedia/Public/Rundown/AvaRundown.h
@@ -422,6 +422,10 @@ public:
 	FOnAvaRundownPageListChanged& GetOnInstancedPageListChanged() { return InstancedPages.OnPageListChanged; }
 	FOnAvaRundownPagesChanged& GetOnPagesChanged() { return OnPagesChanged; }
 
+	/** Delegate called to determine if the playback context can be closed. */
+	DECLARE_MULTICAST_DELEGATE_TwoParams(FOnCanClosePlaybackContext, const UAvaRundown*, bool& /*bOutResult*/);
+	FOnCanClosePlaybackContext::RegistrationType& GetOnCanClosePlaybackContext() const { return OnCanClosePlaybackContext; }
+
 	/**
 	 * Since the playback context is part of the asset for now, there is an explicit call to initialize it.
 	 * This would be done by the editor (or server).
@@ -430,6 +434,11 @@ public:
 	 */
 	void InitializePlaybackContext();
 
+	/**
+	 * Determines if the playback context can be closed.  
+	 */
+	bool CanClosePlaybackContext() const;
+	
 	/**
 	 * Similarly, when the editor is done, it can close the playback context. 
 	 * This will clean up the internal structures for playback and optionally stop all the pages.
@@ -720,4 +729,5 @@ public:
 protected:
 	FOnAvaRundownPagesChanged OnPagesChanged;
 	FOnActiveListChanged OnActiveListChanged;
+	mutable FOnCanClosePlaybackContext OnCanClosePlaybackContext;
 };
diff --git a/Engine/Plugins/Experimental/Avalanche/Source/AvalancheMediaEditor/Private/Rundown/AvaRundownEditor.cpp b/Engine/Plugins/Experimental/Avalanche/Source/AvalancheMediaEditor/Private/Rundown/AvaRundownEditor.cpp
--- a/Engine/Plugins/Experimental/Avalanche/Source/AvalancheMediaEditor/Private/Rundown/AvaRundownEditor.cpp
+++ b/Engine/Plugins/Experimental/Avalanche/Source/AvalancheMediaEditor/Private/Rundown/AvaRundownEditor.cpp
@@ -195,11 +195,15 @@ FAvaRundownEditor::~FAvaRundownEditor()
 	{
 		Rundown->GetOnActiveListChanged().RemoveAll(this);
 		Rundown->GetOnPagePlayerAdded().RemoveAll(this);
-		
-		const bool bStopPages = bStopPagesOnCloseOverride.IsSet() ?
-			bStopPagesOnCloseOverride.GetValue() : UE::AvaRundownEditor::Private::ShouldStopPagesOnClose();
+		Rundown->GetOnCanClosePlaybackContext().RemoveAll(this);
+
+		if (Rundown->CanClosePlaybackContext())
+		{
+			const bool bStopPages = bStopPagesOnCloseOverride.IsSet() ?
+				bStopPagesOnCloseOverride.GetValue() : UE::AvaRundownEditor::Private::ShouldStopPagesOnClose();
 
-		Rundown->ClosePlaybackContext(bStopPages);
+			Rundown->ClosePlaybackContext(bStopPages);
+		}
 	}
 
 	if (InputProcessor.IsValid() && FSlateApplication::IsInitialized())
@@ -241,6 +245,7 @@ void FAvaRundownEditor::InitRundownEditor(const EToolkitMode::Type InMode
 	{
 		InRundown->GetOnActiveListChanged().AddSP(this, &FAvaRundownEditor::OnActiveSubListChanged);
 		InRundown->GetOnPagePlayerAdded().AddSP(this, &FAvaRundownEditor::HandleOnPagePlayerAdded);
+		InRundown->GetOnCanClosePlaybackContext().AddSP(this, &FAvaRundownEditor::OnCanClosePlaybackContext);
 		InRundown->InitializePlaybackContext();
 	}
 
@@ -1035,6 +1040,15 @@ void FAvaRundownEditor::HandleOnPagePlayerAdded(UAvaRundown* InRundown, UAvaRund
 	}
 }
 
+void FAvaRundownEditor::OnCanClosePlaybackContext(const UAvaRundown* InRundown, bool& bOutResult) const
+{
+	const UAvaRundown* Rundown = AvaRundown.Get();
+	if (Rundown != nullptr && InRundown == Rundown)
+	{
+		bOutResult = false;
+	}
+}
+
 TSharedPtr<SDockTab> FAvaRundownEditor::CreateSubListTab(int32 InSubListIndex)
 {
 	TSharedPtr<FApplicationMode> AppMode = GetCurrentModePtr();
diff --git a/Engine/Plugins/Experimental/Avalanche/Source/AvalancheMediaEditor/Private/Rundown/DetailsView/SAvaRundownPageDetails.cpp b/Engine/Plugins/Experimental/Avalanche/Source/AvalancheMediaEditor/Private/Rundown/DetailsView/SAvaRundownPageDetails.cpp
--- a/Engine/Plugins/Experimental/Avalanche/Source/AvalancheMediaEditor/Private/Rundown/DetailsView/SAvaRundownPageDetails.cpp
+++ b/Engine/Plugins/Experimental/Avalanche/Source/AvalancheMediaEditor/Private/Rundown/DetailsView/SAvaRundownPageDetails.cpp
@@ -204,7 +204,7 @@ SAvaRundownPageDetails::~SAvaRundownPageDetails()
 			Rundown->GetOnPagesChanged().RemoveAll(this);
 		}
 	}
-	if (IAvaMediaModule::IsModuleLoaded())
+	if (IAvaMediaModule::IsModuleLoaded() && IAvaMediaModule::Get().IsManagedInstanceCacheAvailable())
 	{
 		IAvaMediaModule::Get().GetManagedInstanceCache().OnEntryInvalidated.RemoveAll(this);
 	}
diff --git a/Engine/Plugins/Experimental/Avalanche/Source/AvalancheMediaEditor/Public/Rundown/AvaRundownEditor.h b/Engine/Plugins/Experimental/Avalanche/Source/AvalancheMediaEditor/Public/Rundown/AvaRundownEditor.h
--- a/Engine/Plugins/Experimental/Avalanche/Source/AvalancheMediaEditor/Public/Rundown/AvaRundownEditor.h
+++ b/Engine/Plugins/Experimental/Avalanche/Source/AvalancheMediaEditor/Public/Rundown/AvaRundownEditor.h
@@ -186,6 +186,7 @@ protected:
 
 	void OnActiveSubListChanged();
 	void HandleOnPagePlayerAdded(UAvaRundown* InRundown, UAvaRundownPagePlayer* InPagePlayer);
+	void OnCanClosePlaybackContext(const UAvaRundown* InRundown, bool& bOutResult) const;
 
 	TWeakObjectPtr<UAvaRundown> AvaRundown;
 
