Motion Design Transition: Selection details to refresh when parameters change

diff --git a/Engine/Plugins/Experimental/Avalanche/Source/AvalancheTransitionEditor/Private/Views/SAvaTransitionSelectionDetails.cpp b/Engine/Plugins/Experimental/Avalanche/Source/AvalancheTransitionEditor/Private/Views/SAvaTransitionSelectionDetails.cpp
--- a/Engine/Plugins/Experimental/Avalanche/Source/AvalancheTransitionEditor/Private/Views/SAvaTransitionSelectionDetails.cpp
+++ b/Engine/Plugins/Experimental/Avalanche/Source/AvalancheTransitionEditor/Private/Views/SAvaTransitionSelectionDetails.cpp
@@ -7,6 +7,7 @@
 #include "IDetailsView.h"
 #include "Modules/ModuleManager.h"
 #include "PropertyEditorModule.h"
+#include "StateTreeDelegates.h"
 #include "StateTreeState.h"
 #include "ViewModels/AvaTransitionViewModel.h"
 
@@ -39,7 +40,8 @@ void SAvaTransitionSelectionDetails::Construct(const FArguments& InArgs, const T
 
 	OnSelectionChanged(InSelection->GetSelectedItems());
 
-	InSelection->OnSelectionChanged().AddSP(this, &SAvaTransitionSelectionDetails::OnSelectionChanged);
+	OnParametersChangedHandle = UE::StateTree::Delegates::OnParametersChanged.AddSP(this, &SAvaTransitionSelectionDetails::Refresh);
+	OnSelectionChangedHandle = InSelection->OnSelectionChanged().AddSP(this, &SAvaTransitionSelectionDetails::OnSelectionChanged);
 
 	ChildSlot
 	[
@@ -49,9 +51,21 @@ void SAvaTransitionSelectionDetails::Construct(const FArguments& InArgs, const T
 
 SAvaTransitionSelectionDetails::~SAvaTransitionSelectionDetails()
 {
+	UE::StateTree::Delegates::OnParametersChanged.Remove(OnParametersChangedHandle);
+	OnParametersChangedHandle.Reset();
+
 	if (TSharedPtr<FAvaTransitionSelection> Selection = SelectionWeak.Pin())
 	{
-		Selection->OnSelectionChanged().RemoveAll(this);
+		Selection->OnSelectionChanged().Remove(OnSelectionChangedHandle);
+		OnSelectionChangedHandle.Reset();
+	}
+}
+
+void SAvaTransitionSelectionDetails::Refresh(const UStateTree& InStateTree)
+{
+	if (DetailsView.IsValid())
+	{
+		DetailsView->ForceRefresh();
 	}
 }
 
diff --git a/Engine/Plugins/Experimental/Avalanche/Source/AvalancheTransitionEditor/Private/Views/SAvaTransitionSelectionDetails.h b/Engine/Plugins/Experimental/Avalanche/Source/AvalancheTransitionEditor/Private/Views/SAvaTransitionSelectionDetails.h
--- a/Engine/Plugins/Experimental/Avalanche/Source/AvalancheTransitionEditor/Private/Views/SAvaTransitionSelectionDetails.h
+++ b/Engine/Plugins/Experimental/Avalanche/Source/AvalancheTransitionEditor/Private/Views/SAvaTransitionSelectionDetails.h
@@ -3,12 +3,14 @@
 #pragma once
 
 #include "Containers/ArrayView.h"
+#include "Delegates/IDelegateInstance.h"
 #include "Widgets/SCompoundWidget.h"
 
 class FAvaTransitionSelection;
 class FAvaTransitionViewModel;
 class IDetailsView;
 class UAvaTransitionTreeEditorData;
+class UStateTree;
 struct FPropertyChangedEvent;
 
 class SAvaTransitionSelectionDetails : public SCompoundWidget
@@ -28,6 +30,8 @@ public:
 	virtual ~SAvaTransitionSelectionDetails() override;
 
 private:
+	void Refresh(const UStateTree& InStateTree);
+
 	void OnSelectionChanged(TConstArrayView<TSharedRef<FAvaTransitionViewModel>> InSelectedItems);
 
 	void OnFinishedChangingProperties(const FPropertyChangedEvent& InPropertyChangedEvent);
@@ -35,4 +39,7 @@ private:
 	TWeakPtr<FAvaTransitionSelection> SelectionWeak;
 
 	TSharedPtr<IDetailsView> DetailsView;
+
+	FDelegateHandle OnParametersChangedHandle;
+	FDelegateHandle OnSelectionChangedHandle;
 };
