Motion Design: added missing blueprint type to TL Sequence enums + added Display Names for enums to not have "ava" in them.

diff --git a/Engine/Plugins/Experimental/Avalanche/Source/AvalancheSequence/Public/AvaSequenceShared.h b/Engine/Plugins/Experimental/Avalanche/Source/AvalancheSequence/Public/AvaSequenceShared.h
--- a/Engine/Plugins/Experimental/Avalanche/Source/AvalancheSequence/Public/AvaSequenceShared.h
+++ b/Engine/Plugins/Experimental/Avalanche/Source/AvalancheSequence/Public/AvaSequenceShared.h
@@ -16,7 +16,7 @@ namespace UE::AvaSequence
 	constexpr double SmallSubFrame = 0.00000005;
 }
 
-UENUM(BlueprintType)
+UENUM(BlueprintType, DisplayName="Motion Design Sequence Play Mode")
 enum class EAvaSequencePlayMode : uint8
 {
 	/** Sequence plays and loops from the beginning to the end. */
@@ -26,7 +26,7 @@ enum class EAvaSequencePlayMode : uint8
 	Reverse,
 };
 
-UENUM(BlueprintType)
+UENUM(BlueprintType, DisplayName="Motion Design Sequence Time Type")
 enum class EAvaSequenceTimeType : uint8
 {
 	None UMETA(Hidden),
@@ -35,7 +35,7 @@ enum class EAvaSequenceTimeType : uint8
 	Mark,
 };
 
-USTRUCT(BlueprintType)
+USTRUCT(BlueprintType, DisplayName="Motion Design Sequence Time")
 struct FAvaSequenceTime
 {
 	GENERATED_BODY()
@@ -116,7 +116,7 @@ struct TStructOpsTypeTraits<FAvaSequenceTime> : public TStructOpsTypeTraitsBase2
 	};
 };
 
-USTRUCT(BlueprintType)
+USTRUCT(BlueprintType, DisplayName="Motion Design Sequence Play Advanced Settings")
 struct FAvaSequencePlayAdvancedSettings
 {
 	GENERATED_BODY()
@@ -132,7 +132,7 @@ struct FAvaSequencePlayAdvancedSettings
 	bool bRestoreState = false;
 };
 
-USTRUCT(BlueprintType)
+USTRUCT(BlueprintType, DisplayName="Motion Design Sequence Play Parameters")
 struct FAvaSequencePlayParams
 {
 	GENERATED_BODY()
diff --git a/Engine/Plugins/Experimental/Avalanche/Source/AvalancheSequence/Public/Transition/AvaTransitionSequenceEnums.h b/Engine/Plugins/Experimental/Avalanche/Source/AvalancheSequence/Public/Transition/AvaTransitionSequenceEnums.h
--- a/Engine/Plugins/Experimental/Avalanche/Source/AvalancheSequence/Public/Transition/AvaTransitionSequenceEnums.h
+++ b/Engine/Plugins/Experimental/Avalanche/Source/AvalancheSequence/Public/Transition/AvaTransitionSequenceEnums.h
@@ -4,7 +4,7 @@
 
 #include "AvaTransitionSequenceEnums.generated.h"
 
-UENUM()
+UENUM(BlueprintType, DisplayName="Motion Design Sequence Wait Type")
 enum class EAvaTransitionSequenceWaitType : uint8
 {
 	None UMETA(Hidden),
@@ -16,7 +16,7 @@ enum class EAvaTransitionSequenceWaitType : uint8
 	WaitUntilStop,
 };
 
-UENUM()
+UENUM(BlueprintType, DisplayName="Motion Design Sequence Query Type")
 enum class EAvaTransitionSequenceQueryType : uint8
 {
 	None UMETA(Hidden),
diff --git a/Engine/Plugins/Experimental/Avalanche/Source/AvalancheTransition/Public/AvaTransitionEnums.h b/Engine/Plugins/Experimental/Avalanche/Source/AvalancheTransition/Public/AvaTransitionEnums.h
--- a/Engine/Plugins/Experimental/Avalanche/Source/AvalancheTransition/Public/AvaTransitionEnums.h
+++ b/Engine/Plugins/Experimental/Avalanche/Source/AvalancheTransition/Public/AvaTransitionEnums.h
@@ -17,7 +17,7 @@ enum class EAvaTransitionIterationResult : uint8
 	Continue,
 };
 
-UENUM(BlueprintType)
+UENUM(BlueprintType, DisplayName="Motion Design Transition Type")
 enum class EAvaTransitionType : uint8
 {
 	None = 0 UMETA(Hidden),
@@ -27,7 +27,7 @@ enum class EAvaTransitionType : uint8
 };
 ENUM_CLASS_FLAGS(EAvaTransitionType);
 
-UENUM(BlueprintType)
+UENUM(BlueprintType, DisplayName="Motion Design Transition Run State")
 enum class EAvaTransitionRunState : uint8
 {
 	Unknown UMETA(Hidden),
@@ -35,7 +35,7 @@ enum class EAvaTransitionRunState : uint8
 	Finished,
 };
 
-UENUM(BlueprintType)
+UENUM(BlueprintType, DisplayName="Motion Design Transition Comparison Result")
 enum class EAvaTransitionComparisonResult : uint8
 {
 	None,
@@ -43,14 +43,14 @@ enum class EAvaTransitionComparisonResult : uint8
 	Same,
 };
 
-UENUM(BlueprintType)
+UENUM(BlueprintType, DisplayName="Motion Design Transition Scene Type")
 enum class EAvaTransitionSceneType : uint8
 {
 	This,
 	Other,
 };
 
-UENUM(BlueprintType)
+UENUM(BlueprintType, DisplayName="Motion Design Transition Layer Compare Type")
 enum class EAvaTransitionLayerCompareType : uint8
 {
 	None UMETA(Hidden),
