[Motion Design] Potential crash fixes in FAvaBroadcastOutputChannel functions.

Repros:
- After deleting some channels. When FAvaBroadcastOutputTileItem updates on an invalid media output object.
- Updating level asset in rundown. List of actors in level can contain null elements.

diff --git a/Engine/Plugins/Experimental/Avalanche/Source/AvalancheMedia/Private/Broadcast/Channel/AvaBroadcastOutputChannel.cpp b/Engine/Plugins/Experimental/Avalanche/Source/AvalancheMedia/Private/Broadcast/Channel/AvaBroadcastOutputChannel.cpp
--- a/Engine/Plugins/Experimental/Avalanche/Source/AvalancheMedia/Private/Broadcast/Channel/AvaBroadcastOutputChannel.cpp
+++ b/Engine/Plugins/Experimental/Avalanche/Source/AvalancheMedia/Private/Broadcast/Channel/AvaBroadcastOutputChannel.cpp
@@ -311,6 +311,11 @@ EAvaBroadcastChannelState FAvaBroadcastOutputChannel::RefreshState()
 
 EAvaBroadcastIssueSeverity FAvaBroadcastOutputChannel::GetMediaOutputIssueSeverity(EAvaBroadcastOutputState InOutputState, const UMediaOutput* InMediaOutput) const
 {
+	if (!IsValid(InMediaOutput))
+	{
+		return EAvaBroadcastIssueSeverity::Errors;
+	}
+	
 	if (InOutputState == EAvaBroadcastOutputState::Live || InOutputState == EAvaBroadcastOutputState::Preparing)
 	{
 		// If the output is broadcasting remote, fetch the status from the playback client (which is proxying that output's status).
@@ -341,25 +346,32 @@ EAvaBroadcastIssueSeverity FAvaBroadcastOutputChannel::GetMediaOutputIssueSeveri
 
 const TArray<FString>& FAvaBroadcastOutputChannel::GetMediaOutputIssueMessages(const UMediaOutput* InMediaOutput) const
 {
-	if (IAvaMediaModule::Get().IsPlaybackClientStarted() && IsMediaOutputRemote(InMediaOutput))
-	{
-		const FAvaBroadcastMediaOutputInfo& OutputInfo = GetMediaOutputInfo(InMediaOutput);
-		const IAvaPlaybackClient& PlaybackClient = IAvaMediaModule::Get().GetPlaybackClient();
-		return PlaybackClient.GetMediaOutputIssueMessages(GetMediaOutputServerName(InMediaOutput), GetChannelName().ToString(), OutputInfo.Guid);
-	}
-	else
+	if (IsValid(InMediaOutput))
 	{
+		if (IAvaMediaModule::Get().IsPlaybackClientStarted() && IsMediaOutputRemote(InMediaOutput))
+		{
+			const FAvaBroadcastMediaOutputInfo& OutputInfo = GetMediaOutputInfo(InMediaOutput);
+			const IAvaPlaybackClient& PlaybackClient = IAvaMediaModule::Get().GetPlaybackClient();
+			return PlaybackClient.GetMediaOutputIssueMessages(GetMediaOutputServerName(InMediaOutput), GetChannelName().ToString(), OutputInfo.Guid);
+		}
+		
 		if (const FLocalMediaOutputStatus* LocalMediaOutputStatus = LocalMediaOutputStatuses.Find(InMediaOutput))
 		{
 			return LocalMediaOutputStatus->Messages;
 		}
 	}
+
 	static const TArray<FString> EmptyStringArray;
 	return EmptyStringArray;
 }
 
 EAvaBroadcastOutputState FAvaBroadcastOutputChannel::GetMediaOutputState(const UMediaOutput* InMediaOutput) const
 {
+	if (!IsValid(InMediaOutput))
+	{
+		return EAvaBroadcastOutputState::Invalid;
+	}
+	
 	if (IsMediaOutputRemote(InMediaOutput))
 	{
 		if (IAvaMediaModule::Get().IsPlaybackClientStarted())
diff --git a/Engine/Plugins/Experimental/Avalanche/Source/AvalancheMedia/Private/Rundown/AvaRundownPageAssetUtils.cpp b/Engine/Plugins/Experimental/Avalanche/Source/AvalancheMedia/Private/Rundown/AvaRundownPageAssetUtils.cpp
--- a/Engine/Plugins/Experimental/Avalanche/Source/AvalancheMedia/Private/Rundown/AvaRundownPageAssetUtils.cpp
+++ b/Engine/Plugins/Experimental/Avalanche/Source/AvalancheMedia/Private/Rundown/AvaRundownPageAssetUtils.cpp
@@ -74,7 +74,7 @@ const UAvaTransitionTree* FAvaRundownPageAssetUtils::FindTransitionTree(const IA
 	static FString AvaTransitionBehaviorActor(TEXT("AvaTransitionBehaviorActor"));
 	for (const AActor* Actor : SceneLevel->Actors)
 	{
-		const UClass* ActorClass = Actor->GetClass();
+		const UClass* ActorClass = IsValid(Actor) ? Actor->GetClass() : nullptr;
 		if (ActorClass && ActorClass->GetName() == AvaTransitionBehaviorActor)
 		{
 			const FObjectPropertyBase* const Property = FindFProperty<FObjectPropertyBase>(ActorClass, TEXT("TransitionTree"));
