Motion Design Transition: Added Blueprint support for Tasks/Conditions

diff --git a/Engine/Plugins/Experimental/Avalanche/Source/AvalancheTransition/Private/AvaTransitionTreeSchema.cpp b/Engine/Plugins/Experimental/Avalanche/Source/AvalancheTransition/Private/AvaTransitionTreeSchema.cpp
--- a/Engine/Plugins/Experimental/Avalanche/Source/AvalancheTransition/Private/AvaTransitionTreeSchema.cpp
+++ b/Engine/Plugins/Experimental/Avalanche/Source/AvalancheTransition/Private/AvaTransitionTreeSchema.cpp
@@ -4,10 +4,12 @@
 #include "AvaTransitionContext.h"
 #include "AvaTransitionScene.h"
 #include "Conditions/AvaTransitionCondition.h"
+#include "Conditions/AvaTransitionConditionBlueprint.h"
 #include "StateTreeEvaluatorBase.h"
 #include "StateTreeTaskBase.h"
 #include "Subsystems/WorldSubsystem.h"
 #include "Tasks/AvaTransitionTask.h"
+#include "Tasks/AvaTransitionTaskBlueprint.h"
 
 bool UAvaTransitionTreeSchema::IsStructAllowed(const UScriptStruct* InScriptStruct) const
 {
@@ -15,6 +17,12 @@ bool UAvaTransitionTreeSchema::IsStructAllowed(const UScriptStruct* InScriptStru
 		|| InScriptStruct->IsChildOf(FAvaTransitionCondition::StaticStruct());
 }
 
+bool UAvaTransitionTreeSchema::IsClassAllowed(const UClass* InClass) const
+{
+	return InClass->IsChildOf<UAvaTransitionTaskBlueprint>()
+		|| InClass->IsChildOf<UAvaTransitionConditionBlueprint>();
+}
+
 bool UAvaTransitionTreeSchema::IsExternalItemAllowed(const UStruct& InStruct) const
 {
 	return InStruct.IsChildOf(FAvaTransitionContext::StaticStruct())
diff --git a/Engine/Plugins/Experimental/Avalanche/Source/AvalancheTransition/Public/AvaTransitionTreeSchema.h b/Engine/Plugins/Experimental/Avalanche/Source/AvalancheTransition/Public/AvaTransitionTreeSchema.h
--- a/Engine/Plugins/Experimental/Avalanche/Source/AvalancheTransition/Public/AvaTransitionTreeSchema.h
+++ b/Engine/Plugins/Experimental/Avalanche/Source/AvalancheTransition/Public/AvaTransitionTreeSchema.h
@@ -17,6 +17,7 @@ public:
 protected:
 	//~ Begin UStateTreeSchema
 	virtual bool IsStructAllowed(const UScriptStruct* InScriptStruct) const override;
+	virtual bool IsClassAllowed(const UClass* InClass) const override;
 	virtual bool IsExternalItemAllowed(const UStruct& InStruct) const override;
 #if WITH_EDITOR
 	virtual bool AllowEnterConditions() const override { return true; }
diff --git a/Engine/Plugins/Experimental/Avalanche/Source/AvalancheTransition/Public/Conditions/AvaTransitionConditionBlueprint.h b/Engine/Plugins/Experimental/Avalanche/Source/AvalancheTransition/Public/Conditions/AvaTransitionConditionBlueprint.h
new file mode 100644
--- /dev/null
+++ b/Engine/Plugins/Experimental/Avalanche/Source/AvalancheTransition/Public/Conditions/AvaTransitionConditionBlueprint.h
@@ -0,0 +1,12 @@
+// Copyright Epic Games, Inc. All Rights Reserved.
+
+#pragma once
+
+#include "Blueprint/StateTreeConditionBlueprintBase.h"
+#include "AvaTransitionConditionBlueprint.generated.h"
+
+UCLASS(MinimalAPI, Abstract)
+class UAvaTransitionConditionBlueprint : public UStateTreeConditionBlueprintBase
+{
+	GENERATED_BODY()
+};
diff --git a/Engine/Plugins/Experimental/Avalanche/Source/AvalancheTransition/Public/Tasks/AvaTransitionTaskBlueprint.h b/Engine/Plugins/Experimental/Avalanche/Source/AvalancheTransition/Public/Tasks/AvaTransitionTaskBlueprint.h
new file mode 100644
--- /dev/null
+++ b/Engine/Plugins/Experimental/Avalanche/Source/AvalancheTransition/Public/Tasks/AvaTransitionTaskBlueprint.h
@@ -0,0 +1,12 @@
+// Copyright Epic Games, Inc. All Rights Reserved.
+
+#pragma once
+
+#include "Blueprint/StateTreeTaskBlueprintBase.h"
+#include "AvaTransitionTaskBlueprint.generated.h"
+
+UCLASS(MinimalAPI, Abstract)
+class UAvaTransitionTaskBlueprint : public UStateTreeTaskBlueprintBase
+{
+	GENERATED_BODY()
+};
