Motion Design: Added Scene Settings Customization Extension points.
Transition Extension uses these to add the State Tree Parameter property to that Tab to modify Instance Parameters.

diff --git a/Engine/Plugins/Experimental/Avalanche/Source/AvalancheEditor/Private/AvaEditorModule.cpp b/Engine/Plugins/Experimental/Avalanche/Source/AvalancheEditor/Private/AvaEditorModule.cpp
--- a/Engine/Plugins/Experimental/Avalanche/Source/AvalancheEditor/Private/AvaEditorModule.cpp
+++ b/Engine/Plugins/Experimental/Avalanche/Source/AvalancheEditor/Private/AvaEditorModule.cpp
@@ -273,7 +273,7 @@ void FAvaEditorModule::RegisterCustomLayouts()
 		&FAvaMeshesDetailCustomization::MakeInstance));
 
 	PropertyModule.RegisterCustomClassLayout(UAvaSceneSettings::StaticClass()->GetFName(), FOnGetDetailCustomizationInstance::CreateStatic(
-		&FAvaSceneSettingsCustomization::MakeInstance));
+		&FAvaSceneSettingsCustomization::MakeDefaultInstance));
 
 	const TSharedRef<FAvaVectorPropertyTypeIdentifier> VectorPropertyTypeIdentifier = MakeShared<FAvaVectorPropertyTypeIdentifier>();
 
diff --git a/Engine/Plugins/Experimental/Avalanche/Source/AvalancheEditor/Private/DetailView/Customizations/AvaSceneSettingsCustomization.cpp b/Engine/Plugins/Experimental/Avalanche/Source/AvalancheEditor/Private/DetailView/Customizations/AvaSceneSettingsCustomization.cpp
--- a/Engine/Plugins/Experimental/Avalanche/Source/AvalancheEditor/Private/DetailView/Customizations/AvaSceneSettingsCustomization.cpp
+++ b/Engine/Plugins/Experimental/Avalanche/Source/AvalancheEditor/Private/DetailView/Customizations/AvaSceneSettingsCustomization.cpp
@@ -2,26 +2,51 @@
 
 #include "AvaSceneSettingsCustomization.h"
 #include "AvaSceneSettings.h"
+#include "DetailCategoryBuilder.h"
 #include "DetailLayoutBuilder.h"
 #include "IAvaAttributeEditorModule.h"
+#include "IAvaEditor.h"
+#include "IAvaEditorExtension.h"
 #include "IAvaSceneRigEditorModule.h"
 #include "PropertyHandle.h"
 #include "Templates/SharedPointer.h"
 
-TSharedRef<IDetailCustomization> FAvaSceneSettingsCustomization::MakeInstance()
+TSharedRef<IDetailCustomization> FAvaSceneSettingsCustomization::MakeDefaultInstance()
 {
-	return MakeShared<FAvaSceneSettingsCustomization>();
+	return FAvaSceneSettingsCustomization::MakeInstance(nullptr);
 }
 
-void FAvaSceneSettingsCustomization::CustomizeDetails(IDetailLayoutBuilder& DetailBuilder)
+TSharedRef<IDetailCustomization> FAvaSceneSettingsCustomization::MakeInstance(TWeakPtr<IAvaEditor> InEditorWeak)
 {
-	if (TSharedPtr<IPropertyHandle> SceneRigHandle = DetailBuilder.GetProperty(UAvaSceneSettings::GetSceneRigPropertyName()))
+	return MakeShared<FAvaSceneSettingsCustomization>(InEditorWeak);
+}
+
+FAvaSceneSettingsCustomization::FAvaSceneSettingsCustomization(const TWeakPtr<IAvaEditor>& InEditorWeak)
+	: EditorWeak(InEditorWeak)
+{
+}
+
+void FAvaSceneSettingsCustomization::CustomizeDetails(IDetailLayoutBuilder& InDetailBuilder)
+{
+	if (TSharedPtr<IPropertyHandle> SceneRigHandle = InDetailBuilder.GetProperty(UAvaSceneSettings::GetSceneRigPropertyName()))
+	{
+		IAvaSceneRigEditorModule::Get().CustomizeSceneRig(SceneRigHandle.ToSharedRef(), InDetailBuilder);
+	}
+
+	if (TSharedPtr<IPropertyHandle> AttributesHandle = InDetailBuilder.GetProperty(UAvaSceneSettings::GetSceneAttributesPropertyName()))
 	{
-		IAvaSceneRigEditorModule::Get().CustomizeSceneRig(SceneRigHandle.ToSharedRef(), DetailBuilder);
+		IAvaAttributeEditorModule::Get().CustomizeAttributes(AttributesHandle.ToSharedRef(), InDetailBuilder);
 	}
 
-	if (TSharedPtr<IPropertyHandle> AttributesHandle = DetailBuilder.GetProperty(UAvaSceneSettings::GetSceneAttributesPropertyName()))
+	if (TSharedPtr<IAvaEditor> Editor = EditorWeak.Pin())
 	{
-		IAvaAttributeEditorModule::Get().CustomizeAttributes(AttributesHandle.ToSharedRef(), DetailBuilder);
+		for (const TSharedRef<IAvaEditorExtension>& Extension : Editor->GetExtensions())
+		{
+			FName CategoryName = Extension->GetCategoryName();
+			if (!CategoryName.IsNone())
+			{
+				Extension->ExtendSettingsCategory(InDetailBuilder.EditCategory(CategoryName));
+			}
+		}
 	}
 }
diff --git a/Engine/Plugins/Experimental/Avalanche/Source/AvalancheEditor/Private/DetailView/Customizations/AvaSceneSettingsCustomization.h b/Engine/Plugins/Experimental/Avalanche/Source/AvalancheEditor/Private/DetailView/Customizations/AvaSceneSettingsCustomization.h
--- a/Engine/Plugins/Experimental/Avalanche/Source/AvalancheEditor/Private/DetailView/Customizations/AvaSceneSettingsCustomization.h
+++ b/Engine/Plugins/Experimental/Avalanche/Source/AvalancheEditor/Private/DetailView/Customizations/AvaSceneSettingsCustomization.h
@@ -5,14 +5,22 @@
 #include "IDetailCustomization.h"
 #include "Templates/SharedPointerFwd.h"
 
+class IAvaEditor;
 class IDetailLayoutBuilder;
 
 class FAvaSceneSettingsCustomization : public IDetailCustomization
 {
 public:
-	static TSharedRef<IDetailCustomization> MakeInstance();
+	static TSharedRef<IDetailCustomization> MakeDefaultInstance();
+
+	static TSharedRef<IDetailCustomization> MakeInstance(TWeakPtr<IAvaEditor> InEditorWeak);
+
+	explicit FAvaSceneSettingsCustomization(const TWeakPtr<IAvaEditor>& InEditorWeak);
 
 	//~ Begin IDetailCustomization
-	virtual void CustomizeDetails(IDetailLayoutBuilder& DetailBuilder) override;
+	virtual void CustomizeDetails(IDetailLayoutBuilder& InDetailBuilder) override;
 	//~ End IDetailCustomization
+
+private:
+	TWeakPtr<IAvaEditor> EditorWeak;
 };
diff --git a/Engine/Plugins/Experimental/Avalanche/Source/AvalancheEditor/Private/Scene/AvaSceneSettingsTabSpawner.cpp b/Engine/Plugins/Experimental/Avalanche/Source/AvalancheEditor/Private/Scene/AvaSceneSettingsTabSpawner.cpp
--- a/Engine/Plugins/Experimental/Avalanche/Source/AvalancheEditor/Private/Scene/AvaSceneSettingsTabSpawner.cpp
+++ b/Engine/Plugins/Experimental/Avalanche/Source/AvalancheEditor/Private/Scene/AvaSceneSettingsTabSpawner.cpp
@@ -2,6 +2,7 @@
 
 #include "AvaSceneSettingsTabSpawner.h"
 #include "AvaSceneSettings.h"
+#include "DetailView/Customizations/AvaSceneSettingsCustomization.h"
 #include "IAvaEditor.h"
 #include "IAvaSceneInterface.h"
 #include "IDetailsView.h"
@@ -49,6 +50,10 @@ TSharedRef<SWidget> FAvaSceneSettingsTabSpawner::CreateTabBody()
 	DetailsViewArgs.NameAreaSettings = FDetailsViewArgs::HideNameArea;
 
 	TSharedRef<IDetailsView> DetailsView = PropertyEditorModule.CreateDetailView(DetailsViewArgs);
+
+	DetailsView->RegisterInstancedCustomPropertyLayout(UAvaSceneSettings::StaticClass()
+		, FOnGetDetailCustomizationInstance::CreateStatic(&FAvaSceneSettingsCustomization::MakeInstance, EditorWeak));
+
 	DetailsView->SetObject(SceneSettings);
 	return DetailsView;
 }
diff --git a/Engine/Plugins/Experimental/Avalanche/Source/AvalancheEditor/Private/Transition/AvaTransitionExtension.cpp b/Engine/Plugins/Experimental/Avalanche/Source/AvalancheEditor/Private/Transition/AvaTransitionExtension.cpp
--- a/Engine/Plugins/Experimental/Avalanche/Source/AvalancheEditor/Private/Transition/AvaTransitionExtension.cpp
+++ b/Engine/Plugins/Experimental/Avalanche/Source/AvalancheEditor/Private/Transition/AvaTransitionExtension.cpp
@@ -10,6 +10,7 @@
 #include "Behavior/IAvaTransitionBehavior.h"
 #include "Conditions/AvaTransitionSceneMatchCondition.h"
 #include "Conditions/AvaTransitionTypeMatchCondition.h"
+#include "DetailCategoryBuilder.h"
 #include "Editor.h"
 #include "IAvaTransitionEditorModule.h"
 #include "IRemoteControlUIModule.h"
@@ -123,6 +124,20 @@ void FAvaTransitionExtension::Deactivate()
 	CloseTransitionEditor();
 }
 
+FName FAvaTransitionExtension::GetCategoryName() const
+{
+	return TEXT("TransitionLogic");
+}
+
+void FAvaTransitionExtension::ExtendSettingsCategory(IDetailCategoryBuilder& InCategoryBuilder)
+{
+	if (IAvaTransitionBehavior* TransitionBehavior = GetTransitionBehavior())
+	{
+		InCategoryBuilder.SetDisplayName(LOCTEXT("CategoryDisplayName", "Transition Logic"));
+		InCategoryBuilder.AddExternalObjectProperty({ &TransitionBehavior->AsUObject() }, TEXT("StateTreeReference"));
+	}
+}
+
 void FAvaTransitionExtension::ExtendToolbarMenu(UToolMenu& InMenu)
 {
 	FToolMenuSection& Section = InMenu.FindOrAddSection(DefaultSectionName);
diff --git a/Engine/Plugins/Experimental/Avalanche/Source/AvalancheEditor/Private/Transition/AvaTransitionExtension.h b/Engine/Plugins/Experimental/Avalanche/Source/AvalancheEditor/Private/Transition/AvaTransitionExtension.h
--- a/Engine/Plugins/Experimental/Avalanche/Source/AvalancheEditor/Private/Transition/AvaTransitionExtension.h
+++ b/Engine/Plugins/Experimental/Avalanche/Source/AvalancheEditor/Private/Transition/AvaTransitionExtension.h
@@ -21,6 +21,8 @@ public:
 	virtual void Construct(const TSharedRef<IAvaEditor>& InEditor) override;
 	virtual void Activate() override;
 	virtual void Deactivate() override;
+	virtual FName GetCategoryName() const override;
+	virtual void ExtendSettingsCategory(IDetailCategoryBuilder& InCategoryBuilder) override;
 	virtual void ExtendToolbarMenu(UToolMenu& InMenu) override;
 	//~ End IAvaEditorExtension
 
diff --git a/Engine/Plugins/Experimental/Avalanche/Source/AvalancheEditorCore/Public/IAvaEditorExtension.h b/Engine/Plugins/Experimental/Avalanche/Source/AvalancheEditorCore/Public/IAvaEditorExtension.h
--- a/Engine/Plugins/Experimental/Avalanche/Source/AvalancheEditorCore/Public/IAvaEditorExtension.h
+++ b/Engine/Plugins/Experimental/Avalanche/Source/AvalancheEditorCore/Public/IAvaEditorExtension.h
@@ -15,6 +15,7 @@ class FAvaEditorSelection;
 class FEditorModeTools;
 class FLayoutExtender;
 class FUICommandList;
+class IDetailCategoryBuilder;
 class IToolkitHost;
 class UToolMenu;
 class UWorld;
@@ -58,6 +59,22 @@ public:
 	/** Extend the Level Editor Layout (only called when instancing for Level Editor) */
 	virtual void ExtendLevelEditorLayout(FLayoutExtender& InExtender) const {}
 
+	/**
+	 * Name of the category this Extension belongs in.
+	 * Used in places like Setting Details Category
+	 */
+	virtual FName GetCategoryName() const
+	{
+		return NAME_None;
+	}
+
+	/**
+	 * Extend the Settings Category by adding External Objects or Properties.
+	 * Must have a valid category name to proceed.
+	 * @see IAvaEditorExtension::GetCategoryName
+	 */
+	virtual void ExtendSettingsCategory(IDetailCategoryBuilder& InCategoryBuilder) {}
+
 	virtual void Save() {}
 
 	virtual void Load() {}
diff --git a/Engine/Plugins/Experimental/Avalanche/Source/AvalancheTransition/Private/Behavior/AvaTransitionBehaviorActor.h b/Engine/Plugins/Experimental/Avalanche/Source/AvalancheTransition/Private/Behavior/AvaTransitionBehaviorActor.h
--- a/Engine/Plugins/Experimental/Avalanche/Source/AvalancheTransition/Private/Behavior/AvaTransitionBehaviorActor.h
+++ b/Engine/Plugins/Experimental/Avalanche/Source/AvalancheTransition/Private/Behavior/AvaTransitionBehaviorActor.h
@@ -49,6 +49,6 @@ private:
 	UPROPERTY()
 	TObjectPtr<UAvaTransitionTree> TransitionTree;
 
-	UPROPERTY()
+	UPROPERTY(VisibleAnywhere, Category="Transition Logic")
 	FStateTreeReference StateTreeReference;
 };
