Motion Design: move tag id to its own file

diff --git a/Engine/Plugins/Experimental/Avalanche/Source/Avalanche/Private/AvaSceneState.cpp b/Engine/Plugins/Experimental/Avalanche/Source/Avalanche/Private/AvaSceneState.cpp
--- a/Engine/Plugins/Experimental/Avalanche/Source/Avalanche/Private/AvaSceneState.cpp
+++ b/Engine/Plugins/Experimental/Avalanche/Source/Avalanche/Private/AvaSceneState.cpp
@@ -18,23 +18,23 @@ void UAvaSceneState::Initialize(UAvaSceneSettings* InSceneSettings)
 
 bool UAvaSceneState::AddTagAttribute(const FAvaTagHandle& InTagHandle)
 {
-	// Return true if already existing
-	if (ContainsTagAttribute(InTagHandle))
+	if (!InTagHandle.IsValid())
 	{
-		return true;
+		return false;
 	}
 
-	if (const FAvaTag* Tag = InTagHandle.GetTag())
+	// Return true if already existing
+	if (ContainsTagAttribute(InTagHandle))
 	{
-		UAvaTagAttribute* TagAttribute = NewObject<UAvaTagAttribute>(this, NAME_None, RF_Transient);
-		check(TagAttribute);
-		TagAttribute->Tag = InTagHandle;
-
-		SceneAttributes.Add(TagAttribute);
 		return true;
 	}
 
-	return false;
+	UAvaTagAttribute* TagAttribute = NewObject<UAvaTagAttribute>(this, NAME_None, RF_Transient);
+	check(TagAttribute);
+	TagAttribute->Tag = InTagHandle;
+
+	SceneAttributes.Add(TagAttribute);
+	return true;
 }
 
 bool UAvaSceneState::RemoveTagAttribute(const FAvaTagHandle& InTagHandle)
diff --git a/Engine/Plugins/Experimental/Avalanche/Source/AvalancheMedia/Private/Rundown/Transition/AvaRundownPageTransition.h b/Engine/Plugins/Experimental/Avalanche/Source/AvalancheMedia/Private/Rundown/Transition/AvaRundownPageTransition.h
--- a/Engine/Plugins/Experimental/Avalanche/Source/AvalancheMedia/Private/Rundown/Transition/AvaRundownPageTransition.h
+++ b/Engine/Plugins/Experimental/Avalanche/Source/AvalancheMedia/Private/Rundown/Transition/AvaRundownPageTransition.h
@@ -2,7 +2,7 @@
 
 #pragma once
 
-#include "AvaTag.h"
+#include "AvaTagId.h"
 #include "Playable/AvaPlayable.h"
 #include "Playback/Transition/AvaPlaybackTransition.h"
 #include "UObject/Object.h"
diff --git a/Engine/Plugins/Experimental/Avalanche/Source/AvalancheMediaEditor/Private/Rundown/ChannelLayerStatus/SAvaRundownChannelLayerStatusList.cpp b/Engine/Plugins/Experimental/Avalanche/Source/AvalancheMediaEditor/Private/Rundown/ChannelLayerStatus/SAvaRundownChannelLayerStatusList.cpp
--- a/Engine/Plugins/Experimental/Avalanche/Source/AvalancheMediaEditor/Private/Rundown/ChannelLayerStatus/SAvaRundownChannelLayerStatusList.cpp
+++ b/Engine/Plugins/Experimental/Avalanche/Source/AvalancheMediaEditor/Private/Rundown/ChannelLayerStatus/SAvaRundownChannelLayerStatusList.cpp
@@ -2,6 +2,7 @@
 
 #include "SAvaRundownChannelLayerStatusList.h"
 
+#include "AvaTag.h"
 #include "Misc/CoreDelegates.h"
 #include "Rundown/AvaRundownEditor.h"
 #include "Rundown/AvaRundownEditorDefines.h"
diff --git a/Engine/Plugins/Experimental/Avalanche/Source/AvalancheSequence/Private/AvaSequence.cpp b/Engine/Plugins/Experimental/Avalanche/Source/AvalancheSequence/Private/AvaSequence.cpp
--- a/Engine/Plugins/Experimental/Avalanche/Source/AvalancheSequence/Private/AvaSequence.cpp
+++ b/Engine/Plugins/Experimental/Avalanche/Source/AvalancheSequence/Private/AvaSequence.cpp
@@ -5,6 +5,7 @@
 #include "AvaSequenceController.h"
 #include "AvaSequencePlayer.h"
 #include "AvaSequenceVersion.h"
+#include "AvaTag.h"
 #include "Director/AvaSequenceDirector.h"
 #include "Director/AvaSequenceDirectorBlueprint.h"
 #include "EngineUtils.h"
diff --git a/Engine/Plugins/Experimental/Avalanche/Source/AvalancheTag/Private/AvaTagLibrary.cpp b/Engine/Plugins/Experimental/Avalanche/Source/AvalancheTag/Private/AvaTagLibrary.cpp
--- a/Engine/Plugins/Experimental/Avalanche/Source/AvalancheTag/Private/AvaTagLibrary.cpp
+++ b/Engine/Plugins/Experimental/Avalanche/Source/AvalancheTag/Private/AvaTagLibrary.cpp
@@ -1,6 +1,7 @@
 // Copyright Epic Games, Inc. All Rights Reserved.
 
 #include "AvaTagLibrary.h"
+#include "AvaTag.h"
 #include "AvaTagHandle.h"
 #include "AvaTagHandleContainer.h"
 #include "AvaTagSoftHandle.h"
diff --git a/Engine/Plugins/Experimental/Avalanche/Source/AvalancheTag/Public/AvaTag.h b/Engine/Plugins/Experimental/Avalanche/Source/AvalancheTag/Public/AvaTag.h
--- a/Engine/Plugins/Experimental/Avalanche/Source/AvalancheTag/Public/AvaTag.h
+++ b/Engine/Plugins/Experimental/Avalanche/Source/AvalancheTag/Public/AvaTag.h
@@ -2,8 +2,6 @@
 
 #pragma once
 
-#include "Containers/UnrealString.h"
-#include "Misc/Guid.h"
 #include "UObject/NameTypes.h"
 #include "AvaTag.generated.h"
 
@@ -39,62 +37,3 @@ struct FAvaTag
 	UPROPERTY(EditAnywhere, Category="Tag")
 	FName TagName;
 };
-
-/**
- * Struct to identify a Tag. This is used by FAvaTagHandle to reference an FAvaTag
- */
-USTRUCT()
-struct FAvaTagId
-{
-	GENERATED_BODY()
-
-	/** Id is initialized to a zero */
-	FAvaTagId() = default;
-
-	/**
-	 * Force init where Id is initialized to a new guid.
-	 * This is used in TCppStructOps::Construct
-	 */
-	explicit FAvaTagId(EForceInit)
-		: Id(FGuid::NewGuid())
-	{
-	}
-
-	explicit FAvaTagId(const FGuid& InId)
-		: Id(InId)
-	{
-	}
-
-	bool operator==(const FAvaTagId& InOther) const
-	{
-		return Id == InOther.Id;
-	}
-
-	bool IsValid() const
-	{
-		return Id.IsValid();
-	}
-
-	FString ToString() const
-	{
-		return Id.ToString();
-	}
-
-	friend uint32 GetTypeHash(const FAvaTagId& InTagId)
-	{
-		return GetTypeHash(InTagId.Id);
-	}
-
-private:
-	UPROPERTY(EditAnywhere, Category="Tag", meta=(IgnoreForMemberInitializationTest))
-	FGuid Id;
-};
-
-template<>
-struct TStructOpsTypeTraits<FAvaTagId> : TStructOpsTypeTraitsBase2<FAvaTagId>
-{
-	enum
-	{
-		WithNoInitConstructor = true,
-	};
-};
diff --git a/Engine/Plugins/Experimental/Avalanche/Source/AvalancheTag/Public/AvaTagCollection.h b/Engine/Plugins/Experimental/Avalanche/Source/AvalancheTag/Public/AvaTagCollection.h
--- a/Engine/Plugins/Experimental/Avalanche/Source/AvalancheTag/Public/AvaTagCollection.h
+++ b/Engine/Plugins/Experimental/Avalanche/Source/AvalancheTag/Public/AvaTagCollection.h
@@ -3,7 +3,8 @@
 #pragma once
 
 #include "AvaTag.h"
-#include "Containers/Set.h"
+#include "AvaTagId.h"
+#include "Containers/Map.h"
 #include "UObject/Object.h"
 #include "AvaTagCollection.generated.h"
 
@@ -11,18 +12,18 @@
  * Tag Collection that identifies a tag with an underlying Tag Id Guid
  * and provides Tag reference capabilities
  */
-UCLASS()
-class AVALANCHETAG_API UAvaTagCollection : public UObject
+UCLASS(MinimalAPI)
+class UAvaTagCollection : public UObject
 {
 	GENERATED_BODY()
 
 public:
-	const FAvaTag* GetTag(const FAvaTagId& InTagId) const;
+	AVALANCHETAG_API const FAvaTag* GetTag(const FAvaTagId& InTagId) const;
 
 	/** Returns the keys of the Tag Map */
-	TArray<FAvaTagId> GetTagIds() const;
+	AVALANCHETAG_API TArray<FAvaTagId> GetTagIds() const;
 
-	static FName GetTagMapName();
+	AVALANCHETAG_API static FName GetTagMapName();
 
 private:
 	UPROPERTY(EditAnywhere, Category="Tag")
diff --git a/Engine/Plugins/Experimental/Avalanche/Source/AvalancheTag/Public/AvaTagHandle.h b/Engine/Plugins/Experimental/Avalanche/Source/AvalancheTag/Public/AvaTagHandle.h
--- a/Engine/Plugins/Experimental/Avalanche/Source/AvalancheTag/Public/AvaTagHandle.h
+++ b/Engine/Plugins/Experimental/Avalanche/Source/AvalancheTag/Public/AvaTagHandle.h
@@ -2,18 +2,19 @@
 
 #pragma once
 
-#include "AvaTag.h"
+#include "AvaTagId.h"
 #include "Containers/ContainersFwd.h"
 #include "AvaTagHandle.generated.h"
 
 class UAvaTagCollection;
+struct FAvaTag;
 
 /**
  * Handle to a Tag in a particular Source.
  * This should be used by the UStructs/UObjects to properly reference a particular FAvaTag.
  */
 USTRUCT(BlueprintType)
-struct AVALANCHETAG_API FAvaTagHandle
+struct FAvaTagHandle
 {
 	GENERATED_BODY()
 
@@ -25,21 +26,21 @@ struct AVALANCHETAG_API FAvaTagHandle
 	{
 	}
 
-	const FAvaTag* GetTag() const;
+	AVALANCHETAG_API const FAvaTag* GetTag() const;
 
-	FString ToString() const;
+	AVALANCHETAG_API FString ToString() const;
 
-	FString ToDebugString() const;
+	AVALANCHETAG_API FString ToDebugString() const;
 
-	FName ToName() const;
+	AVALANCHETAG_API FName ToName() const;
 
-	void PostSerialize(const FArchive& Ar);
+	AVALANCHETAG_API void PostSerialize(const FArchive& Ar);
 
 	/** Returns true if the Tag Handles resolve to same valued FAvaTags, even if the Source or Tag Id is different */
-	bool MatchesTag(const FAvaTagHandle& InOther) const;
+	AVALANCHETAG_API bool MatchesTag(const FAvaTagHandle& InOther) const;
 
 	/** Returns true if the Tag Handles is the exact same as the other (Same Source and Tag Id) */
-	bool MatchesExact(const FAvaTagHandle& InOther) const;
+	AVALANCHETAG_API bool MatchesExact(const FAvaTagHandle& InOther) const;
 
 	bool IsValid() const
 	{
diff --git a/Engine/Plugins/Experimental/Avalanche/Source/AvalancheTag/Public/AvaTagHandleContainer.h b/Engine/Plugins/Experimental/Avalanche/Source/AvalancheTag/Public/AvaTagHandleContainer.h
--- a/Engine/Plugins/Experimental/Avalanche/Source/AvalancheTag/Public/AvaTagHandleContainer.h
+++ b/Engine/Plugins/Experimental/Avalanche/Source/AvalancheTag/Public/AvaTagHandleContainer.h
@@ -2,11 +2,12 @@
 
 #pragma once
 
-#include "AvaTag.h"
+#include "AvaTagId.h"
 #include "Containers/ContainersFwd.h"
 #include "AvaTagHandleContainer.generated.h"
 
 class UAvaTagCollection;
+struct FAvaTag;
 struct FAvaTagHandle;
 
 /**
diff --git a/Engine/Plugins/Experimental/Avalanche/Source/AvalancheTag/Public/AvaTagId.h b/Engine/Plugins/Experimental/Avalanche/Source/AvalancheTag/Public/AvaTagId.h
new file mode 100644
--- /dev/null
+++ b/Engine/Plugins/Experimental/Avalanche/Source/AvalancheTag/Public/AvaTagId.h
@@ -0,0 +1,65 @@
+// Copyright Epic Games, Inc. All Rights Reserved.
+
+#pragma once
+
+#include "Misc/Guid.h"
+#include "AvaTagId.generated.h"
+
+/**
+ * Struct to identify a Tag. This is used by FAvaTagHandle to reference an FAvaTag
+ */
+USTRUCT()
+struct FAvaTagId
+{
+	GENERATED_BODY()
+
+	/** Id is initialized to a zero */
+	FAvaTagId() = default;
+
+	/**
+	 * Force init where Id is initialized to a new guid.
+	 * This is used in TCppStructOps::Construct
+	 */
+	explicit FAvaTagId(EForceInit)
+		: Id(FGuid::NewGuid())
+	{
+	}
+
+	explicit FAvaTagId(const FGuid& InId)
+		: Id(InId)
+	{
+	}
+
+	bool operator==(const FAvaTagId& InOther) const
+	{
+		return Id == InOther.Id;
+	}
+
+	bool IsValid() const
+	{
+		return Id.IsValid();
+	}
+
+	FString ToString() const
+	{
+		return Id.ToString();
+	}
+
+	friend uint32 GetTypeHash(const FAvaTagId& InTagId)
+	{
+		return GetTypeHash(InTagId.Id);
+	}
+
+private:
+	UPROPERTY(EditAnywhere, Category="Tag", meta=(IgnoreForMemberInitializationTest))
+	FGuid Id;
+};
+
+template<>
+struct TStructOpsTypeTraits<FAvaTagId> : TStructOpsTypeTraitsBase2<FAvaTagId>
+{
+	enum
+	{
+		WithNoInitConstructor = true,
+	};
+};
diff --git a/Engine/Plugins/Experimental/Avalanche/Source/AvalancheTag/Public/AvaTagSoftHandle.h b/Engine/Plugins/Experimental/Avalanche/Source/AvalancheTag/Public/AvaTagSoftHandle.h
--- a/Engine/Plugins/Experimental/Avalanche/Source/AvalancheTag/Public/AvaTagSoftHandle.h
+++ b/Engine/Plugins/Experimental/Avalanche/Source/AvalancheTag/Public/AvaTagSoftHandle.h
@@ -2,7 +2,7 @@
 
 #pragma once
 
-#include "AvaTag.h"
+#include "AvaTagId.h"
 #include "UObject/SoftObjectPtr.h"
 #include "AvaTagSoftHandle.generated.h"
 
