MotionDesign : Modifiers
- TranslucentPriorityModifier sets same priority for all actor components when using outliner mode

diff --git a/Engine/Plugins/Experimental/Avalanche/Source/AvalancheModifiers/Private/Modifiers/AvaTranslucentPriorityModifier.cpp b/Engine/Plugins/Experimental/Avalanche/Source/AvalancheModifiers/Private/Modifiers/AvaTranslucentPriorityModifier.cpp
--- a/Engine/Plugins/Experimental/Avalanche/Source/AvalancheModifiers/Private/Modifiers/AvaTranslucentPriorityModifier.cpp
+++ b/Engine/Plugins/Experimental/Avalanche/Source/AvalancheModifiers/Private/Modifiers/AvaTranslucentPriorityModifier.cpp
@@ -9,6 +9,7 @@
 #include "EngineUtils.h"
 #include "Engine/RendererSettings.h"
 #include "Framework/AvaGameInstance.h"
+#include "GameFramework/Actor.h"
 #include "Modifiers/ActorModifierCoreBase.h"
 #include "Shared/AvaTranslucentPriorityModifierShared.h"
 
@@ -195,34 +196,46 @@ void UAvaTranslucentPriorityModifier::Apply()
 		int32 TranslucentSortPriority = GlobalOffset;
 
 		// Increment sort priority for each component that this modifier handles
+		AActor* PreviousComponentOwner = nullptr;
 		for (const FAvaTranslucentPriorityModifierComponentState& SortedComponentState : CachedSortedComponentStates)
 		{
 			UPrimitiveComponent* Component = SortedComponentState.PrimitiveComponentWeak.Get();
 			UAvaTranslucentPriorityModifier* ComponentModifier = SortedComponentState.ModifierWeak.Get();
 
-			if (!Component || !ComponentModifier)
+			if (!Component
+				|| !Component->GetOwner()
+				|| !ComponentModifier)
 			{
 				continue;
 			}
 
-			if (Component->TranslucencySortPriority != TranslucentSortPriority)
+			// This modifier handles this component
+			if (ComponentModifier == this)
 			{
-				// This modifier handles this component
-				if (ComponentModifier == this)
-				{
-					LastSortPriorities.Add(Component, TranslucentSortPriority);
-					Component->SetTranslucentSortPriority(TranslucentSortPriority);
-				}
-				// Another modifier handles this component
-				else
-				{
-					// Cache to avoid doing the same query for the same result
-					ComponentModifier->CachedSortedComponentStates = CachedSortedComponentStates;
-					ComponentModifier->MarkModifierDirty();
-				}
+				LastSortPriorities.Add(Component, TranslucentSortPriority);
+				Component->SetTranslucentSortPriority(TranslucentSortPriority);
+			}
+			// Another modifier handles this component
+			else if (Component->TranslucencySortPriority != TranslucentSortPriority)
+			{
+				// Cache to avoid doing the same query for the same result
+				ComponentModifier->CachedSortedComponentStates = CachedSortedComponentStates;
+				ComponentModifier->MarkModifierDirty();
+			}
+
+			if (!PreviousComponentOwner)
+			{
+				PreviousComponentOwner = Component->GetOwner();
+			}
+
+			// Only increase when new actor detected or mode is camera
+			if (Mode == EAvaTranslucentPriorityModifierMode::AutoCameraDistance
+				|| PreviousComponentOwner != Component->GetOwner())
+			{
+				TranslucentSortPriority += GlobalStep;
 			}
 
-			TranslucentSortPriority += GlobalStep;
+			PreviousComponentOwner = Component->GetOwner();
 		}
 	}
 
