Non-unity error fix

diff --git a/Engine/Plugins/Experimental/Avalanche/Source/AvalancheSequence/Private/AvaSequencePlaybackActor.cpp b/Engine/Plugins/Experimental/Avalanche/Source/AvalancheSequence/Private/AvaSequencePlaybackActor.cpp
--- a/Engine/Plugins/Experimental/Avalanche/Source/AvalancheSequence/Private/AvaSequencePlaybackActor.cpp
+++ b/Engine/Plugins/Experimental/Avalanche/Source/AvalancheSequence/Private/AvaSequencePlaybackActor.cpp
@@ -3,6 +3,7 @@
 #include "AvaSequencePlaybackActor.h"
 #include "AvaSequencePlayer.h"
 #include "AvaSequenceSubsystem.h"
+#include "AvaTag.h"
 #include "Engine/World.h"
 #include "IAvaSequenceProvider.h"
 #include "MovieSceneSequenceTickManager.h"
