Motion Design
- Fix for tone mapper setting popup not saving value on relaunch.

diff --git a/Engine/Plugins/Experimental/Avalanche/Source/AvalancheViewport/Private/Visualizers/AvaViewportPostProcessVisualizer.cpp b/Engine/Plugins/Experimental/Avalanche/Source/AvalancheViewport/Private/Visualizers/AvaViewportPostProcessVisualizer.cpp
--- a/Engine/Plugins/Experimental/Avalanche/Source/AvalancheViewport/Private/Visualizers/AvaViewportPostProcessVisualizer.cpp
+++ b/Engine/Plugins/Experimental/Avalanche/Source/AvalancheViewport/Private/Visualizers/AvaViewportPostProcessVisualizer.cpp
@@ -66,10 +66,10 @@ bool FAvaViewportPostProcessVisualizer::CanActivate(bool bInSilent) const
 		return true;
 	}
 
-	URendererSettings* RenderSettings = GetMutableDefault<URendererSettings>();
-	check(RenderSettings);
+	URendererSettings* RendererSettings = GetMutableDefault<URendererSettings>();
+	check(RendererSettings);
 
-	if (RenderSettings->bEnableAlphaChannelInPostProcessing == EAlphaChannelMode::AllowThroughTonemapper)
+	if (RendererSettings->bEnableAlphaChannelInPostProcessing == EAlphaChannelMode::AllowThroughTonemapper)
 	{
 		return true;
 	}
@@ -83,10 +83,24 @@ bool FAvaViewportPostProcessVisualizer::CanActivate(bool bInSilent) const
 	switch (Response)
 	{
 		case EAppReturnType::Yes:
-			RenderSettings->bEnableAlphaChannelInPostProcessing = EAlphaChannelMode::AllowThroughTonemapper;
-			RenderSettings->SaveConfig();
+		{
+			RendererSettings->bEnableAlphaChannelInPostProcessing = EAlphaChannelMode::AllowThroughTonemapper;
+
+			if (IConsoleVariable* PropagateAlphaCVar = IConsoleManager::Get().FindConsoleVariable(TEXT("r.PostProcessing.PropagateAlpha")))
+			{
+				PropagateAlphaCVar->Set(EAlphaChannelMode::AllowThroughTonemapper);
+			}
+
+			FProperty* Property = RendererSettings->GetClass()->FindPropertyByName(GET_MEMBER_NAME_CHECKED(URendererSettings, bEnableAlphaChannelInPostProcessing));
+
+			FPropertyChangedEvent PropertyChangedEvent(Property, EPropertyChangeType::ValueSet, {RendererSettings});
+			RendererSettings->PostEditChangeProperty(PropertyChangedEvent);
+
+			RendererSettings->UpdateSinglePropertyInConfigFile(Property, RendererSettings->GetDefaultConfigFilename());
+
 			FModuleManager::GetModuleChecked<ISettingsEditorModule>("SettingsEditor").OnApplicationRestartRequired();
 			break;
+		}
 
 		case EAppReturnType::No:
 			// Continue to enable the option, but not the render setting.
