Motion Design Outliner: Allow showing transient actors (e.g. spawnables)

diff --git a/Engine/Plugins/Experimental/Avalanche/Source/AvalancheOutliner/Private/AvaOutliner.cpp b/Engine/Plugins/Experimental/Avalanche/Source/AvalancheOutliner/Private/AvaOutliner.cpp
--- a/Engine/Plugins/Experimental/Avalanche/Source/AvalancheOutliner/Private/AvaOutliner.cpp
+++ b/Engine/Plugins/Experimental/Avalanche/Source/AvalancheOutliner/Private/AvaOutliner.cpp
@@ -116,8 +116,12 @@ bool FAvaOutliner::IsActorAllowedInOutliner(const AActor* InActor) const
 		return false;
 	}
 
-	// Do not show Transient Actors in Outliner nor actors that are NOT editable
-	if (InActor->HasAnyFlags(RF_Transient) || !InActor->IsEditable() || !InActor->IsListedInSceneOutliner())
+	// Note: RF_Transient is no longer checked (to hide transient actors).
+	// It could be re-added with Outliner parameters to show/hide Transient actors (as seen in FActorMode::IsActorDisplayable)
+	// However, in SLevelEditor::CreateSceneOutliner bShowTransient is set to true.
+
+	// Do not show Actors that aren't editable or not meant to be listed in Outliner
+	if (!InActor->IsEditable() || !InActor->IsListedInSceneOutliner())
 	{
 		return false;
 	}
