Motion Design Transition: Fix issue with Autogenerated States not getting the bindings from their reference states.

diff --git a/Engine/Plugins/Experimental/Avalanche/Source/AvalancheSequencer/Private/Transition/AvaSequenceTransitionCompiler.cpp b/Engine/Plugins/Experimental/Avalanche/Source/AvalancheSequencer/Private/Transition/AvaSequenceTransitionCompiler.cpp
--- a/Engine/Plugins/Experimental/Avalanche/Source/AvalancheSequencer/Private/Transition/AvaSequenceTransitionCompiler.cpp
+++ b/Engine/Plugins/Experimental/Avalanche/Source/AvalancheSequencer/Private/Transition/AvaSequenceTransitionCompiler.cpp
@@ -57,7 +57,7 @@ namespace UE::AvaSequencer::Private
 		});
 	}
 
-	void SetInitState(UStateTreeState& InInitState, const FPlayState& InPlayState)
+	void SetInitState(FStateTreeEditorPropertyBindings* InBindings, UStateTreeState& InInitState, const FPlayState& InPlayState)
 	{
 		InInitState.Modify();
 
@@ -82,13 +82,20 @@ namespace UE::AvaSequencer::Private
 			{
 				const FAvaTransitionPlaySequenceTask::FInstanceDataType& PlayTask = PlayTaskNode->Instance.Get<FAvaTransitionPlaySequenceTask::FInstanceDataType>();
 
-				FAvaTransitionInitializeSequence::FInstanceDataType& InitTask = InInitState.AddTask<FAvaTransitionInitializeSequence>().GetInstanceData();
+				TStateTreeEditorNode<FAvaTransitionInitializeSequence>& InitTaskNode = InInitState.AddTask<FAvaTransitionInitializeSequence>();
+
+				FAvaTransitionInitializeSequence::FInstanceDataType& InitTask = InitTaskNode.GetInstanceData();
 				InitTask.InitializeTime     = PlayTask.PlaySettings.Start;
 				InitTask.PlayMode           = PlayTask.PlaySettings.PlayMode;
 				InitTask.QueryType          = PlayTask.QueryType;
 				InitTask.SequenceName       = PlayTask.SequenceName;
 				InitTask.SequenceTag        = PlayTask.SequenceTag;
 				InitTask.bPerformExactMatch = PlayTask.bPerformExactMatch;
+
+				if (InBindings)
+				{
+					InBindings->CopyBindings(PlayTaskNode->ID, InitTaskNode.ID);
+				}
 			}
 		}
 	}
@@ -235,6 +242,8 @@ void FAvaSequenceTransitionCompiler::CompileStatesRecursive(UAvaTransitionTreeEd
 		return;
 	}
 
+	FStateTreeEditorPropertyBindings* const Bindings = InEditorData.GetPropertyEditorBindings();
+
 	// Find existing Initialize Nodes and update them
 	for (int32 StateIndex = 0; StateIndex < WaitStateIndex; ++StateIndex)
 	{
@@ -246,7 +255,7 @@ void FAvaSequenceTransitionCompiler::CompileStatesRecursive(UAvaTransitionTreeEd
 			continue;
 		}
 
-		Private::SetInitState(State, PlayStates.Pop());
+		Private::SetInitState(Bindings, State, PlayStates.Pop());
 
 		// Break as there are no more Play States left. However, there could still be remaining previously generated states that are now unhandled and need to be removed
 		if (PlayStates.IsEmpty())
@@ -260,7 +269,7 @@ void FAvaSequenceTransitionCompiler::CompileStatesRecursive(UAvaTransitionTreeEd
 	{
 		UStateTreeState& NewState = InEditorData.CreateState(*WaitState, /*bAfter*/false);
 		InEditorData.SetAutogeneratedState(NewState.ID, AutogeneratedStateContext);
-		Private::SetInitState(NewState, PlayStates.Pop());
+		Private::SetInitState(Bindings, NewState, PlayStates.Pop());
 	}
 
 	// todo: at the moment, refresh is called even if all the states are up-to-date
