Motion Design Modifiers: Fixed Global Opacity Modifier not working with Material Designer

diff --git a/Engine/Plugins/Experimental/Avalanche/Source/AvalancheModifiers/Private/Modifiers/AvaGlobalOpacityModifier.cpp b/Engine/Plugins/Experimental/Avalanche/Source/AvalancheModifiers/Private/Modifiers/AvaGlobalOpacityModifier.cpp
--- a/Engine/Plugins/Experimental/Avalanche/Source/AvalancheModifiers/Private/Modifiers/AvaGlobalOpacityModifier.cpp
+++ b/Engine/Plugins/Experimental/Avalanche/Source/AvalancheModifiers/Private/Modifiers/AvaGlobalOpacityModifier.cpp
@@ -15,7 +15,7 @@ UAvaGlobalOpacityModifier::UAvaGlobalOpacityModifier()
 	bShowMaterialParameters = false;
 #endif
 	MaterialClass = UDynamicMaterialInstance::StaticClass();
-	MaterialParameters.ScalarParameters.Add(UAvaGlobalOpacityModifier::MaterialDesignerGlobalOpacityValueName, GlobalOpacity);
+	MaterialParameters.ScalarParameters.Add(UDynamicMaterialModel::GlobalOpacityParameterName, GlobalOpacity);
 }
 
 #if WITH_EDITOR
@@ -49,7 +49,7 @@ void UAvaGlobalOpacityModifier::OnGlobalOpacityChanged()
 {
 	GlobalOpacity = FMath::Clamp<float>(GlobalOpacity, UE_SMALL_NUMBER * 2, 1.f);
 
-	float& GlobalOpacityRef = MaterialParameters.ScalarParameters.FindChecked(UAvaGlobalOpacityModifier::MaterialDesignerGlobalOpacityValueName);
+	float& GlobalOpacityRef = MaterialParameters.ScalarParameters.FindOrAdd(UDynamicMaterialModel::GlobalOpacityParameterName);
 	GlobalOpacityRef = GlobalOpacity;
 	OnMaterialParametersChanged();
 }
@@ -63,7 +63,7 @@ void UAvaGlobalOpacityModifier::OnActorMaterialAdded(UMaterialInstanceDynamic* I
 	{
 		if (const UDynamicMaterialModel* Model = MDI->GetMaterialModel())
 		{
-			if (UDMMaterialValue* GlobalOpacityValue = Model->GetGlobalParameterValue(UDynamicMaterialModel::GlobalOpacityValueName))
+			if (UDMMaterialValue* GlobalOpacityValue = Model->GetGlobalParameterValue(UDynamicMaterialModel::GlobalOpacityParameterName))
 			{
 				GlobalOpacityValue->GetOnUpdate().RemoveAll(this);
 				GlobalOpacityValue->GetOnUpdate().AddUObject(this, &UAvaGlobalOpacityModifier::OnDynamicMaterialValueChanged);
@@ -82,7 +82,7 @@ void UAvaGlobalOpacityModifier::OnActorMaterialRemoved(UMaterialInstanceDynamic*
 	{
 		if (const UDynamicMaterialModel* Model = MDI->GetMaterialModel())
 		{
-			if (UDMMaterialValue* GlobalOpacityValue = Model->GetGlobalParameterValue(UDynamicMaterialModel::GlobalOpacityValueName))
+			if (UDMMaterialValue* GlobalOpacityValue = Model->GetGlobalParameterValue(UDynamicMaterialModel::GlobalOpacityParameterName))
 			{
 				GlobalOpacityValue->GetOnUpdate().RemoveAll(this);
 			}
diff --git a/Engine/Plugins/Experimental/Avalanche/Source/AvalancheModifiers/Public/Modifiers/AvaGlobalOpacityModifier.h b/Engine/Plugins/Experimental/Avalanche/Source/AvalancheModifiers/Public/Modifiers/AvaGlobalOpacityModifier.h
--- a/Engine/Plugins/Experimental/Avalanche/Source/AvalancheModifiers/Public/Modifiers/AvaGlobalOpacityModifier.h
+++ b/Engine/Plugins/Experimental/Avalanche/Source/AvalancheModifiers/Public/Modifiers/AvaGlobalOpacityModifier.h
@@ -14,8 +14,6 @@ class UAvaGlobalOpacityModifier : public UAvaMaterialParameterModifier
 {
 	GENERATED_BODY()
 
-	static inline const FName MaterialDesignerGlobalOpacityValueName = FName("VALUE_GlobalOpacity");
-
 public:
 	UAvaGlobalOpacityModifier();
 
