Motion Design Tags: Added Tag Aliases

diff --git a/Engine/Plugins/Experimental/Avalanche/Source/AvalancheAttribute/Private/Tags/AvaTagAttribute.cpp b/Engine/Plugins/Experimental/Avalanche/Source/AvalancheAttribute/Private/Tags/AvaTagAttribute.cpp
--- a/Engine/Plugins/Experimental/Avalanche/Source/AvalancheAttribute/Private/Tags/AvaTagAttribute.cpp
+++ b/Engine/Plugins/Experimental/Avalanche/Source/AvalancheAttribute/Private/Tags/AvaTagAttribute.cpp
@@ -32,7 +32,7 @@ bool UAvaTagAttribute::ClearTagHandle(const FAvaTagHandle& InTagHandle)
 
 bool UAvaTagAttribute::ContainsTag(const FAvaTagHandle& InTagHandle) const
 {
-	return Tag.MatchesTag(InTagHandle);
+	return Tag.Overlaps(InTagHandle);
 }
 
 bool UAvaTagAttribute::HasValidTagHandle() const
diff --git a/Engine/Plugins/Experimental/Avalanche/Source/AvalancheMediaEditor/Private/Rundown/ChannelLayerStatus/SAvaRundownChannelLayerStatusList.cpp b/Engine/Plugins/Experimental/Avalanche/Source/AvalancheMediaEditor/Private/Rundown/ChannelLayerStatus/SAvaRundownChannelLayerStatusList.cpp
--- a/Engine/Plugins/Experimental/Avalanche/Source/AvalancheMediaEditor/Private/Rundown/ChannelLayerStatus/SAvaRundownChannelLayerStatusList.cpp
+++ b/Engine/Plugins/Experimental/Avalanche/Source/AvalancheMediaEditor/Private/Rundown/ChannelLayerStatus/SAvaRundownChannelLayerStatusList.cpp
@@ -113,9 +113,8 @@ void SAvaRundownChannelLayerStatusList::RefreshList()
 				continue;
 			}
 
-			const FAvaTag* LayerTag = InstancePlayer->TransitionLayer.GetTag();
-
-			if (!LayerTag)
+			const FName LayerName = InstancePlayer->TransitionLayer.ToName();
+			if (LayerName.IsNone())
 			{
 				continue;
 			}
@@ -127,8 +126,7 @@ void SAvaRundownChannelLayerStatusList::RefreshList()
 				continue;
 			}
 
-			const FName& LayerName = LayerTag->TagName;
-			const FText LayerDescription = FText::FromName(LayerTag->TagName);
+			const FText LayerDescription = FText::FromName(LayerName);
 			FLinearColor ComboPageColor = FLinearColor::Transparent;
 
 			// If we're not a combo page, just add our id to the layer request
diff --git a/Engine/Plugins/Experimental/Avalanche/Source/AvalancheSequence/Private/AvaSequence.cpp b/Engine/Plugins/Experimental/Avalanche/Source/AvalancheSequence/Private/AvaSequence.cpp
--- a/Engine/Plugins/Experimental/Avalanche/Source/AvalancheSequence/Private/AvaSequence.cpp
+++ b/Engine/Plugins/Experimental/Avalanche/Source/AvalancheSequence/Private/AvaSequence.cpp
@@ -107,10 +107,9 @@ void UAvaSequence::SetLabel(FName InLabel)
 	Label = InLabel;
 }
 
-FAvaTag UAvaSequence::GetSequenceTag() const
+FAvaTagHandle UAvaSequence::GetSequenceTag() const
 {
-	const FAvaTag* ResolvedTag = Tag.GetTag();
-	return ResolvedTag ? *ResolvedTag : FAvaTag();
+	return Tag;
 }
 
 void UAvaSequence::SetSequenceTag(const FAvaTagHandle& InSequenceTag)
diff --git a/Engine/Plugins/Experimental/Avalanche/Source/AvalancheSequence/Private/AvaSequencePlaybackActor.cpp b/Engine/Plugins/Experimental/Avalanche/Source/AvalancheSequence/Private/AvaSequencePlaybackActor.cpp
--- a/Engine/Plugins/Experimental/Avalanche/Source/AvalancheSequence/Private/AvaSequencePlaybackActor.cpp
+++ b/Engine/Plugins/Experimental/Avalanche/Source/AvalancheSequence/Private/AvaSequencePlaybackActor.cpp
@@ -110,7 +110,7 @@ UAvaSequencePlayer* AAvaSequencePlaybackActor::PreviewFrame(UAvaSequence* InSequ
 	return PlaySequence(InSequence, PlaySettings);
 }
 
-UAvaSequencePlayer* AAvaSequencePlaybackActor::PlaySequenceBySoftReference(TSoftObjectPtr<UAvaSequence> InSequence, FAvaSequencePlayParams InPlaySettings)
+UAvaSequencePlayer* AAvaSequencePlaybackActor::PlaySequenceBySoftReference(TSoftObjectPtr<UAvaSequence> InSequence, const FAvaSequencePlayParams& InPlaySettings)
 {
 	if (UAvaSequence* const ResolvedSequence = InSequence.Get())
 	{
@@ -124,12 +124,12 @@ UAvaSequencePlayer* AAvaSequencePlaybackActor::PlaySequenceBySoftReference(TSoft
 	return nullptr;
 }
 
-TArray<UAvaSequencePlayer*> AAvaSequencePlaybackActor::PlaySequencesByLabel(FName InSequenceLabel, FAvaSequencePlayParams InPlaySettings)
+TArray<UAvaSequencePlayer*> AAvaSequencePlaybackActor::PlaySequencesByLabel(FName InSequenceLabel, const FAvaSequencePlayParams& InPlaySettings)
 {
-	return PlaySequencesByLabels({ InSequenceLabel }, MoveTemp(InPlaySettings));
+	return PlaySequencesByLabels({ InSequenceLabel }, InPlaySettings);
 }
 
-TArray<UAvaSequencePlayer*> AAvaSequencePlaybackActor::PlaySequencesBySoftReference(const TArray<TSoftObjectPtr<UAvaSequence>>& InSequences, FAvaSequencePlayParams InPlaySettings)
+TArray<UAvaSequencePlayer*> AAvaSequencePlaybackActor::PlaySequencesBySoftReference(const TArray<TSoftObjectPtr<UAvaSequence>>& InSequences, const FAvaSequencePlayParams& InPlaySettings)
 {
 	TArray<UAvaSequencePlayer*> Players;
 	Players.Reserve(InSequences.Num());
@@ -142,14 +142,14 @@ TArray<UAvaSequencePlayer*> AAvaSequencePlaybackActor::PlaySequencesBySoftRefere
 	return Players;
 }
 
-TArray<UAvaSequencePlayer*> AAvaSequencePlaybackActor::PlaySequencesByTag(const FAvaTag& InTag, bool bInExactMatch, FAvaSequencePlayParams InPlaySettings)
+TArray<UAvaSequencePlayer*> AAvaSequencePlaybackActor::PlaySequencesByTag(const FAvaTagHandle& InTagHandle, bool bInExactMatch, const FAvaSequencePlayParams& InPlaySettings)
 {
-	if (!InTag.IsValid())
+	if (!InTagHandle.IsValid())
 	{
 		return TArray<UAvaSequencePlayer*>();
 	}
 
-	const TArray<UAvaSequence*> SequencesToPlay = GetSequencesByTag(InTag, bInExactMatch);
+	const TArray<UAvaSequence*> SequencesToPlay = GetSequencesByTag(InTagHandle, bInExactMatch);
 
 	TArray<UAvaSequencePlayer*> SequencePlayers;
 	SequencePlayers.Reserve(SequencesToPlay.Num());
@@ -165,7 +165,7 @@ TArray<UAvaSequencePlayer*> AAvaSequencePlaybackActor::PlaySequencesByTag(const
 	return SequencePlayers;
 }
 
-TArray<UAvaSequencePlayer*> AAvaSequencePlaybackActor::PlaySequencesByLabels(const TArray<FName>& InSequenceLabels, FAvaSequencePlayParams InPlaySettings)
+TArray<UAvaSequencePlayer*> AAvaSequencePlaybackActor::PlaySequencesByLabels(const TArray<FName>& InSequenceLabels, const FAvaSequencePlayParams& InPlaySettings)
 {
 	if (InSequenceLabels.IsEmpty())
 	{
@@ -205,11 +205,11 @@ TArray<UAvaSequencePlayer*> AAvaSequencePlaybackActor::ContinueSequencesByLabel(
 	return ContinueSequencesByLabels({ InSequenceLabel });
 }
 
-TArray<UAvaSequencePlayer*> AAvaSequencePlaybackActor::ContinueSequencesByTag(const FAvaTag& InTag, bool bInExactMatch)
+TArray<UAvaSequencePlayer*> AAvaSequencePlaybackActor::ContinueSequencesByTag(const FAvaTagHandle& InTagHandle, bool bInExactMatch)
 {
 	TArray<UAvaSequencePlayer*> SequencePlayers;
 	
-	const TArray<UAvaSequence*> SequencesToContinue = GetSequencesByTag(InTag, bInExactMatch);
+	const TArray<UAvaSequence*> SequencesToContinue = GetSequencesByTag(InTagHandle, bInExactMatch);
 	for (UAvaSequence* Sequence : SequencesToContinue)
 	{
 		if (UAvaSequencePlayer* Player = ContinueSequence(Sequence))
@@ -330,9 +330,9 @@ TArray<UAvaSequence*> AAvaSequencePlaybackActor::GetSequencesByLabel(TConstArray
 	return OutSequences;
 }
 
-TArray<UAvaSequence*> AAvaSequencePlaybackActor::GetSequencesByTag(const FAvaTag& InTag, bool bInExactMatch) const
+TArray<UAvaSequence*> AAvaSequencePlaybackActor::GetSequencesByTag(const FAvaTagHandle& InTagHandle, bool bInExactMatch) const
 {
-	if (!InTag.IsValid())
+	if (!InTagHandle.IsValid())
 	{
 		UE_LOG(LogAvaSequencePlayback, Verbose
 			, TEXT("Input Tag was empty in playback actor %s")
@@ -352,7 +352,7 @@ TArray<UAvaSequence*> AAvaSequencePlaybackActor::GetSequencesByTag(const FAvaTag
 
 	for (UAvaSequence* Sequence : AllSequences)
 	{
-		if (Sequence && Sequence->GetSequenceTag() == InTag)
+		if (Sequence && Sequence->GetSequenceTag().Overlaps(InTagHandle))
 		{
 			OutSequences.Add(Sequence);
 		}
@@ -399,9 +399,9 @@ TArray<UAvaSequencePlayer*> AAvaSequencePlaybackActor::GetSequencePlayersByLabel
 	return GetSequencePlayersByLabels({ InSequenceLabel });
 }
 
-TArray<UAvaSequencePlayer*> AAvaSequencePlaybackActor::GetSequencePlayersByTag(const FAvaTag& InTag, bool bInExactMatch) const
+TArray<UAvaSequencePlayer*> AAvaSequencePlaybackActor::GetSequencePlayersByTag(const FAvaTagHandle& InTagHandle, bool bInExactMatch) const
 {
-	const TArray<UAvaSequence*> Sequences = GetSequencesByTag(InTag, bInExactMatch);
+	const TArray<UAvaSequence*> Sequences = GetSequencesByTag(InTagHandle, bInExactMatch);
 
 	TArray<UAvaSequencePlayer*> SequencePlayers;
 	SequencePlayers.Reserve(Sequences.Num());
diff --git a/Engine/Plugins/Experimental/Avalanche/Source/AvalancheSequence/Private/AvaSequencePlaybackActor.h b/Engine/Plugins/Experimental/Avalanche/Source/AvalancheSequence/Private/AvaSequencePlaybackActor.h
--- a/Engine/Plugins/Experimental/Avalanche/Source/AvalancheSequence/Private/AvaSequencePlaybackActor.h
+++ b/Engine/Plugins/Experimental/Avalanche/Source/AvalancheSequence/Private/AvaSequencePlaybackActor.h
@@ -30,22 +30,22 @@ protected:
 	virtual void CleanupPlayers() override;
 	virtual UAvaSequencePlayer* PlaySequence(UAvaSequence* InSequence, const FAvaSequencePlayParams& InPlaySettings = FAvaSequencePlayParams()) override;
 	virtual UAvaSequencePlayer* PreviewFrame(UAvaSequence* InSequence) override;
-	virtual UAvaSequencePlayer* PlaySequenceBySoftReference(TSoftObjectPtr<UAvaSequence> InSequence, FAvaSequencePlayParams InPlaySettings) override;
-	virtual TArray<UAvaSequencePlayer*> PlaySequencesByLabel(FName InSequenceLabel, FAvaSequencePlayParams InPlaySettings) override;
-	virtual TArray<UAvaSequencePlayer*> PlaySequencesBySoftReference(const TArray<TSoftObjectPtr<UAvaSequence>>& InSequences, FAvaSequencePlayParams InPlaySettings) override;
-	virtual TArray<UAvaSequencePlayer*> PlaySequencesByLabels(const TArray<FName>& InSequenceLabels, FAvaSequencePlayParams InPlaySettings) override;
-	virtual TArray<UAvaSequencePlayer*> PlaySequencesByTag(const FAvaTag& InTag, bool bInExactMatch, FAvaSequencePlayParams InPlaySettings) override;
+	virtual UAvaSequencePlayer* PlaySequenceBySoftReference(TSoftObjectPtr<UAvaSequence> InSequence, const FAvaSequencePlayParams& InPlaySettings) override;
+	virtual TArray<UAvaSequencePlayer*> PlaySequencesByLabel(FName InSequenceLabel, const FAvaSequencePlayParams& InPlaySettings) override;
+	virtual TArray<UAvaSequencePlayer*> PlaySequencesBySoftReference(const TArray<TSoftObjectPtr<UAvaSequence>>& InSequences, const FAvaSequencePlayParams& InPlaySettings) override;
+	virtual TArray<UAvaSequencePlayer*> PlaySequencesByLabels(const TArray<FName>& InSequenceLabels, const FAvaSequencePlayParams& InPlaySettings) override;
+	virtual TArray<UAvaSequencePlayer*> PlaySequencesByTag(const FAvaTagHandle& InTagHandle, bool bInExactMatch, const FAvaSequencePlayParams& InPlaySettings) override;
 	virtual TArray<UAvaSequencePlayer*> PlayScheduledSequences() override;
 	virtual UAvaSequencePlayer* ContinueSequence(UAvaSequence* InSequence) override;
 	virtual TArray<UAvaSequencePlayer*> ContinueSequencesByLabel(FName InSequenceLabel) override;
 	virtual TArray<UAvaSequencePlayer*> ContinueSequencesByLabels(const TArray<FName>& InSequenceLabels) override;
-	virtual TArray<UAvaSequencePlayer*> ContinueSequencesByTag(const FAvaTag& InTag, bool bInExactMatch) override;
+	virtual TArray<UAvaSequencePlayer*> ContinueSequencesByTag(const FAvaTagHandle& InTagHandle, bool bInExactMatch) override;
 	virtual void StopSequence(UAvaSequence* InSequence) override;
 	virtual UObject* GetPlaybackContext() const override;
 	virtual UAvaSequencePlayer* GetSequencePlayer(const UAvaSequence* InSequence) const override;
 	virtual TArray<UAvaSequencePlayer*> GetSequencePlayersByLabel(FName InSequenceLabel) const override;
 	virtual TArray<UAvaSequencePlayer*> GetSequencePlayersByLabels(const TArray<FName>& InSequenceLabels) const override;
-	virtual TArray<UAvaSequencePlayer*> GetSequencePlayersByTag(const FAvaTag& InTag, bool bInExactMatch) const override;
+	virtual TArray<UAvaSequencePlayer*> GetSequencePlayersByTag(const FAvaTagHandle& InTagHandle, bool bInExactMatch) const override;
 	virtual TArray<UAvaSequencePlayer*> GetAllSequencePlayers() const override;
 	virtual bool HasActiveSequencePlayers() const override;
 	virtual void UpdateCameraCut(const UE::MovieScene::FOnCameraCutUpdatedParams& InCameraCutParams) override;
@@ -76,7 +76,7 @@ private:
 
 	TArray<UAvaSequence*> GetSequencesByLabel(TConstArrayView<FName> InSequenceLabels) const;
 
-	TArray<UAvaSequence*> GetSequencesByTag(const FAvaTag& InTag, bool bInExactMatch) const;
+	TArray<UAvaSequence*> GetSequencesByTag(const FAvaTagHandle& InTagHandle, bool bInExactMatch) const;
 
 	UFUNCTION(BlueprintCallable, CallInEditor, DisplayName="Play Scheduled Sequences", Category="Scheduled Playback")
 	void BP_PlayScheduledSequences();
diff --git a/Engine/Plugins/Experimental/Avalanche/Source/AvalancheSequence/Private/Transition/AvaTransitionContinueSequenceTask.cpp b/Engine/Plugins/Experimental/Avalanche/Source/AvalancheSequence/Private/Transition/AvaTransitionContinueSequenceTask.cpp
--- a/Engine/Plugins/Experimental/Avalanche/Source/AvalancheSequence/Private/Transition/AvaTransitionContinueSequenceTask.cpp
+++ b/Engine/Plugins/Experimental/Avalanche/Source/AvalancheSequence/Private/Transition/AvaTransitionContinueSequenceTask.cpp
@@ -20,11 +20,7 @@ TArray<UAvaSequencePlayer*> FAvaTransitionContinueSequenceTask::ExecuteSequenceT
 		return PlaybackObject->ContinueSequencesByLabel(InstanceData.SequenceName);
 
 	case EAvaTransitionSequenceQueryType::Tag:
-		if (const FAvaTag* Tag = InstanceData.SequenceTag.GetTag())
-		{
-			return PlaybackObject->ContinueSequencesByTag(*Tag, InstanceData.bPerformExactMatch);
-		}
-		return TArray<UAvaSequencePlayer*>();
+		return PlaybackObject->ContinueSequencesByTag(InstanceData.SequenceTag, InstanceData.bPerformExactMatch);
 	}
 
 	checkNoEntry();
diff --git a/Engine/Plugins/Experimental/Avalanche/Source/AvalancheSequence/Private/Transition/AvaTransitionInitializeSequence.cpp b/Engine/Plugins/Experimental/Avalanche/Source/AvalancheSequence/Private/Transition/AvaTransitionInitializeSequence.cpp
--- a/Engine/Plugins/Experimental/Avalanche/Source/AvalancheSequence/Private/Transition/AvaTransitionInitializeSequence.cpp
+++ b/Engine/Plugins/Experimental/Avalanche/Source/AvalancheSequence/Private/Transition/AvaTransitionInitializeSequence.cpp
@@ -51,11 +51,7 @@ TArray<UAvaSequencePlayer*> FAvaTransitionInitializeSequence::ExecuteSequenceTas
 		return PlaybackObject->PlaySequencesByLabel(InstanceData.SequenceName, PlaySettings);
 
 	case EAvaTransitionSequenceQueryType::Tag:
-		if (const FAvaTag* Tag = InstanceData.SequenceTag.GetTag())
-		{
-			return PlaybackObject->PlaySequencesByTag(*Tag, InstanceData.bPerformExactMatch, PlaySettings);
-		}
-		return TArray<UAvaSequencePlayer*>();
+		return PlaybackObject->PlaySequencesByTag(InstanceData.SequenceTag, InstanceData.bPerformExactMatch, PlaySettings);
 	}
 
 	checkNoEntry();
diff --git a/Engine/Plugins/Experimental/Avalanche/Source/AvalancheSequence/Private/Transition/AvaTransitionPlaySequenceTask.cpp b/Engine/Plugins/Experimental/Avalanche/Source/AvalancheSequence/Private/Transition/AvaTransitionPlaySequenceTask.cpp
--- a/Engine/Plugins/Experimental/Avalanche/Source/AvalancheSequence/Private/Transition/AvaTransitionPlaySequenceTask.cpp
+++ b/Engine/Plugins/Experimental/Avalanche/Source/AvalancheSequence/Private/Transition/AvaTransitionPlaySequenceTask.cpp
@@ -57,11 +57,7 @@ TArray<UAvaSequencePlayer*> FAvaTransitionPlaySequenceTask::ExecuteSequenceTask(
 		return PlaybackObject->PlaySequencesByLabel(InstanceData.SequenceName, InstanceData.PlaySettings);
 
 	case EAvaTransitionSequenceQueryType::Tag:
-		if (const FAvaTag* Tag = InstanceData.SequenceTag.GetTag())
-		{
-			return PlaybackObject->PlaySequencesByTag(*Tag, InstanceData.bPerformExactMatch, InstanceData.PlaySettings);
-		}
-		return TArray<UAvaSequencePlayer*>();
+		return PlaybackObject->PlaySequencesByTag(InstanceData.SequenceTag, InstanceData.bPerformExactMatch, InstanceData.PlaySettings);
 	}
 
 	checkNoEntry();
diff --git a/Engine/Plugins/Experimental/Avalanche/Source/AvalancheSequence/Private/Transition/AvaTransitionWaitForSequenceTask.cpp b/Engine/Plugins/Experimental/Avalanche/Source/AvalancheSequence/Private/Transition/AvaTransitionWaitForSequenceTask.cpp
--- a/Engine/Plugins/Experimental/Avalanche/Source/AvalancheSequence/Private/Transition/AvaTransitionWaitForSequenceTask.cpp
+++ b/Engine/Plugins/Experimental/Avalanche/Source/AvalancheSequence/Private/Transition/AvaTransitionWaitForSequenceTask.cpp
@@ -21,11 +21,7 @@ TArray<UAvaSequencePlayer*> FAvaTransitionWaitForSequenceTask::ExecuteSequenceTa
 		return PlaybackObject->GetSequencePlayersByLabel(InstanceData.SequenceName);
 
 	case EAvaTransitionSequenceQueryType::Tag:
-		if (const FAvaTag* Tag = InstanceData.SequenceTag.GetTag())
-		{
-			return PlaybackObject->GetSequencePlayersByTag(*Tag, InstanceData.bPerformExactMatch);
-		}
-		return TArray<UAvaSequencePlayer*>();
+		return PlaybackObject->GetSequencePlayersByTag(InstanceData.SequenceTag, InstanceData.bPerformExactMatch);
 	}
 
 	checkNoEntry();
diff --git a/Engine/Plugins/Experimental/Avalanche/Source/AvalancheSequence/Public/AvaSequence.h b/Engine/Plugins/Experimental/Avalanche/Source/AvalancheSequence/Public/AvaSequence.h
--- a/Engine/Plugins/Experimental/Avalanche/Source/AvalancheSequence/Public/AvaSequence.h
+++ b/Engine/Plugins/Experimental/Avalanche/Source/AvalancheSequence/Public/AvaSequence.h
@@ -53,7 +53,7 @@ public:
 	void SetLabel(FName InLabel);
 
 	UFUNCTION(BlueprintPure, Category = "Motion Design|Sequence")
-	FAvaTag GetSequenceTag() const;
+	FAvaTagHandle GetSequenceTag() const;
 
 	void SetSequenceTag(const FAvaTagHandle& InSequenceTag);
 
diff --git a/Engine/Plugins/Experimental/Avalanche/Source/AvalancheSequence/Public/AvaSequencePlaybackObject.h b/Engine/Plugins/Experimental/Avalanche/Source/AvalancheSequence/Public/AvaSequencePlaybackObject.h
--- a/Engine/Plugins/Experimental/Avalanche/Source/AvalancheSequence/Public/AvaSequencePlaybackObject.h
+++ b/Engine/Plugins/Experimental/Avalanche/Source/AvalancheSequence/Public/AvaSequencePlaybackObject.h
@@ -12,7 +12,7 @@ class UAvaSequence;
 class UAvaSequencePlayer;
 class ULevel;
 class UObject;
-struct FAvaTag;
+struct FAvaTagHandle;
 
 namespace UE::MovieScene
 {
@@ -57,7 +57,7 @@ public:
 	 * @return the player instantiated for the Sequence, or null if Sequence was not valid for playback
 	 */
 	UFUNCTION(BlueprintCallable, DisplayName = "Play Sequence (by Soft Reference)", Category = "Playback")
-	virtual UAvaSequencePlayer* PlaySequenceBySoftReference(TSoftObjectPtr<UAvaSequence> InSequence, FAvaSequencePlayParams InPlaySettings) = 0;
+	virtual UAvaSequencePlayer* PlaySequenceBySoftReference(TSoftObjectPtr<UAvaSequence> InSequence, const FAvaSequencePlayParams& InPlaySettings) = 0;
 
 	/**
 	 * Plays all the sequences that have the provided label
@@ -66,7 +66,7 @@ public:
 	 * @return an array of the Sequence Players with possible invalid/null entries kept so that each Player matches in Index with the input Sequence it is playing
 	 */
 	UFUNCTION(BlueprintCallable, DisplayName = "Play Sequences (by Label)", Category = "Playback")
-	virtual TArray<UAvaSequencePlayer*> PlaySequencesByLabel(FName InSequenceLabel, FAvaSequencePlayParams InPlaySettings) = 0;
+	virtual TArray<UAvaSequencePlayer*> PlaySequencesByLabel(FName InSequenceLabel, const FAvaSequencePlayParams& InPlaySettings) = 0;
 
 	/**
 	 * Plays multiple Sequences by their Soft Reference
@@ -75,7 +75,7 @@ public:
 	 * @return an array of the Sequence Players with possible invalid/null entries kept so that each Player matches in Index with the input Sequence it is playing
 	 */
 	UFUNCTION(BlueprintCallable, DisplayName = "Play Sequences (by Soft Reference)", Category = "Playback")
-	virtual TArray<UAvaSequencePlayer*> PlaySequencesBySoftReference(const TArray<TSoftObjectPtr<UAvaSequence>>& InSequences, FAvaSequencePlayParams InPlaySettings) = 0;
+	virtual TArray<UAvaSequencePlayer*> PlaySequencesBySoftReference(const TArray<TSoftObjectPtr<UAvaSequence>>& InSequences, const FAvaSequencePlayParams& InPlaySettings) = 0;
 
 	/**
 	 * Plays multiple Sequences by an array of sequence labels
@@ -84,17 +84,17 @@ public:
 	 * @return an array of the Sequence Players with possible invalid/null entries kept so that each Player matches in Index with the input Sequence it is playing
 	 */
 	UFUNCTION(BlueprintCallable, DisplayName = "Play Sequences (by Labels)", Category = "Playback")
-	virtual TArray<UAvaSequencePlayer*> PlaySequencesByLabels(const TArray<FName>& InSequenceLabels, FAvaSequencePlayParams InPlaySettings) = 0;
+	virtual TArray<UAvaSequencePlayer*> PlaySequencesByLabels(const TArray<FName>& InSequenceLabels, const FAvaSequencePlayParams& InPlaySettings) = 0;
 
 	/**
 	 * Plays all the Sequences that match the given gameplay tag(s)
-	 * @param InTag the tag to match
+	 * @param InTagHandle the tag to match
 	 * @param bInExactMatch whether to only consider sequences that have the tag exactly
 	 * @param InPlaySettings the play settings to use for playback
 	 * @return an array of the Sequence Players with only valid entries kept
 	 */
 	UFUNCTION(BlueprintCallable, DisplayName = "Play Sequences (by Tag)", Category = "Playback")
-	virtual TArray<UAvaSequencePlayer*> PlaySequencesByTag(const FAvaTag& InTag, bool bInExactMatch, FAvaSequencePlayParams InPlaySettings) = 0;
+	virtual TArray<UAvaSequencePlayer*> PlaySequencesByTag(const FAvaTagHandle& InTagHandle, bool bInExactMatch, const FAvaSequencePlayParams& InPlaySettings) = 0;
 
 	/**
 	 * Plays the Scheduled Sequences with the Scheduled Play Settings
@@ -128,12 +128,12 @@ public:
 
 	/**
 	 * Triggers Continues in all the sequences matching the provided tag
-	 * @param InTag the tag to match
+	 * @param InTagHandle the tag to match
 	 * @param bInExactMatch whether to only consider sequences that have the tag exactly
 	 * @return the array of the Sequence Players with only valid entries that fired the continue
 	 */
 	UFUNCTION(BlueprintCallable, DisplayName = "Continue Sequences (by Tag)", Category = "Playback")
-	virtual TArray<UAvaSequencePlayer*> ContinueSequencesByTag(const FAvaTag& InTag, bool bInExactMatch) = 0;
+	virtual TArray<UAvaSequencePlayer*> ContinueSequencesByTag(const FAvaTagHandle& InTagHandle, bool bInExactMatch) = 0;
 
 	virtual void StopSequence(UAvaSequence* InSequence) = 0;
 
@@ -152,7 +152,7 @@ public:
 
 	virtual TArray<UAvaSequencePlayer*> GetSequencePlayersByLabels(const TArray<FName>& InSequenceLabels) const = 0;
 
-	virtual TArray<UAvaSequencePlayer*> GetSequencePlayersByTag(const FAvaTag& InTag, bool bInExactMatch) const = 0;
+	virtual TArray<UAvaSequencePlayer*> GetSequencePlayersByTag(const FAvaTagHandle& InTagHandle, bool bInExactMatch) const = 0;
 
 	/** Retrieves all Active Sequence Players */
 	UFUNCTION(BlueprintCallable, DisplayName = "Get Active Sequence Players", Category = "Playback")
diff --git a/Engine/Plugins/Experimental/Avalanche/Source/AvalancheTag/Private/AvaTagAlias.cpp b/Engine/Plugins/Experimental/Avalanche/Source/AvalancheTag/Private/AvaTagAlias.cpp
new file mode 100644
--- /dev/null
+++ b/Engine/Plugins/Experimental/Avalanche/Source/AvalancheTag/Private/AvaTagAlias.cpp
@@ -0,0 +1,41 @@
+// Copyright Epic Games, Inc. All Rights Reserved.
+
+#include "AvaTagAlias.h"
+#include "AvaTag.h"
+#include "AvaTagCollection.h"
+#include "AvaTagList.h"
+
+#if WITH_EDITOR
+void FAvaTagAlias::SetOwner(UAvaTagCollection* InOwner)
+{
+	OwnerWeak = InOwner;
+}
+
+UAvaTagCollection* FAvaTagAlias::GetOwner() const
+{
+	return OwnerWeak.Get();
+}
+
+FString FAvaTagAlias::GetTagsAsString() const
+{
+	UAvaTagCollection* Owner = GetOwner();
+	if (!Owner)
+	{
+		return FString();
+	}
+
+	FString OutString;
+	for (const FAvaTagId& TagId : TagIds)
+	{
+		for (const FAvaTag* Tag : Owner->GetTags(TagId))
+		{
+			if (!OutString.IsEmpty())
+			{
+				OutString.Append(TEXT(", "));
+			}
+			OutString.Append(Tag->ToString());
+		}
+	}
+	return OutString;
+}
+#endif
diff --git a/Engine/Plugins/Experimental/Avalanche/Source/AvalancheTag/Private/AvaTagCollection.cpp b/Engine/Plugins/Experimental/Avalanche/Source/AvalancheTag/Private/AvaTagCollection.cpp
--- a/Engine/Plugins/Experimental/Avalanche/Source/AvalancheTag/Private/AvaTagCollection.cpp
+++ b/Engine/Plugins/Experimental/Avalanche/Source/AvalancheTag/Private/AvaTagCollection.cpp
@@ -1,20 +1,110 @@
 // Copyright Epic Games, Inc. All Rights Reserved.
 
 #include "AvaTagCollection.h"
+#include "AvaTagList.h"
 
-const FAvaTag* UAvaTagCollection::GetTag(const FAvaTagId& InTagId) const
+FAvaTagList UAvaTagCollection::GetTags(const FAvaTagId& InTagId) const
 {
-	return Tags.Find(InTagId);
+	FAvaTagList TagList;
+
+	if (const FAvaTag* Tag = Tags.Find(InTagId))
+	{
+		TagList.Tags.Add(Tag);
+	}
+
+	if (const FAvaTagAlias* TagAlias = Aliases.Find(InTagId))
+	{
+		TagList.Tags.Reserve(TagList.Tags.Num() + TagAlias->TagIds.Num());
+		for (const FAvaTagId& TagId : TagAlias->TagIds)
+		{
+			if (const FAvaTag* Tag = Tags.Find(TagId))
+			{
+				TagList.Tags.Add(Tag);
+			}
+		}
+	}
+
+	return TagList;
 }
 
-TArray<FAvaTagId> UAvaTagCollection::GetTagIds() const
+FName UAvaTagCollection::GetTagName(const FAvaTagId& InTagId) const
 {
+	if (const FAvaTag* Tag = Tags.Find(InTagId))
+	{
+		return Tag->TagName;
+	}
+
+	if (const FAvaTagAlias* TagAlias = Aliases.Find(InTagId))
+	{
+		return TagAlias->AliasName;
+	}
+
+	return NAME_None;
+}
+
+TArray<FAvaTagId> UAvaTagCollection::GetTagIds(bool bInIncludeAliases) const
+{
+	int32 TagIdCount = Tags.Num();
+	if (bInIncludeAliases)
+	{
+		TagIdCount += Aliases.Num();
+	}
+
 	TArray<FAvaTagId> TagIds;
-	Tags.GetKeys(TagIds);
+	TagIds.Reserve(TagIdCount);
+
+	for (const TPair<FAvaTagId, FAvaTag>& Pair : Tags)
+	{
+		TagIds.Add(Pair.Key);
+	}
+
+	if (bInIncludeAliases)
+	{
+		for (const TPair<FAvaTagId, FAvaTagAlias>& Pair : Aliases)
+		{
+			TagIds.Add(Pair.Key);
+		}
+	}
+
 	return TagIds;
 }
 
+void UAvaTagCollection::PostLoad()
+{
+	Super::PostLoad();
+#if WITH_EDITOR
+	UpdateAliasOwner();
+#endif
+}
+
+#if WITH_EDITOR
+void UAvaTagCollection::PostEditChangeProperty(FPropertyChangedEvent& InPropertyChangedEvent)
+{
+	Super::PostEditChangeProperty(InPropertyChangedEvent);
+
+	if (InPropertyChangedEvent.GetMemberPropertyName() == GET_MEMBER_NAME_CHECKED(UAvaTagCollection, Aliases))
+	{
+		UpdateAliasOwner();
+	}
+}
+#endif
+
 FName UAvaTagCollection::GetTagMapName()
 {
 	return GET_MEMBER_NAME_CHECKED(UAvaTagCollection, Tags);
 }
+
+FName UAvaTagCollection::GetAliasMapName()
+{
+	return GET_MEMBER_NAME_CHECKED(UAvaTagCollection, Aliases);
+}
+
+#if WITH_EDITOR
+void UAvaTagCollection::UpdateAliasOwner()
+{
+	for (TPair<FAvaTagId, FAvaTagAlias>& Alias : Aliases)
+	{
+		Alias.Value.SetOwner(this);
+	}
+}
+#endif
\ No newline at end of file
diff --git a/Engine/Plugins/Experimental/Avalanche/Source/AvalancheTag/Private/AvaTagHandle.cpp b/Engine/Plugins/Experimental/Avalanche/Source/AvalancheTag/Private/AvaTagHandle.cpp
--- a/Engine/Plugins/Experimental/Avalanche/Source/AvalancheTag/Private/AvaTagHandle.cpp
+++ b/Engine/Plugins/Experimental/Avalanche/Source/AvalancheTag/Private/AvaTagHandle.cpp
@@ -2,23 +2,20 @@
 
 #include "AvaTagHandle.h"
 #include "AvaTagCollection.h"
+#include "AvaTagList.h"
 
-const FAvaTag* FAvaTagHandle::GetTag() const
+FAvaTagList FAvaTagHandle::GetTags() const
 {
 	if (::IsValid(Source))
 	{
-		return Source->GetTag(TagId);
+		return Source->GetTags(TagId);
 	}
-	return nullptr;
+	return FAvaTagList();
 }
 
 FString FAvaTagHandle::ToString() const
 {
-	if (const FAvaTag* Tag = GetTag())
-	{
-		return Tag->ToString();
-	}
-	return FString();
+	return ToName().ToString();
 }
 
 FString FAvaTagHandle::ToDebugString() const
@@ -28,9 +25,9 @@ FString FAvaTagHandle::ToDebugString() const
 
 FName FAvaTagHandle::ToName() const
 {
-	if (const FAvaTag* Tag = GetTag())
+	if (::IsValid(Source))
 	{
-		return Tag->TagName;
+		return Source->GetTagName(TagId);
 	}
 	return NAME_None;
 }
@@ -43,17 +40,39 @@ void FAvaTagHandle::PostSerialize(const FArchive& Ar)
 	}
 }
 
-bool FAvaTagHandle::MatchesTag(const FAvaTagHandle& InOther) const
+bool FAvaTagHandle::Overlaps(const FAvaTagHandle& InOther) const
 {
-	const FAvaTag* ThisTag  = GetTag();
-	const FAvaTag* OtherTag = InOther.GetTag();
+	if (MatchesExact(InOther))
+	{
+		return true;
+	}
+
+	// No Overlap if this has no Tags
+	FAvaTagList ThisTagList = GetTags();
+	if (ThisTagList.Tags.IsEmpty())
+	{
+		return false;
+	}
+
+	// No Overlap if other is empty
+	FAvaTagList OtherTagList = InOther.GetTags();
+	if (OtherTagList.Tags.IsEmpty())
+	{
+		return false;
+	}
 
-	if (ThisTag && OtherTag)
+	for (const FAvaTag* ThisTag : ThisTagList)
 	{
-		return *ThisTag == *OtherTag;
+		for (const FAvaTag* OtherTag : OtherTagList)
+		{
+			if (*ThisTag == *OtherTag)
+			{
+				return true;
+			}
+		}
 	}
 
-	return MatchesExact(InOther);
+	return false;
 }
 
 bool FAvaTagHandle::MatchesExact(const FAvaTagHandle& InOther) const
diff --git a/Engine/Plugins/Experimental/Avalanche/Source/AvalancheTag/Private/AvaTagHandleContainer.cpp b/Engine/Plugins/Experimental/Avalanche/Source/AvalancheTag/Private/AvaTagHandleContainer.cpp
--- a/Engine/Plugins/Experimental/Avalanche/Source/AvalancheTag/Private/AvaTagHandleContainer.cpp
+++ b/Engine/Plugins/Experimental/Avalanche/Source/AvalancheTag/Private/AvaTagHandleContainer.cpp
@@ -3,6 +3,7 @@
 #include "AvaTagHandleContainer.h"
 #include "AvaTagCollection.h"
 #include "AvaTagHandle.h"
+#include "AvaTagList.h"
 
 FAvaTagHandleContainer::FAvaTagHandleContainer(const FAvaTagHandle& InTagHandle)
 	: Source(InTagHandle.Source)
@@ -22,18 +23,30 @@ bool FAvaTagHandleContainer::ContainsTag(const FAvaTagHandle& InTagHandle) const
 		return true;
 	}
 
-	const FAvaTag* OtherTag = InTagHandle.GetTag();
-	if (!OtherTag)
+	// Populate the Tag Set with resolved Tags
+	TSet<FAvaTag> OtherTagSet;
 	{
-		return false;
+		FAvaTagList OtherTagList = InTagHandle.GetTags();
+		if (OtherTagList.Tags.IsEmpty())
+		{
+			return false;
+		}
+
+		OtherTagSet.Reserve(OtherTagList.Tags.Num());
+		for (const FAvaTag* OtherTag : OtherTagList)
+		{
+			OtherTagSet.Add(*OtherTag);
+		}
 	}
 
 	for (const FAvaTagId& TagId : TagIds)
 	{
-		const FAvaTag* Tag = Source->GetTag(TagId);
-		if (Tag && *Tag == *OtherTag)
+		for (const FAvaTag* Tag : Source->GetTags(TagId))
 		{
-			return true;
+			if (OtherTagSet.Contains(*Tag))
+			{
+				return true;
+			}
 		}
 	}
 
@@ -55,7 +68,7 @@ FString FAvaTagHandleContainer::ToString() const
 	FString OutString;
 	for (const FAvaTagId& TagId : TagIds)
 	{
-		if (const FAvaTag* Tag = Source->GetTag(TagId))
+		for (const FAvaTag* Tag : Source->GetTags(TagId))
 		{
 			if (!OutString.IsEmpty())
 			{
@@ -135,7 +148,7 @@ TArray<FAvaTag> FAvaTagHandleContainer::ResolveTags() const
 
 	for (const FAvaTagId& TagId : TagIds)
 	{
-		if (const FAvaTag* Tag = Source->GetTag(TagId))
+		for (const FAvaTag* Tag : Source->GetTags(TagId))
 		{
 			Tags.Add(*Tag);
 		}
diff --git a/Engine/Plugins/Experimental/Avalanche/Source/AvalancheTag/Private/AvaTagLibrary.cpp b/Engine/Plugins/Experimental/Avalanche/Source/AvalancheTag/Private/AvaTagLibrary.cpp
--- a/Engine/Plugins/Experimental/Avalanche/Source/AvalancheTag/Private/AvaTagLibrary.cpp
+++ b/Engine/Plugins/Experimental/Avalanche/Source/AvalancheTag/Private/AvaTagLibrary.cpp
@@ -4,15 +4,22 @@
 #include "AvaTag.h"
 #include "AvaTagHandle.h"
 #include "AvaTagHandleContainer.h"
+#include "AvaTagList.h"
 #include "AvaTagSoftHandle.h"
 
-FAvaTag UAvaTagLibrary::ResolveTagHandle(const FAvaTagHandle& InTagHandle)
+TArray<FAvaTag> UAvaTagLibrary::ResolveTagHandle(const FAvaTagHandle& InTagHandle)
 {
-	if (const FAvaTag* Tag = InTagHandle.GetTag())
+	FAvaTagList TagList = InTagHandle.GetTags();
+
+	TArray<FAvaTag> Tags;
+	Tags.Reserve(TagList.Tags.Num());
+
+	for (const FAvaTag* Tag : TagList)
 	{
-		return *Tag;
+		Tags.Add(*Tag);
 	}
-	return FAvaTag();
+
+	return Tags;
 }
 
 TArray<FAvaTag> UAvaTagLibrary::ResolveTagHandles(const FAvaTagHandleContainer& InTagHandleContainer)
diff --git a/Engine/Plugins/Experimental/Avalanche/Source/AvalancheTag/Public/AvaTagAlias.h b/Engine/Plugins/Experimental/Avalanche/Source/AvalancheTag/Public/AvaTagAlias.h
new file mode 100644
--- /dev/null
+++ b/Engine/Plugins/Experimental/Avalanche/Source/AvalancheTag/Public/AvaTagAlias.h
@@ -0,0 +1,43 @@
+// Copyright Epic Games, Inc. All Rights Reserved.
+
+#pragma once
+
+#include "AvaTagId.h"
+#include "AvaTagAlias.generated.h"
+
+class UAvaTagCollection;
+
+/**
+ * An alias represents multiple other Tag Ids.
+ *
+ * Unlike Tag Containers which would need to be updated in every place it's used when the set of tags it needs to manipulate is added to or removed from,
+ * Aliases are a layer of abstraction that allows the set of tags to be added to or removed from without affecting the places where the alias is used.
+ */
+USTRUCT(BlueprintType)
+struct FAvaTagAlias
+{
+	GENERATED_BODY()
+
+#if WITH_EDITOR
+	void SetOwner(UAvaTagCollection* InOwner);
+
+	AVALANCHETAG_API UAvaTagCollection* GetOwner() const;
+
+	AVALANCHETAG_API FString GetTagsAsString() const;
+#endif
+
+	UPROPERTY(EditAnywhere, Category="Tag")
+	FName AliasName;
+
+	UPROPERTY(EditAnywhere, Category="Tag")
+	TArray<FAvaTagId> TagIds;
+
+private:
+#if WITH_EDITOR
+	/**
+	 * Tag Collection used to resolve the Tag Ids
+	 * Set by the Tag Collection that holds the alias map on load/change
+	 */
+	TWeakObjectPtr<UAvaTagCollection> OwnerWeak;
+#endif
+};
diff --git a/Engine/Plugins/Experimental/Avalanche/Source/AvalancheTag/Public/AvaTagCollection.h b/Engine/Plugins/Experimental/Avalanche/Source/AvalancheTag/Public/AvaTagCollection.h
--- a/Engine/Plugins/Experimental/Avalanche/Source/AvalancheTag/Public/AvaTagCollection.h
+++ b/Engine/Plugins/Experimental/Avalanche/Source/AvalancheTag/Public/AvaTagCollection.h
@@ -3,11 +3,14 @@
 #pragma once
 
 #include "AvaTag.h"
+#include "AvaTagAlias.h"
 #include "AvaTagId.h"
 #include "Containers/Map.h"
 #include "UObject/Object.h"
 #include "AvaTagCollection.generated.h"
 
+struct FAvaTagList;
+
 /**
  * Tag Collection that identifies a tag with an underlying Tag Id Guid
  * and provides Tag reference capabilities
@@ -18,14 +21,40 @@ class UAvaTagCollection : public UObject
 	GENERATED_BODY()
 
 public:
-	AVALANCHETAG_API const FAvaTag* GetTag(const FAvaTagId& InTagId) const;
+	/**
+	 * Returns an array of valid pointers to the tags that are under the given TagId
+	 * Which could be an Alias
+	 */
+	AVALANCHETAG_API FAvaTagList GetTags(const FAvaTagId& InTagId) const;
+
+	/**
+	 * Returns the name of the Tag mapped to the given TagId.
+	 * If the TagId is mapped to an Alias, it returns the Alias name instead.
+	 */
+	AVALANCHETAG_API FName GetTagName(const FAvaTagId& InTagId) const;
+
+	/** Returns the keys of the Tag Map, optionally including aliases */
+	AVALANCHETAG_API TArray<FAvaTagId> GetTagIds(bool bInIncludeAliases) const;
 
-	/** Returns the keys of the Tag Map */
-	AVALANCHETAG_API TArray<FAvaTagId> GetTagIds() const;
+	//~ Begin UObject
+	AVALANCHETAG_API virtual void PostLoad() override;
+#if WITH_EDITOR
+	AVALANCHETAG_API virtual void PostEditChangeProperty(FPropertyChangedEvent& InPropertyChangedEvent) override;
+#endif
+	//~ End UObject
 
 	AVALANCHETAG_API static FName GetTagMapName();
+	AVALANCHETAG_API static FName GetAliasMapName();
 
 private:
+#if WITH_EDITOR
+	/** Updates the weak ptr to the tag collection for each alias in the alias map */
+	void UpdateAliasOwner();
+#endif
+
 	UPROPERTY(EditAnywhere, Category="Tag")
 	TMap<FAvaTagId, FAvaTag> Tags;
+
+	UPROPERTY(EditAnywhere, Category="Tag")
+	TMap<FAvaTagId, FAvaTagAlias> Aliases;
 };
diff --git a/Engine/Plugins/Experimental/Avalanche/Source/AvalancheTag/Public/AvaTagHandle.h b/Engine/Plugins/Experimental/Avalanche/Source/AvalancheTag/Public/AvaTagHandle.h
--- a/Engine/Plugins/Experimental/Avalanche/Source/AvalancheTag/Public/AvaTagHandle.h
+++ b/Engine/Plugins/Experimental/Avalanche/Source/AvalancheTag/Public/AvaTagHandle.h
@@ -8,10 +8,11 @@
 
 class UAvaTagCollection;
 struct FAvaTag;
+struct FAvaTagList;
 
 /**
- * Handle to a Tag in a particular Source.
- * This should be used by the UStructs/UObjects to properly reference a particular FAvaTag.
+ * Handle to a Tag or Alias (set of Tags) in a particular Source.
+ * This should be used by the UStructs/UObjects to properly reference a Tag or Alias (set of Tags).
  */
 USTRUCT(BlueprintType)
 struct FAvaTagHandle
@@ -26,7 +27,12 @@ struct FAvaTagHandle
 	{
 	}
 
-	AVALANCHETAG_API const FAvaTag* GetTag() const;
+	/**
+	 * Returns the resolved tags from the Handle
+	 * If the Handle is to a particular Tag, it should return the array with a single element.
+	 * If the Handle is to an alias, it should return the array of tags the alias represents.
+	 */
+	AVALANCHETAG_API FAvaTagList GetTags() const;
 
 	AVALANCHETAG_API FString ToString() const;
 
@@ -36,8 +42,8 @@ struct FAvaTagHandle
 
 	AVALANCHETAG_API void PostSerialize(const FArchive& Ar);
 
-	/** Returns true if the Tag Handles resolve to same valued FAvaTags, even if the Source or Tag Id is different */
-	AVALANCHETAG_API bool MatchesTag(const FAvaTagHandle& InOther) const;
+	/** Returns true if the Tag Handles have overlapping FAvaTags, even if the Source or Tag Id is different */
+	AVALANCHETAG_API bool Overlaps(const FAvaTagHandle& InOther) const;
 
 	/** Returns true if the Tag Handles is the exact same as the other (Same Source and Tag Id) */
 	AVALANCHETAG_API bool MatchesExact(const FAvaTagHandle& InOther) const;
diff --git a/Engine/Plugins/Experimental/Avalanche/Source/AvalancheTag/Public/AvaTagLibrary.h b/Engine/Plugins/Experimental/Avalanche/Source/AvalancheTag/Public/AvaTagLibrary.h
--- a/Engine/Plugins/Experimental/Avalanche/Source/AvalancheTag/Public/AvaTagLibrary.h
+++ b/Engine/Plugins/Experimental/Avalanche/Source/AvalancheTag/Public/AvaTagLibrary.h
@@ -18,7 +18,7 @@ class UAvaTagLibrary : public UBlueprintFunctionLibrary
 
 public:
 	UFUNCTION(BlueprintPure, Category = "Motion Design|Tag", meta=(CompactNodeTitle = "->", BlueprintAutocast))
-	static AVALANCHETAG_API FAvaTag ResolveTagHandle(const FAvaTagHandle& InTagHandle);
+	static AVALANCHETAG_API TArray<FAvaTag> ResolveTagHandle(const FAvaTagHandle& InTagHandle);
 
 	UFUNCTION(BlueprintPure, Category = "Motion Design|Tag", meta=(CompactNodeTitle = "->", BlueprintAutocast))
 	static AVALANCHETAG_API TArray<FAvaTag> ResolveTagHandles(const FAvaTagHandleContainer& InTagHandleContainer);
diff --git a/Engine/Plugins/Experimental/Avalanche/Source/AvalancheTag/Public/AvaTagList.h b/Engine/Plugins/Experimental/Avalanche/Source/AvalancheTag/Public/AvaTagList.h
new file mode 100644
--- /dev/null
+++ b/Engine/Plugins/Experimental/Avalanche/Source/AvalancheTag/Public/AvaTagList.h
@@ -0,0 +1,24 @@
+// Copyright Epic Games, Inc. All Rights Reserved.
+
+#pragma once
+
+#include "Containers/Array.h"
+
+struct FAvaTag;
+
+struct FAvaTagList
+{
+	TArray<const FAvaTag*, TInlineAllocator<1>> Tags;
+
+	decltype(Tags)::RangedForIteratorType begin() { return Tags.begin(); }
+	decltype(Tags)::RangedForIteratorType end() { return Tags.end(); }
+
+	decltype(Tags)::RangedForConstIteratorType begin() const { return Tags.begin(); }
+	decltype(Tags)::RangedForConstIteratorType end() const { return Tags.end(); }
+
+	decltype(Tags)::RangedForReverseIteratorType rbegin() { return Tags.rbegin(); }
+	decltype(Tags)::RangedForReverseIteratorType rend() { return Tags.rend(); }
+
+	decltype(Tags)::RangedForConstReverseIteratorType rbegin() const { return Tags.rbegin(); }
+	decltype(Tags)::RangedForConstReverseIteratorType rend() const { return Tags.rend(); }
+};
diff --git a/Engine/Plugins/Experimental/Avalanche/Source/AvalancheTagEditor/Private/AvaTagEditorModule.cpp b/Engine/Plugins/Experimental/Avalanche/Source/AvalancheTagEditor/Private/AvaTagEditorModule.cpp
--- a/Engine/Plugins/Experimental/Avalanche/Source/AvalancheTagEditor/Private/AvaTagEditorModule.cpp
+++ b/Engine/Plugins/Experimental/Avalanche/Source/AvalancheTagEditor/Private/AvaTagEditorModule.cpp
@@ -7,6 +7,7 @@
 #include "AvaTagHandle.h"
 #include "AvaTagHandleContainer.h"
 #include "AvaTagSoftHandle.h"
+#include "Customization/AvaTagAliasCustomization.h"
 #include "Customization/AvaTagCollectionCustomization.h"
 #include "Customization/AvaTagHandleCustomization.h"
 #include "Customization/TagCustomizers/AvaTagHandleContainerCustomizer.h"
@@ -51,6 +52,12 @@ void FAvalancheTagEditorModule::RegisterCustomizations()
 	PropertyEditorModule.RegisterCustomPropertyTypeLayout(CustomizedTypes.Add_GetRef(FAvaTagHandleContainer::StaticStruct()->GetFName())
 		, UE::AvaTagEditor::Private::CreateTagHandleCustomization<FAvaTagHandleContainerCustomizer>());
 
+	PropertyEditorModule.RegisterCustomPropertyTypeLayout(CustomizedTypes.Add_GetRef(FAvaTagAlias::StaticStruct()->GetFName())
+		, FOnGetPropertyTypeCustomizationInstance::CreateLambda([]
+		{
+			return MakeShared<FAvaTagAliasCustomization>();
+		}));
+
 	PropertyEditorModule.RegisterCustomClassLayout(CustomizedClasses.Add_GetRef(UAvaTagCollection::StaticClass()->GetFName())
 		, FOnGetDetailCustomizationInstance::CreateLambda([]
 		{
diff --git a/Engine/Plugins/Experimental/Avalanche/Source/AvalancheTagEditor/Private/Customization/AvaTagAliasCustomization.cpp b/Engine/Plugins/Experimental/Avalanche/Source/AvalancheTagEditor/Private/Customization/AvaTagAliasCustomization.cpp
new file mode 100644
--- /dev/null
+++ b/Engine/Plugins/Experimental/Avalanche/Source/AvalancheTagEditor/Private/Customization/AvaTagAliasCustomization.cpp
@@ -0,0 +1,56 @@
+// Copyright Epic Games, Inc. All Rights Reserved.
+
+#include "AvaTagAliasCustomization.h"
+#include "AvaTagAlias.h"
+#include "Builders/AvaTagElementHelper.h"
+#include "DetailWidgetRow.h"
+#include "IDetailChildrenBuilder.h"
+#include "SAvaTagPicker.h"
+#include "TagCustomizers/AvaTagAliasCustomizer.h"
+#include "Widgets/Layout/SBox.h"
+#include "Widgets/SBoxPanel.h"
+
+FAvaTagAliasCustomization::FAvaTagAliasCustomization()
+	: TagElementHelper(MakeShared<FAvaTagElementHelper>())
+	, TagCustomizer(MakeShared<FAvaTagAliasCustomizer>())
+{
+}
+
+void FAvaTagAliasCustomization::CustomizeHeader(TSharedRef<IPropertyHandle> InPropertyHandle, FDetailWidgetRow& InHeaderRow, IPropertyTypeCustomizationUtils& InCustomizationUtils)
+{
+	TSharedPtr<IPropertyHandle> TagNameHandle = InPropertyHandle->GetChildHandle(GET_MEMBER_NAME_CHECKED(FAvaTagAlias, AliasName));
+
+	InHeaderRow
+		.WholeRowContent()
+		.HAlign(HAlign_Left)
+		[
+			SNew(SHorizontalBox)
+			+ SHorizontalBox::Slot()
+			.AutoWidth()
+			.VAlign(VAlign_Center)
+			[
+				SNew(SBox)
+				.MinDesiredWidth(150.f)
+				[
+					TagNameHandle->CreatePropertyValueWidget()	
+				]
+			]
+			+ SHorizontalBox::Slot()
+			.AutoWidth()
+			.VAlign(VAlign_Center)
+			[
+				SNew(SBox)
+				.MinDesiredWidth(150.f)
+				[
+					SNew(SAvaTagPicker, InPropertyHandle, TagCustomizer)
+				]
+			]
+			+ SHorizontalBox::Slot()
+			.AutoWidth()
+			.HAlign(HAlign_Left)
+			.VAlign(VAlign_Center)
+			[
+				TagElementHelper->CreatePropertyButtonsWidget(InPropertyHandle)
+			]
+		];
+}
diff --git a/Engine/Plugins/Experimental/Avalanche/Source/AvalancheTagEditor/Private/Customization/AvaTagAliasCustomization.h b/Engine/Plugins/Experimental/Avalanche/Source/AvalancheTagEditor/Private/Customization/AvaTagAliasCustomization.h
new file mode 100644
--- /dev/null
+++ b/Engine/Plugins/Experimental/Avalanche/Source/AvalancheTagEditor/Private/Customization/AvaTagAliasCustomization.h
@@ -0,0 +1,23 @@
+// Copyright Epic Games, Inc. All Rights Reserved.
+
+#pragma once
+
+#include "IPropertyTypeCustomization.h"
+
+class IAvaTagHandleCustomizer;
+class FAvaTagElementHelper;
+
+class FAvaTagAliasCustomization : public IPropertyTypeCustomization
+{
+public:
+	FAvaTagAliasCustomization();
+
+	//~ Begin IPropertyTypeCustomization
+	virtual void CustomizeHeader(TSharedRef<IPropertyHandle> InPropertyHandle, FDetailWidgetRow& InHeaderRow, IPropertyTypeCustomizationUtils& InCustomizationUtils) override;
+	virtual void CustomizeChildren(TSharedRef<IPropertyHandle> InPropertyHandle, IDetailChildrenBuilder& InChildrenBuilder, IPropertyTypeCustomizationUtils& InCustomizationUtils) override {}
+	//~ End IPropertyTypeCustomization
+
+	TSharedRef<FAvaTagElementHelper> TagElementHelper;
+
+	TSharedRef<IAvaTagHandleCustomizer> TagCustomizer;
+};
diff --git a/Engine/Plugins/Experimental/Avalanche/Source/AvalancheTagEditor/Private/Customization/AvaTagCollectionCustomization.cpp b/Engine/Plugins/Experimental/Avalanche/Source/AvalancheTagEditor/Private/Customization/AvaTagCollectionCustomization.cpp
--- a/Engine/Plugins/Experimental/Avalanche/Source/AvalancheTagEditor/Private/Customization/AvaTagCollectionCustomization.cpp
+++ b/Engine/Plugins/Experimental/Avalanche/Source/AvalancheTagEditor/Private/Customization/AvaTagCollectionCustomization.cpp
@@ -2,27 +2,31 @@
 
 #include "AvaTagCollectionCustomization.h"
 #include "AvaTagCollection.h"
+#include "Builders/AvaTagAliasMapBuilder.h"
+#include "Builders/AvaTagTagMapBuilder.h"
 #include "DetailCategoryBuilder.h"
 #include "DetailLayoutBuilder.h"
-#include "Editor.h"
-#include "Framework/MultiBox/MultiBoxBuilder.h"
+#include "DetailWidgetRow.h"
 #include "IDetailChildrenBuilder.h"
-#include "PropertyCustomizationHelpers.h"
-#include "Widgets/Input/SComboButton.h"
 #include "Widgets/Layout/SBox.h"
 #include "Widgets/SBoxPanel.h"
-#include "Widgets/Text/STextBlock.h"
-
-#define LOCTEXT_NAMESPACE "AvaTagCollectionCustomization"
 
 void FAvaTagCollectionCustomization::CustomizeDetails(IDetailLayoutBuilder& InDetailBuilder)
 {
-	const TSharedPtr<IPropertyHandle> TagMapProperty = InDetailBuilder.GetProperty(UAvaTagCollection::GetTagMapName());
-	TagMapProperty->MarkHiddenByCustomization();
+	TSharedRef<IPropertyHandle> TagMapHandle = InDetailBuilder.GetProperty(UAvaTagCollection::GetTagMapName());
+	InDetailBuilder.EditCategory(TEXT("Tags"))
+		.HeaderContent(BuildHeaderContent(TagMapHandle), /*bWholeRowContent*/true)
+		.AddCustomBuilder(MakeShared<FAvaTagTagMapBuilder>(TagMapHandle), /*bForAdvanced*/false);
 
-	IDetailCategoryBuilder& Category = InDetailBuilder.EditCategory(TEXT("Tag"));
+	TSharedRef<IPropertyHandle> AliasMapHandle = InDetailBuilder.GetProperty(UAvaTagCollection::GetAliasMapName());
+	InDetailBuilder.EditCategory(TEXT("Aliases"))
+		.HeaderContent(BuildHeaderContent(AliasMapHandle), /*bWholeRowContent*/true)
+		.AddCustomBuilder(MakeShared<FAvaTagAliasMapBuilder>(AliasMapHandle), /*bForAdvanced*/false);
+}
 
-	TSharedRef<SBox> HeaderContentWidget = SNew(SBox)
+TSharedRef<SWidget> FAvaTagCollectionCustomization::BuildHeaderContent(const TSharedRef<IPropertyHandle>& InPropertyHandle)
+{
+	return SNew(SBox)
 		.HAlign(HAlign_Left)
 		[
 			SNew(SBox)
@@ -33,162 +37,15 @@ void FAvaTagCollectionCustomization::CustomizeDetails(IDetailLayoutBuilder& InDe
 				.FillWidth(1.f)
 				.VAlign(VAlign_Center)
 				[
-					SNew(STextBlock)
-					.Text(Category.GetDisplayName())
-					.Font(IDetailLayoutBuilder::GetDetailFont())
+					InPropertyHandle->CreatePropertyNameWidget()
 				]
 				+ SHorizontalBox::Slot()
 				.AutoWidth()
 				.HAlign(HAlign_Right)
 				.VAlign(VAlign_Center)
 				[
-					TagMapProperty->CreateDefaultPropertyButtonWidgets()
+					InPropertyHandle->CreateDefaultPropertyButtonWidgets()
 				]
 			]
 		];
-
-	Category.HeaderContent(HeaderContentWidget, /*bWholeRowContent*/true);
-	Category.AddCustomBuilder(MakeShared<FAvaTagMapBuilder>(TagMapProperty.ToSharedRef()), /*bForAdvanced*/false);
 }
-
-FAvaTagMapBuilder::FAvaTagMapBuilder(const TSharedRef<IPropertyHandle>& InTagMapProperty)
-	: MapProperty(InTagMapProperty->AsMap())
-	, BaseProperty(InTagMapProperty)
-{
-	check(MapProperty.IsValid());
-
-	// Delegate for when the number of children in the array changes
-	FSimpleDelegate OnNumChildrenChanged = FSimpleDelegate::CreateRaw(this, &FAvaTagMapBuilder::OnNumChildrenChanged);
-	OnNumElementsChangedHandle = MapProperty->SetOnNumElementsChanged(OnNumChildrenChanged);
-
-	BaseProperty->MarkHiddenByCustomization();
-}
-
-FAvaTagMapBuilder::~FAvaTagMapBuilder()
-{
-	MapProperty->UnregisterOnNumElementsChanged(OnNumElementsChangedHandle);
-}
-
-FName FAvaTagMapBuilder::GetName() const
-{
-	return BaseProperty->GetProperty()->GetFName();
-}
-
-void FAvaTagMapBuilder::GenerateChildContent(IDetailChildrenBuilder& InChildrenBuilder)
-{
-	uint32 NumChildren = 0;
-	BaseProperty->GetNumChildren(NumChildren);
-
-	for (uint32 ChildIndex = 0; ChildIndex < NumChildren; ++ChildIndex)
-	{
-		TSharedPtr<IPropertyHandle> TagHandle = BaseProperty->GetChildHandle(ChildIndex);
-		if (!TagHandle.IsValid())
-		{
-			continue;
-		}
-
-		if (TSharedPtr<IPropertyHandle> TagNameHandle = TagHandle->GetChildHandle(GET_MEMBER_NAME_CHECKED(FAvaTag, TagName)))
-		{
-			InChildrenBuilder.AddProperty(TagHandle.ToSharedRef())
-				.CustomWidget()
-				.WholeRowContent()
-				.HAlign(HAlign_Left)
-				[
-					SNew(SBox)
-					.MinDesiredWidth(250.f)
-					[
-						SNew(SHorizontalBox)
-						+ SHorizontalBox::Slot()
-						.FillWidth(1.f)
-						.VAlign(VAlign_Center)
-						[
-							TagNameHandle->CreatePropertyValueWidget()
-						]
-						+ SHorizontalBox::Slot()
-						.AutoWidth()
-						.HAlign(HAlign_Right)
-						.VAlign(VAlign_Center)
-						[
-							CreatePropertyButtonsWidget(TagHandle)
-						]
-					]
-				];
-		}
-	}
-}
-
-void FAvaTagMapBuilder::OnNumChildrenChanged()
-{
-	OnRebuildChildren.ExecuteIfBound();
-}
-
-TSharedPtr<IPropertyHandle> FAvaTagMapBuilder::GetPropertyHandle() const
-{
-	return BaseProperty;
-}
-
-void FAvaTagMapBuilder::SetOnRebuildChildren(FSimpleDelegate InOnRebuildChildren)
-{
-	OnRebuildChildren = InOnRebuildChildren;
-}
-
-TSharedRef<SWidget> FAvaTagMapBuilder::CreatePropertyButtonsWidget(TSharedPtr<IPropertyHandle> InTagHandle)
-{
-	FMenuBuilder MenuContentBuilder(true, nullptr, nullptr, true);
-
-	MenuContentBuilder.AddMenuEntry(NSLOCTEXT("PropertyCustomizationHelpers", "DeleteButtonLabel", "Delete")
-		, FText::GetEmpty()
-		, FSlateIcon()
-		, FExecuteAction::CreateSP(this, &FAvaTagMapBuilder::DeleteItem, InTagHandle));
-
-	MenuContentBuilder.AddMenuEntry(LOCTEXT("SearchForReferencesLabel", "Search for References")
-		, FText::GetEmpty()
-		, FSlateIcon()
-		, FExecuteAction::CreateSP(this, &FAvaTagMapBuilder::SearchForReferences, InTagHandle->GetKeyHandle()));
-
-	return SNew(SComboButton)
-		.ComboButtonStyle(FAppStyle::Get(), "SimpleComboButton")
-		.ContentPadding(2)
-		.ForegroundColor(FSlateColor::UseForeground())
-		.HasDownArrow(true)
-		.MenuContent()
-		[
-			MenuContentBuilder.MakeWidget()
-		];
-}
-
-void FAvaTagMapBuilder::DeleteItem(TSharedPtr<IPropertyHandle> InTagHandle)
-{
-	if (InTagHandle.IsValid())
-	{
-		MapProperty->DeleteItem(InTagHandle->GetArrayIndex());	
-	}
-}
-
-void FAvaTagMapBuilder::SearchForReferences(TSharedPtr<IPropertyHandle> InTagIdHandle)
-{
-	if (!FEditorDelegates::OnOpenReferenceViewer.IsBound())
-	{
-		return;
-	}
-
-	if (!InTagIdHandle.IsValid())
-	{
-		return;
-	}
-
-	TArray<const void*> TagIdRawData;
-	InTagIdHandle->AccessRawData(TagIdRawData);
-	if (TagIdRawData.IsEmpty() || !TagIdRawData[0])
-	{
-		return;
-	}
-
-	const FAvaTagId& TagId = *static_cast<const FAvaTagId*>(TagIdRawData[0]);
-
-	TArray<FAssetIdentifier> AssetIdentifiers;
-	AssetIdentifiers.Emplace(FAvaTagId::StaticStruct(), *TagId.ToString());
-	FEditorDelegates::OnOpenReferenceViewer.Broadcast(AssetIdentifiers, FReferenceViewerParams());
-}
-
-#undef LOCTEXT_NAMESPACE
diff --git a/Engine/Plugins/Experimental/Avalanche/Source/AvalancheTagEditor/Private/Customization/AvaTagCollectionCustomization.h b/Engine/Plugins/Experimental/Avalanche/Source/AvalancheTagEditor/Private/Customization/AvaTagCollectionCustomization.h
--- a/Engine/Plugins/Experimental/Avalanche/Source/AvalancheTagEditor/Private/Customization/AvaTagCollectionCustomization.h
+++ b/Engine/Plugins/Experimental/Avalanche/Source/AvalancheTagEditor/Private/Customization/AvaTagCollectionCustomization.h
@@ -2,12 +2,9 @@
 
 #pragma once
 
-#include "Delegates/Delegate.h"
-#include "IDetailCustomNodeBuilder.h"
 #include "IDetailCustomization.h"
-#include "Templates/SharedPointer.h"
 
-class IPropertyHandleMap;
+class IPropertyHandle;
 class SWidget;
 
 class FAvaTagCollectionCustomization : public IDetailCustomization
@@ -16,39 +13,6 @@ public:
 	//~ Begin IDetailCustomization
 	virtual void CustomizeDetails(IDetailLayoutBuilder& InDetailBuilder) override;
 	//~ End IDetailCustomization
-};
-
-class FAvaTagMapBuilder : public IDetailCustomNodeBuilder, public TSharedFromThis<FAvaTagMapBuilder>
-{
-public:
-	FAvaTagMapBuilder(const TSharedRef<IPropertyHandle>& InTagMapProperty);
-
-	virtual ~FAvaTagMapBuilder() override;
-
-	//~ Begin IDetailCustomNodeBuilder
-	virtual FName GetName() const override;
-	virtual bool InitiallyCollapsed() const override { return false; }
-	virtual void GenerateHeaderRowContent(FDetailWidgetRow& InNodeRow) override {}
-	virtual void GenerateChildContent(IDetailChildrenBuilder& InChildrenBuilder) override;
-	virtual TSharedPtr<IPropertyHandle> GetPropertyHandle() const override;
-	virtual void SetOnRebuildChildren(FSimpleDelegate InOnRebuildChildren) override;
-	//~ End IDetailCustomNodeBuilder
-
-private:
-	TSharedRef<SWidget> CreatePropertyButtonsWidget(TSharedPtr<IPropertyHandle> InTagHandle);
 
-	void DeleteItem(TSharedPtr<IPropertyHandle> InTagHandle);
-
-	void SearchForReferences(TSharedPtr<IPropertyHandle> InTagIdHandle);
-
-	void OnNumChildrenChanged();
-
-	TSharedPtr<IPropertyHandleMap> MapProperty;
-
-	TSharedPtr<IPropertyHandle> BaseProperty;
-
-	FDelegateHandle OnNumElementsChangedHandle;
-
-	FSimpleDelegate OnRebuildChildren;
+	TSharedRef<SWidget> BuildHeaderContent(const TSharedRef<IPropertyHandle>& InPropertyHandle);
 };
-
diff --git a/Engine/Plugins/Experimental/Avalanche/Source/AvalancheTagEditor/Private/Customization/Builders/AvaTagAliasMapBuilder.cpp b/Engine/Plugins/Experimental/Avalanche/Source/AvalancheTagEditor/Private/Customization/Builders/AvaTagAliasMapBuilder.cpp
new file mode 100644
--- /dev/null
+++ b/Engine/Plugins/Experimental/Avalanche/Source/AvalancheTagEditor/Private/Customization/Builders/AvaTagAliasMapBuilder.cpp
@@ -0,0 +1,13 @@
+// Copyright Epic Games, Inc. All Rights Reserved.
+
+#include "AvaTagAliasMapBuilder.h"
+#include "IDetailChildrenBuilder.h"
+
+void FAvaTagAliasMapBuilder::GenerateChildContent(IDetailChildrenBuilder& InChildrenBuilder)
+{
+	ForEachChildProperty(
+		[&InChildrenBuilder](const TSharedRef<IPropertyHandle>& InChildHandle)
+		{
+			InChildrenBuilder.AddProperty(InChildHandle);
+		});
+}
diff --git a/Engine/Plugins/Experimental/Avalanche/Source/AvalancheTagEditor/Private/Customization/Builders/AvaTagAliasMapBuilder.h b/Engine/Plugins/Experimental/Avalanche/Source/AvalancheTagEditor/Private/Customization/Builders/AvaTagAliasMapBuilder.h
new file mode 100644
--- /dev/null
+++ b/Engine/Plugins/Experimental/Avalanche/Source/AvalancheTagEditor/Private/Customization/Builders/AvaTagAliasMapBuilder.h
@@ -0,0 +1,19 @@
+// Copyright Epic Games, Inc. All Rights Reserved.
+
+#pragma once
+
+#include "AvaTagMapBuilder.h"
+
+/** Map builder for the Alias Map */
+class FAvaTagAliasMapBuilder : public FAvaTagMapBuilder
+{
+public:
+	explicit FAvaTagAliasMapBuilder(const TSharedRef<IPropertyHandle>& InMapProperty)
+		: FAvaTagMapBuilder(InMapProperty)
+	{
+	}
+
+	//~ Begin IDetailCustomNodeBuilder
+	virtual void GenerateChildContent(IDetailChildrenBuilder& InChildrenBuilder) override;
+	//~ End IDetailCustomNodeBuilder
+};
diff --git a/Engine/Plugins/Experimental/Avalanche/Source/AvalancheTagEditor/Private/Customization/Builders/AvaTagElementHelper.cpp b/Engine/Plugins/Experimental/Avalanche/Source/AvalancheTagEditor/Private/Customization/Builders/AvaTagElementHelper.cpp
new file mode 100644
--- /dev/null
+++ b/Engine/Plugins/Experimental/Avalanche/Source/AvalancheTagEditor/Private/Customization/Builders/AvaTagElementHelper.cpp
@@ -0,0 +1,114 @@
+// Copyright Epic Games, Inc. All Rights Reserved.
+
+#include "AvaTagElementHelper.h"
+#include "AssetRegistry/AssetIdentifier.h"
+#include "AvaTag.h"
+#include "AvaTagId.h"
+#include "Editor.h"
+#include "Framework/MultiBox/MultiBoxBuilder.h"
+#include "PropertyHandle.h"
+#include "Widgets/Input/SComboBox.h"
+
+#define LOCTEXT_NAMESPACE "AvaTagElementHelper"
+
+TSharedRef<SWidget> FAvaTagElementHelper::CreatePropertyButtonsWidget(TSharedPtr<IPropertyHandle> InElementHandle)
+{
+	FMenuBuilder MenuContentBuilder(/*bInShouldCloseWindowAfterMenuSelection*/true
+		, /*InCommandList*/nullptr
+		, /*InExtender*/nullptr
+		, /*bInCloseSelfOnly*/true);
+
+	if (CanDeleteItem(InElementHandle))
+	{
+		MenuContentBuilder.AddMenuEntry(LOCTEXT("DeleteButtonLabel", "Delete")
+			, FText::GetEmpty()
+			, FSlateIcon()
+			, FExecuteAction::CreateSP(this, &FAvaTagElementHelper::DeleteItem, InElementHandle));
+	}
+
+	MenuContentBuilder.AddMenuEntry(LOCTEXT("SearchForReferencesLabel", "Search for References")
+		, FText::GetEmpty()
+		, FSlateIcon()
+		, FExecuteAction::CreateSP(this, &FAvaTagElementHelper::SearchForReferences, InElementHandle->GetKeyHandle()));
+
+	return SNew(SComboButton)
+		.ComboButtonStyle(FAppStyle::Get(), "SimpleComboButton")
+		.ContentPadding(2)
+		.ForegroundColor(FSlateColor::UseForeground())
+		.HasDownArrow(true)
+		.MenuContent()
+		[
+			MenuContentBuilder.MakeWidget()
+		];
+}
+
+bool FAvaTagElementHelper::CanDeleteItem(const TSharedPtr<IPropertyHandle>& InElementHandle) const
+{
+	if (!InElementHandle.IsValid())
+	{
+		return false;
+	}
+
+	TSharedPtr<IPropertyHandle> ParentHandle = InElementHandle->GetParentHandle();
+	if (!ParentHandle.IsValid())
+	{
+		return false;
+	}
+
+	return ParentHandle->AsArray() || ParentHandle->AsMap() || ParentHandle->AsSet();
+}
+
+void FAvaTagElementHelper::DeleteItem(TSharedPtr<IPropertyHandle> InElementHandle)
+{
+	if (!InElementHandle.IsValid())
+	{
+		return;
+	}
+
+	TSharedPtr<IPropertyHandle> ParentHandle = InElementHandle->GetParentHandle();
+	if (!ParentHandle.IsValid())
+	{
+		return;
+	}
+
+	if (TSharedPtr<IPropertyHandleArray> ArrayProperty = ParentHandle->AsArray())
+	{
+		ArrayProperty->DeleteItem(InElementHandle->GetArrayIndex());
+	}
+	else if (TSharedPtr<IPropertyHandleMap> MapProperty = ParentHandle->AsMap())
+	{
+		MapProperty->DeleteItem(InElementHandle->GetArrayIndex());
+	}
+	else if (TSharedPtr<IPropertyHandleSet> SetProperty = ParentHandle->AsSet())
+	{
+		SetProperty->DeleteItem(InElementHandle->GetArrayIndex());
+	}
+}
+
+void FAvaTagElementHelper::SearchForReferences(TSharedPtr<IPropertyHandle> InTagIdHandle)
+{
+	if (!FEditorDelegates::OnOpenReferenceViewer.IsBound())
+	{
+		return;
+	}
+
+	if (!InTagIdHandle.IsValid())
+	{
+		return;
+	}
+
+	TArray<const void*> TagIdRawData;
+	InTagIdHandle->AccessRawData(TagIdRawData);
+	if (TagIdRawData.IsEmpty() || !TagIdRawData[0])
+	{
+		return;
+	}
+
+	const FAvaTagId& TagId = *static_cast<const FAvaTagId*>(TagIdRawData[0]);
+
+	TArray<FAssetIdentifier> AssetIdentifiers;
+	AssetIdentifiers.Emplace(FAvaTagId::StaticStruct(), *TagId.ToString());
+	FEditorDelegates::OnOpenReferenceViewer.Broadcast(AssetIdentifiers, FReferenceViewerParams());
+}
+
+#undef LOCTEXT_NAMESPACE
diff --git a/Engine/Plugins/Experimental/Avalanche/Source/AvalancheTagEditor/Private/Customization/Builders/AvaTagElementHelper.h b/Engine/Plugins/Experimental/Avalanche/Source/AvalancheTagEditor/Private/Customization/Builders/AvaTagElementHelper.h
new file mode 100644
--- /dev/null
+++ b/Engine/Plugins/Experimental/Avalanche/Source/AvalancheTagEditor/Private/Customization/Builders/AvaTagElementHelper.h
@@ -0,0 +1,21 @@
+// Copyright Epic Games, Inc. All Rights Reserved.
+
+#pragma once
+
+#include "Templates/SharedPointer.h"
+
+class IPropertyHandle;
+class SWidget;
+
+/** Helper class to re-use same functionality across Node Builders and Customizations */
+class FAvaTagElementHelper : public TSharedFromThis<FAvaTagElementHelper>
+{
+public:
+	TSharedRef<SWidget> CreatePropertyButtonsWidget(TSharedPtr<IPropertyHandle> InElementHandle);
+
+	bool CanDeleteItem(const TSharedPtr<IPropertyHandle>& InElementHandle) const;
+
+	void DeleteItem(TSharedPtr<IPropertyHandle> InElementHandle);
+
+	void SearchForReferences(TSharedPtr<IPropertyHandle> InTagIdHandle);
+};
diff --git a/Engine/Plugins/Experimental/Avalanche/Source/AvalancheTagEditor/Private/Customization/Builders/AvaTagMapBuilder.cpp b/Engine/Plugins/Experimental/Avalanche/Source/AvalancheTagEditor/Private/Customization/Builders/AvaTagMapBuilder.cpp
new file mode 100644
--- /dev/null
+++ b/Engine/Plugins/Experimental/Avalanche/Source/AvalancheTagEditor/Private/Customization/Builders/AvaTagMapBuilder.cpp
@@ -0,0 +1,57 @@
+// Copyright Epic Games, Inc. All Rights Reserved.
+
+#include "AvaTagMapBuilder.h"
+#include "PropertyHandle.h"
+
+FAvaTagMapBuilder::FAvaTagMapBuilder(const TSharedRef<IPropertyHandle>& InMapProperty)
+	: BaseProperty(InMapProperty)
+	, MapProperty(InMapProperty->AsMap().ToSharedRef())
+{
+	BaseProperty->MarkHiddenByCustomization();
+
+	// Delegate for when the number of children in the array changes
+	FSimpleDelegate OnNumChildrenChanged = FSimpleDelegate::CreateRaw(this, &FAvaTagMapBuilder::OnNumChildrenChanged);
+	OnNumElementsChangedHandle = MapProperty->SetOnNumElementsChanged(OnNumChildrenChanged);
+}
+
+FAvaTagMapBuilder::~FAvaTagMapBuilder()
+{
+	MapProperty->UnregisterOnNumElementsChanged(OnNumElementsChangedHandle);
+	OnNumElementsChangedHandle.Reset();
+}
+
+FName FAvaTagMapBuilder::GetName() const
+{
+	const FProperty* Property = BaseProperty->GetProperty();
+	ensure(Property);
+	return Property ? Property->GetFName() : NAME_None;
+}
+
+TSharedPtr<IPropertyHandle> FAvaTagMapBuilder::GetPropertyHandle() const
+{
+	return BaseProperty;
+}
+
+void FAvaTagMapBuilder::SetOnRebuildChildren(FSimpleDelegate InOnRebuildChildren)
+{
+	OnRebuildChildren = InOnRebuildChildren;
+}
+
+void FAvaTagMapBuilder::ForEachChildProperty(TFunctionRef<void(const TSharedRef<IPropertyHandle>&)> InCallable)
+{
+	uint32 NumChildren = 0;
+	BaseProperty->GetNumChildren(NumChildren);
+
+	for (uint32 ChildIndex = 0; ChildIndex < NumChildren; ++ChildIndex)
+	{
+		if (TSharedPtr<IPropertyHandle> ChildHandle = BaseProperty->GetChildHandle(ChildIndex))
+		{
+			InCallable(ChildHandle.ToSharedRef());
+		}
+	}
+}
+
+void FAvaTagMapBuilder::OnNumChildrenChanged()
+{
+	OnRebuildChildren.ExecuteIfBound();
+}
diff --git a/Engine/Plugins/Experimental/Avalanche/Source/AvalancheTagEditor/Private/Customization/Builders/AvaTagMapBuilder.h b/Engine/Plugins/Experimental/Avalanche/Source/AvalancheTagEditor/Private/Customization/Builders/AvaTagMapBuilder.h
new file mode 100644
--- /dev/null
+++ b/Engine/Plugins/Experimental/Avalanche/Source/AvalancheTagEditor/Private/Customization/Builders/AvaTagMapBuilder.h
@@ -0,0 +1,39 @@
+// Copyright Epic Games, Inc. All Rights Reserved.
+
+#pragma once
+
+#include "Delegates/Delegate.h"
+#include "IDetailCustomNodeBuilder.h"
+#include "Templates/SharedPointer.h"
+
+class IPropertyHandleMap;
+
+/** Base class to Build Map properties */
+class FAvaTagMapBuilder : public IDetailCustomNodeBuilder, public TSharedFromThis<FAvaTagMapBuilder>
+{
+public:
+	FAvaTagMapBuilder(const TSharedRef<IPropertyHandle>& InMapProperty);
+
+	virtual ~FAvaTagMapBuilder() override;
+
+	//~ Begin IDetailCustomNodeBuilder
+	virtual FName GetName() const override;
+	virtual bool InitiallyCollapsed() const override { return false; }
+	virtual void GenerateHeaderRowContent(FDetailWidgetRow& InNodeRow) override {}
+	virtual TSharedPtr<IPropertyHandle> GetPropertyHandle() const override;
+	virtual void SetOnRebuildChildren(FSimpleDelegate InOnRebuildChildren) override;
+	//~ End IDetailCustomNodeBuilder
+
+protected:
+	void ForEachChildProperty(TFunctionRef<void(const TSharedRef<IPropertyHandle>&)> InCallable);
+
+	void OnNumChildrenChanged();
+
+	TSharedRef<IPropertyHandle> BaseProperty;
+
+	TSharedRef<IPropertyHandleMap> MapProperty;
+
+	FDelegateHandle OnNumElementsChangedHandle;
+
+	FSimpleDelegate OnRebuildChildren;
+};
diff --git a/Engine/Plugins/Experimental/Avalanche/Source/AvalancheTagEditor/Private/Customization/Builders/AvaTagTagMapBuilder.cpp b/Engine/Plugins/Experimental/Avalanche/Source/AvalancheTagEditor/Private/Customization/Builders/AvaTagTagMapBuilder.cpp
new file mode 100644
--- /dev/null
+++ b/Engine/Plugins/Experimental/Avalanche/Source/AvalancheTagEditor/Private/Customization/Builders/AvaTagTagMapBuilder.cpp
@@ -0,0 +1,51 @@
+// Copyright Epic Games, Inc. All Rights Reserved.
+
+#include "AvaTagTagMapBuilder.h"
+#include "AvaTag.h"
+#include "AvaTagElementHelper.h"
+#include "DetailWidgetRow.h"
+#include "IDetailChildrenBuilder.h"
+#include "IDetailPropertyRow.h"
+#include "Widgets/Layout/SBox.h"
+#include "Widgets/SBoxPanel.h"
+
+FAvaTagTagMapBuilder::FAvaTagTagMapBuilder(const TSharedRef<IPropertyHandle>& InMapProperty)
+	: FAvaTagMapBuilder(InMapProperty)
+	, TagElementHelper(MakeShared<FAvaTagElementHelper>())
+{
+}
+
+void FAvaTagTagMapBuilder::GenerateChildContent(IDetailChildrenBuilder& InChildrenBuilder)
+{
+	ForEachChildProperty(
+		[&InChildrenBuilder, this](const TSharedRef<IPropertyHandle>& InChildHandle)
+		{
+			TSharedPtr<IPropertyHandle> TagNameHandle = InChildHandle->GetChildHandle(GET_MEMBER_NAME_CHECKED(FAvaTag, TagName));
+			check(TagNameHandle.IsValid());
+
+			InChildrenBuilder.AddProperty(InChildHandle)
+				.CustomWidget(/*bShowChildren*/false)
+				.WholeRowContent()
+				.HAlign(HAlign_Left)
+				[
+					SNew(SHorizontalBox)
+					+ SHorizontalBox::Slot()
+					.AutoWidth()
+					.VAlign(VAlign_Center)
+					[
+						SNew(SBox)
+						.MinDesiredWidth(150.f)
+						[
+							TagNameHandle->CreatePropertyValueWidget()
+						]
+					]
+					+ SHorizontalBox::Slot()
+					.AutoWidth()
+					.HAlign(HAlign_Right)
+					.VAlign(VAlign_Center)
+					[
+						TagElementHelper->CreatePropertyButtonsWidget(InChildHandle)
+					]
+				];
+		});
+}
diff --git a/Engine/Plugins/Experimental/Avalanche/Source/AvalancheTagEditor/Private/Customization/Builders/AvaTagTagMapBuilder.h b/Engine/Plugins/Experimental/Avalanche/Source/AvalancheTagEditor/Private/Customization/Builders/AvaTagTagMapBuilder.h
new file mode 100644
--- /dev/null
+++ b/Engine/Plugins/Experimental/Avalanche/Source/AvalancheTagEditor/Private/Customization/Builders/AvaTagTagMapBuilder.h
@@ -0,0 +1,21 @@
+// Copyright Epic Games, Inc. All Rights Reserved.
+
+#pragma once
+
+#include "AvaTagMapBuilder.h"
+
+class FAvaTagElementHelper;
+
+/** Map builder for the Tag Map */
+class FAvaTagTagMapBuilder : public FAvaTagMapBuilder
+{
+public:
+	explicit FAvaTagTagMapBuilder(const TSharedRef<IPropertyHandle>& InMapProperty);
+
+	//~ Begin IDetailCustomNodeBuilder
+	virtual void GenerateChildContent(IDetailChildrenBuilder& InChildrenBuilder) override;
+	//~ End IDetailCustomNodeBuilder
+
+private:
+	TSharedRef<FAvaTagElementHelper> TagElementHelper;
+};
diff --git a/Engine/Plugins/Experimental/Avalanche/Source/AvalancheTagEditor/Private/Customization/SAvaTagPicker.cpp b/Engine/Plugins/Experimental/Avalanche/Source/AvalancheTagEditor/Private/Customization/SAvaTagPicker.cpp
--- a/Engine/Plugins/Experimental/Avalanche/Source/AvalancheTagEditor/Private/Customization/SAvaTagPicker.cpp
+++ b/Engine/Plugins/Experimental/Avalanche/Source/AvalancheTagEditor/Private/Customization/SAvaTagPicker.cpp
@@ -28,30 +28,20 @@ void SAvaTagPicker::Construct(const FArguments& InArgs, const TSharedRef<IProper
 
 	TagCustomizer = InTagCustomizer;
 
-	TagCollectionPropertyHandle = InTagCustomizer->GetTagCollectionHandle(InStructPropertyHandle);
-	check(TagCollectionPropertyHandle.IsValid());
+	TSharedRef<SVerticalBox> ContentVerticalBox = SNew(SVerticalBox);
 
-	const FToolBarStyle& SlimToolbarStyle = FAppStyle::Get().GetWidgetStyle<FToolBarStyle>("SlimToolBar");
+	TSharedRef<SHorizontalBox> WrapperHorizontalBox = SNew(SHorizontalBox)
+		+ SHorizontalBox::Slot()
+        .FillWidth(1.f)
+        [
+        	ContentVerticalBox
+        ];
 
-	// Selection Mode set to single, yet preventing Items to be selected via OnIsSelectableOrNavigable, so that hover cue appears
-	TagListView = SNew(SListView<TSharedPtr<FAvaTagHandle>>)
-		.ListViewStyle(&FAppStyle::Get().GetWidgetStyle<FTableViewStyle>("SimpleListView"))
-		.OnGenerateRow(this, &SAvaTagPicker::CreateTagTableRow)
-		.SelectionMode(ESelectionMode::Single)
-		.ListItemsSource(&TagOptions)
-		.IsFocusable(false)
-		.HandleGamepadEvents(false)
-		.HandleSpacebarSelection(false)
-		.HandleDirectionalNavigation(false);
+	TagCollectionPropertyHandle = InTagCustomizer->GetTagCollectionHandle(InStructPropertyHandle);
 
-	ChildSlot
-	[
-		SNew(SHorizontalBox)
-		+ SHorizontalBox::Slot()
-		.FillWidth(1.f)
-		[
-			SNew(SVerticalBox)
-			+ SVerticalBox::Slot()
+	if (TagCollectionPropertyHandle.IsValid())
+	{
+		ContentVerticalBox->AddSlot()
 			.AutoHeight()
 			.Padding(0)
 			.HAlign(HAlign_Left)
@@ -59,49 +49,63 @@ void SAvaTagPicker::Construct(const FArguments& InArgs, const TSharedRef<IProper
 			[
 				SAssignNew(TagCollectionPicker, SAvaTagCollectionPicker, TagCollectionPropertyHandle.ToSharedRef())
 				.OnTagCollectionChanged(this, &SAvaTagPicker::OnTagCollectionChanged)
-			]
-			+ SVerticalBox::Slot()
-			.AutoHeight()
-			.Padding(0)
-			.HAlign(HAlign_Fill)
+			];
+
+		WrapperHorizontalBox->AddSlot()
+			.AutoWidth()
+			.HAlign(HAlign_Center)
 			.VAlign(VAlign_Center)
 			[
-				SAssignNew(TagComboButton, SComboButton)
-				.HasDownArrow(true)
-				.ContentPadding(FMargin(0, -1))
-				.OnMenuOpenChanged(this, &SAvaTagPicker::OnTagMenuOpenChanged)
-				.MenuContent()
+				SAssignNew(TagCollectionOptions, SMenuAnchor)
+				.Content()
 				[
-					TagListView.ToSharedRef()
+					SNew(SButton)
+					.OnClicked(this, &SAvaTagPicker::OpenContextMenu)
+					.ButtonStyle(FAppStyle::Get(), "SimpleButton")
+					.ToolTipText(LOCTEXT("TagCollectionOptionsToolTip", "Tag Collection Options"))
+					.ContentPadding(FMargin(4, 6))
+					[
+						SNew(SImage)
+						.Image(&FAppStyle::Get().GetWidgetStyle<FToolBarStyle>("SlimToolBar").SettingsComboButton.DownArrowImage)
+						.ColorAndOpacity(FSlateColor::UseForeground())
+					]
 				]
-				.ButtonContent()
-				[
-					SNew(STextBlock)
-					.Text(this, &SAvaTagPicker::GetValueDisplayText)
-					.Font(IDetailLayoutBuilder::GetDetailFont())
-				]
-			]
-		]
-		+ SHorizontalBox::Slot()
-		.AutoWidth()
-		.HAlign(HAlign_Center)
+			];
+	}
+
+	ContentVerticalBox->AddSlot()
+		.AutoHeight()
+		.Padding(0)
+		.HAlign(HAlign_Fill)
 		.VAlign(VAlign_Center)
 		[
-			SAssignNew(TagCollectionOptions, SMenuAnchor)
-			.Content()
+			SAssignNew(TagComboButton, SComboButton)
+			.HasDownArrow(true)
+			.ContentPadding(FMargin(0, -1))
+			.OnMenuOpenChanged(this, &SAvaTagPicker::OnTagMenuOpenChanged)
+			.MenuContent()
 			[
-				SNew(SButton)
-				.OnClicked(this, &SAvaTagPicker::OpenContextMenu)
-				.ButtonStyle(FAppStyle::Get(), "SimpleButton")
-				.ToolTipText(LOCTEXT("TagCollectionOptionsToolTip", "Tag Collection Options"))
-				.ContentPadding(FMargin(4, 6))
-				[
-					SNew(SImage)
-					.Image(&SlimToolbarStyle.SettingsComboButton.DownArrowImage)
-					.ColorAndOpacity(FSlateColor::UseForeground())
-				]
+				SAssignNew(TagListView, SListView<TSharedPtr<FAvaTagHandle>>)
+				.ListViewStyle(&FAppStyle::Get().GetWidgetStyle<FTableViewStyle>("SimpleListView"))
+				.OnGenerateRow(this, &SAvaTagPicker::CreateTagTableRow)
+				.SelectionMode(ESelectionMode::Single) // Selection Mode set to single, yet preventing Items to be selected via OnIsSelectableOrNavigable, so that hover cue appears
+				.ListItemsSource(&TagOptions)
+				.IsFocusable(false)
+				.HandleGamepadEvents(false)
+				.HandleSpacebarSelection(false)
+				.HandleDirectionalNavigation(false)
 			]
-		]
+			.ButtonContent()
+			[
+				SNew(STextBlock)
+				.Text(this, &SAvaTagPicker::GetValueDisplayText)
+				.Font(IDetailLayoutBuilder::GetDetailFont())
+			]
+		];
+
+	ChildSlot
+	[
+		WrapperHorizontalBox
 	];
 }
 
@@ -110,7 +114,12 @@ void SAvaTagPicker::Tick(const FGeometry& InGeometry, const double InCurrentTime
 	if (bRequestOpenTagCollectionPicker)
 	{
 		TagComboButton->SetIsOpen(false);
-		TagCollectionPicker->SetIsOpen(true);
+
+		if (TagCollectionPicker.IsValid())
+		{
+			TagCollectionPicker->SetIsOpen(true);
+		}
+
 		bRequestOpenTagCollectionPicker = false;
 	}
 }
@@ -126,8 +135,11 @@ FReply SAvaTagPicker::OnMouseButtonUp(const FGeometry& InGeometry, const FPointe
 
 FReply SAvaTagPicker::OpenContextMenu()
 {
-	if (!TagCollectionOptions->IsOpen())
+	if (TagCollectionOptions.IsValid() && !TagCollectionOptions->IsOpen())
 	{
+		// If TagCollectionOptions is valid, the TagCollectionHandle should also be
+		ensure(TagCollectionPropertyHandle.IsValid());
+
 		TagCollectionOptions->SetMenuContent(FAvaTagCollectionPickerContextMenu::Get().GenerateContextMenuWidget(TagCollectionPropertyHandle));
 		TagCollectionOptions->SetIsOpen(true);
 	}
@@ -174,7 +186,7 @@ void SAvaTagPicker::RefreshTagOptions()
 	// Get or Load. Loading is ok at this point because this we're refreshing the tags the Tag Collection has to offer (i.e. we're peeking inside of it)
 	if (const UAvaTagCollection* TagCollection = GetOrLoadTagCollection())
 	{
-		TArray<FAvaTagId> TagIds = TagCollection->GetTagIds();
+		TArray<FAvaTagId> TagIds = TagCollection->GetTagIds(TagCustomizer->AllowAliasTags());
 
 		// Add default (none) TagId as an option, if there's no multiple choice
 		if (!TagCustomizer->AllowMultipleTags())
@@ -299,7 +311,7 @@ FText SAvaTagPicker::GetValueDisplayText() const
 		{
 			if (!InStructRawData)
 			{
-				return true;	
+				return true;
 			}
 			FName CurrentValueName = TagCustomizer->GetDisplayValueName(InStructRawData);
 			if (InDataIndex == 0)
diff --git a/Engine/Plugins/Experimental/Avalanche/Source/AvalancheTagEditor/Private/Customization/TagCustomizers/AvaTagAliasCustomizer.cpp b/Engine/Plugins/Experimental/Avalanche/Source/AvalancheTagEditor/Private/Customization/TagCustomizers/AvaTagAliasCustomizer.cpp
new file mode 100644
--- /dev/null
+++ b/Engine/Plugins/Experimental/Avalanche/Source/AvalancheTagEditor/Private/Customization/TagCustomizers/AvaTagAliasCustomizer.cpp
@@ -0,0 +1,58 @@
+// Copyright Epic Games, Inc. All Rights Reserved.
+
+#include "AvaTagAliasCustomizer.h"
+#include "AvaTagAlias.h"
+#include "AvaTagHandle.h"
+#include "PropertyHandle.h"
+#include "ScopedTransaction.h"
+
+#define LOCTEXT_NAMESPACE "AvaTagAliasCustomizer"
+
+TSharedPtr<IPropertyHandle> FAvaTagAliasCustomizer::GetTagCollectionHandle(const TSharedRef<IPropertyHandle>& InStructHandle) const
+{
+	// There's no source property in Tag Alias so do not generate a source picker.
+	return nullptr;
+}
+
+const UAvaTagCollection* FAvaTagAliasCustomizer::GetOrLoadTagCollection(const void* InStructRawData) const
+{
+	return static_cast<const FAvaTagAlias*>(InStructRawData)->GetOwner();
+}
+
+void FAvaTagAliasCustomizer::SetTagHandleAdded(const TSharedRef<IPropertyHandle>& InContainerProperty, const FAvaTagHandle& InTagHandle, bool bInAdd) const
+{
+	FScopedTransaction Transaction(bInAdd
+		? LOCTEXT("AddTagHandleInAlias", "Add Tag Handle in Alias")
+		: LOCTEXT("RemoveTagHandleInAlias", "Remove Tag Handle from Alias"));
+
+	InContainerProperty->NotifyPreChange();
+	InContainerProperty->EnumerateRawData(
+		[&InTagHandle, bInAdd](void* InStructRawData, const int32, const int32)->bool
+		{
+			FAvaTagAlias* Alias = static_cast<FAvaTagAlias*>(InStructRawData);
+			if (bInAdd)
+			{
+				Alias->TagIds.AddUnique(InTagHandle.TagId);
+			}
+			else
+			{
+				Alias->TagIds.Remove(InTagHandle.TagId);
+			}
+			return true;
+		});
+
+	InContainerProperty->NotifyPostChange(bInAdd ? EPropertyChangeType::ArrayAdd : EPropertyChangeType::ArrayRemove);
+	InContainerProperty->NotifyFinishedChangingProperties();
+}
+
+bool FAvaTagAliasCustomizer::ContainsTagHandle(const void* InStructRawData, const FAvaTagHandle& InTagHandle) const
+{
+	return static_cast<const FAvaTagAlias*>(InStructRawData)->TagIds.Contains(InTagHandle.TagId);
+}
+
+FName FAvaTagAliasCustomizer::GetDisplayValueName(const void* InStructRawData) const
+{
+	return *static_cast<const FAvaTagAlias*>(InStructRawData)->GetTagsAsString();
+}
+
+#undef LOCTEXT_NAMESPACE
diff --git a/Engine/Plugins/Experimental/Avalanche/Source/AvalancheTagEditor/Private/Customization/TagCustomizers/AvaTagAliasCustomizer.h b/Engine/Plugins/Experimental/Avalanche/Source/AvalancheTagEditor/Private/Customization/TagCustomizers/AvaTagAliasCustomizer.h
new file mode 100644
--- /dev/null
+++ b/Engine/Plugins/Experimental/Avalanche/Source/AvalancheTagEditor/Private/Customization/TagCustomizers/AvaTagAliasCustomizer.h
@@ -0,0 +1,20 @@
+// Copyright Epic Games, Inc. All Rights Reserved.
+
+#pragma once
+
+#include "IAvaTagHandleCustomizer.h"
+
+/** Customizer for FAvaTagAlias */
+class FAvaTagAliasCustomizer : public IAvaTagHandleCustomizer
+{
+public:
+	//~ Begin IAvaTagHandleCustomizer
+	virtual TSharedPtr<IPropertyHandle> GetTagCollectionHandle(const TSharedRef<IPropertyHandle>& InStructHandle) const override;
+	virtual const UAvaTagCollection* GetOrLoadTagCollection(const void* InStructRawData) const override;
+	virtual void SetTagHandleAdded(const TSharedRef<IPropertyHandle>& InContainerProperty, const FAvaTagHandle& InTagHandle, bool bInAdd) const override;
+	virtual bool ContainsTagHandle(const void* InStructRawData, const FAvaTagHandle& InTagHandle) const override;
+	virtual FName GetDisplayValueName(const void* InStructRawData) const override;
+	virtual bool AllowMultipleTags() const override { return true; }
+	virtual bool AllowAliasTags() const override { return false; }
+	//~ End IAvaTagHandleCustomizer
+};
diff --git a/Engine/Plugins/Experimental/Avalanche/Source/AvalancheTagEditor/Private/Customization/TagCustomizers/IAvaTagHandleCustomizer.h b/Engine/Plugins/Experimental/Avalanche/Source/AvalancheTagEditor/Private/Customization/TagCustomizers/IAvaTagHandleCustomizer.h
--- a/Engine/Plugins/Experimental/Avalanche/Source/AvalancheTagEditor/Private/Customization/TagCustomizers/IAvaTagHandleCustomizer.h
+++ b/Engine/Plugins/Experimental/Avalanche/Source/AvalancheTagEditor/Private/Customization/TagCustomizers/IAvaTagHandleCustomizer.h
@@ -15,7 +15,10 @@ class IAvaTagHandleCustomizer
 public:
 	virtual ~IAvaTagHandleCustomizer() = default;
 
-	/** Gets the Child Property Handle to the Tag Collection Property */
+	/**
+	 * Gets the Child Property Handle to the Tag Collection Property
+	 * If returning nullptr, the picker will not generate a Tag Collection Source Picker
+	 */
 	virtual TSharedPtr<IPropertyHandle> GetTagCollectionHandle(const TSharedRef<IPropertyHandle>& InStructHandle) const = 0;
 
 	/** Gets or Loads the Tag Collection from the Raw Data Struct */
@@ -32,4 +35,7 @@ public:
 
 	/** Whether multiple tags are allowed to be selected, or just a single one */
 	virtual bool AllowMultipleTags() const { return false; }
+
+	/** Whether to show alias tags as options */
+	virtual bool AllowAliasTags() const { return true; }
 };
diff --git a/Engine/Plugins/Experimental/Avalanche/Source/AvalancheTransition/Private/Execution/AvaTransitionExecutor.cpp b/Engine/Plugins/Experimental/Avalanche/Source/AvalancheTransition/Private/Execution/AvaTransitionExecutor.cpp
--- a/Engine/Plugins/Experimental/Avalanche/Source/AvalancheTransition/Private/Execution/AvaTransitionExecutor.cpp
+++ b/Engine/Plugins/Experimental/Avalanche/Source/AvalancheTransition/Private/Execution/AvaTransitionExecutor.cpp
@@ -1,7 +1,9 @@
 // Copyright Epic Games, Inc. All Rights Reserved.
 
 #include "AvaTransitionExecutor.h"
+#include "AvaTag.h"
 #include "AvaTagHandleKeyFuncs.h"
+#include "AvaTagList.h"
 #include "AvaTransitionLayer.h"
 #include "AvaTransitionSubsystem.h"
 #include "Behavior/IAvaTransitionBehavior.h"
@@ -41,46 +43,66 @@ void FAvaTransitionExecutor::Setup()
 			InInstance.Setup();
 		});
 
-	using FAvaTagHandleMapType = TMap<FAvaTagHandle, EAvaTransitionType, FDefaultSetAllocator, TAvaTagHandleMapKeyFuncs<EAvaTransitionType, false>>;
+	struct FLayerInfo
+	{
+		// All the Behavior Instances found for a given Layer
+		TArray<FAvaTransitionBehaviorInstance*> BehaviorInstances;
 
-	FAvaTagHandleMapType LayerToTransitionTypeMap;
-	LayerToTransitionTypeMap.Reserve(Instances.Num());
+		// The accumulated Transition Type for a given Layer (e.g. combinations could be In, Out or In | Out)
+		EAvaTransitionType TransitionType;
+	};
 
-	// Gather Map of Layer to the Transition Types for that Layer
-	for (const FAvaTransitionBehaviorInstance& Instance : Instances)
+	// Map of BehaviorInstance to the Resolved Tags
+	TMap<FAvaTag, FLayerInfo> TagLayerInfo;
 	{
-		EAvaTransitionType& LayerTransitionType = LayerToTransitionTypeMap.FindOrAdd(Instance.GetTransitionLayer());
-		LayerTransitionType |= Instance.GetTransitionType();
+		// Rough Estimate: assume each Instance is in its own single unique tag
+		TagLayerInfo.Reserve(Instances.Num());
+
+		// Gather the Layer info for each Instance
+		for (FAvaTransitionBehaviorInstance& Instance : Instances)
+		{
+			for (const FAvaTag* Tag : Instance.GetTransitionLayer().GetTags())
+			{
+				FLayerInfo& LayerInfo = TagLayerInfo.FindOrAdd(*Tag);
+				LayerInfo.BehaviorInstances.AddUnique(&Instance);
+				LayerInfo.TransitionType |= Instance.GetTransitionType();
+			}
+		}
 	}
 
 	// Ensure there's an exiting null instance for every entering Transition Instance in a layer
-	for (const TPair<FAvaTagHandle, EAvaTransitionType>& Pair : LayerToTransitionTypeMap)
+	for (const TPair<FAvaTag, FLayerInfo>& Pair : TagLayerInfo)
 	{
-		if (Pair.Value == EAvaTransitionType::In && !EnumHasAnyFlags(Pair.Value, EAvaTransitionType::Out))
+		const FLayerInfo& LayerInfo = Pair.Value;
+		if (LayerInfo.TransitionType == EAvaTransitionType::In && !EnumHasAnyFlags(LayerInfo.TransitionType, EAvaTransitionType::Out))
 		{
+			check(!LayerInfo.BehaviorInstances.IsEmpty());
 			FAvaTransitionBehaviorInstance& NullInstanceCopy = Instances.Add_GetRef(NullInstance);
 			NullInstanceCopy.SetTransitionType(EAvaTransitionType::Out);
-			NullInstanceCopy.SetOverrideLayer(Pair.Key);
+			NullInstanceCopy.SetOverrideLayer(LayerInfo.BehaviorInstances[0]->GetTransitionLayer());
 			NullInstanceCopy.Setup();
 		}
 	}
 
 	// For the Instances that are going out, if they belong in the same Transition Layer as an Instance going In
 	// mark them as Needs Discard (this does not mean the scene will be discarded as there could be logic that reverts this flag)
-	for (FAvaTransitionBehaviorInstance& Instance : Instances)
+	for (const TPair<FAvaTag, FLayerInfo>& Pair : TagLayerInfo)
 	{
-		if (Instance.GetTransitionType() == EAvaTransitionType::In)
+		if (!EnumHasAnyFlags(Pair.Value.TransitionType, EAvaTransitionType::In))
 		{
 			continue;
 		}
 
-		Instance.SetTransitionType(EAvaTransitionType::Out);
+		for (FAvaTransitionBehaviorInstance* Instance : Pair.Value.BehaviorInstances)
+		{
+			if (Instance->GetTransitionType() == EAvaTransitionType::In)
+			{
+				continue;
+			}
 
-		const EAvaTransitionType LayerTransitionType = LayerToTransitionTypeMap[Instance.GetTransitionLayer()];
+			Instance->SetTransitionType(EAvaTransitionType::Out);
 
-		if (EnumHasAnyFlags(LayerTransitionType, EAvaTransitionType::In))
-		{
-			if (FAvaTransitionScene* TransitionScene = Instance.GetTransitionContext().GetTransitionScene())
+			if (FAvaTransitionScene* TransitionScene = Instance->GetTransitionContext().GetTransitionScene())
 			{
 				TransitionScene->SetFlags(EAvaTransitionSceneFlags::NeedsDiscard);
 			}
