Motion Design MRQ: Added sequence null check

diff --git a/Engine/Plugins/Experimental/Avalanche/Source/AvalancheMRQEditor/Private/AvaMRQEditorRundowntUtils.cpp b/Engine/Plugins/Experimental/Avalanche/Source/AvalancheMRQEditor/Private/AvaMRQEditorRundowntUtils.cpp
--- a/Engine/Plugins/Experimental/Avalanche/Source/AvalancheMRQEditor/Private/AvaMRQEditorRundowntUtils.cpp
+++ b/Engine/Plugins/Experimental/Avalanche/Source/AvalancheMRQEditor/Private/AvaMRQEditorRundowntUtils.cpp
@@ -105,6 +105,11 @@ namespace UE::AvaMRQEditor::Private
 
 	void RenderSequence(FAvaMRQScopedRender& InScopedRender, UWorld* InWorld, const UAvaRundown& InRundown, const FAvaRundownPage& InPage, UAvaSequence* InSequence)
 	{
+		if (!InSequence)
+		{
+			return;
+		}
+
 		UMoviePipelineExecutorJob* Job = InScopedRender.AllocateJob();
 		if (!Job)
 		{
