Motion Design Sequencer: Fix issue with keying nonActor/Comp objects via Auto-Key not properly working on first try

diff --git a/Engine/Plugins/Experimental/Avalanche/Source/AvalancheSequence/Private/AvaSequence.cpp b/Engine/Plugins/Experimental/Avalanche/Source/AvalancheSequence/Private/AvaSequence.cpp
--- a/Engine/Plugins/Experimental/Avalanche/Source/AvalancheSequence/Private/AvaSequence.cpp
+++ b/Engine/Plugins/Experimental/Avalanche/Source/AvalancheSequence/Private/AvaSequence.cpp
@@ -505,6 +505,11 @@ UObject* UAvaSequence::GetParentObject(UObject* InObject) const
 		: nullptr;
 }
 
+bool UAvaSequence::CanAnimateObject(UObject& InObject) const
+{
+	return Super::CanAnimateObject(InObject) || InObject.GetTypedOuter<AActor>();
+}
+
 void UAvaSequence::Serialize(FArchive& Ar)
 {
 	Ar.UsingCustomVersion(FAvaSequenceVersion::GUID);
diff --git a/Engine/Plugins/Experimental/Avalanche/Source/AvalancheSequence/Public/AvaSequence.h b/Engine/Plugins/Experimental/Avalanche/Source/AvalancheSequence/Public/AvaSequence.h
--- a/Engine/Plugins/Experimental/Avalanche/Source/AvalancheSequence/Public/AvaSequence.h
+++ b/Engine/Plugins/Experimental/Avalanche/Source/AvalancheSequence/Public/AvaSequence.h
@@ -119,6 +119,7 @@ public:
 	virtual UObject* CreateDirectorInstance(TSharedRef<const FSharedPlaybackState> SharedPlaybackState, FMovieSceneSequenceID InSequenceID) override;
 	virtual bool CanPossessObject(UObject& InObject, UObject* InPlaybackContext) const override;
 	virtual UObject* GetParentObject(UObject* InObject) const override;
+	virtual bool CanAnimateObject(UObject& InObject) const override;
 	//~ End UMovieSceneSequence
 
 	//~ Begin UObject
