Motion Design Transition: Set "motion design" DisplayName to blueprint task and condition

diff --git a/Engine/Plugins/Experimental/Avalanche/Source/AvalancheTransition/Public/Conditions/AvaTransitionConditionBlueprint.h b/Engine/Plugins/Experimental/Avalanche/Source/AvalancheTransition/Public/Conditions/AvaTransitionConditionBlueprint.h
--- a/Engine/Plugins/Experimental/Avalanche/Source/AvalancheTransition/Public/Conditions/AvaTransitionConditionBlueprint.h
+++ b/Engine/Plugins/Experimental/Avalanche/Source/AvalancheTransition/Public/Conditions/AvaTransitionConditionBlueprint.h
@@ -5,7 +5,7 @@
 #include "Blueprint/StateTreeConditionBlueprintBase.h"
 #include "AvaTransitionConditionBlueprint.generated.h"
 
-UCLASS(MinimalAPI, Abstract)
+UCLASS(MinimalAPI, Abstract, DisplayName="Motion Design Transition Condition Blueprint")
 class UAvaTransitionConditionBlueprint : public UStateTreeConditionBlueprintBase
 {
 	GENERATED_BODY()
diff --git a/Engine/Plugins/Experimental/Avalanche/Source/AvalancheTransition/Public/Tasks/AvaTransitionTaskBlueprint.h b/Engine/Plugins/Experimental/Avalanche/Source/AvalancheTransition/Public/Tasks/AvaTransitionTaskBlueprint.h
--- a/Engine/Plugins/Experimental/Avalanche/Source/AvalancheTransition/Public/Tasks/AvaTransitionTaskBlueprint.h
+++ b/Engine/Plugins/Experimental/Avalanche/Source/AvalancheTransition/Public/Tasks/AvaTransitionTaskBlueprint.h
@@ -5,7 +5,7 @@
 #include "Blueprint/StateTreeTaskBlueprintBase.h"
 #include "AvaTransitionTaskBlueprint.generated.h"
 
-UCLASS(MinimalAPI, Abstract)
+UCLASS(MinimalAPI, Abstract, DisplayName="Motion Design Transition Task Blueprint")
 class UAvaTransitionTaskBlueprint : public UStateTreeTaskBlueprintBase
 {
 	GENERATED_BODY()
