Motion Design: Move widgets that modify struct data to the value widget rather than name widget.

diff --git a/Engine/Plugins/Experimental/Avalanche/Source/AvalancheSequencer/Private/Customization/AvaSequenceTimeCustomization.cpp b/Engine/Plugins/Experimental/Avalanche/Source/AvalancheSequencer/Private/Customization/AvaSequenceTimeCustomization.cpp
--- a/Engine/Plugins/Experimental/Avalanche/Source/AvalancheSequencer/Private/Customization/AvaSequenceTimeCustomization.cpp
+++ b/Engine/Plugins/Experimental/Avalanche/Source/AvalancheSequencer/Private/Customization/AvaSequenceTimeCustomization.cpp
@@ -8,77 +8,89 @@
 
 void FAvaSequenceTimeCustomization::CustomizeHeader(TSharedRef<IPropertyHandle> InPropertyHandle, FDetailWidgetRow& InHeaderRow, IPropertyTypeCustomizationUtils& InCustomizationUtils)
 {
-	TSharedPtr<IPropertyHandle> PositionTypeHandle = InPropertyHandle->GetChildHandle(GET_MEMBER_NAME_CHECKED(FAvaSequenceTime, TimeType));
+	auto CreatePropertySlotWidget =
+		[InPropertyHandle](FName InPropertyName)->TSharedRef<SWidget>
+		{
+			TSharedPtr<IPropertyHandle> ChildPropertyHandle = InPropertyHandle->GetChildHandle(InPropertyName);
 
-	TSharedRef<SHorizontalBox> PropertyValueBox = SNew(SHorizontalBox);
+			TSharedRef<SWidget> PropertyValueWidget = ChildPropertyHandle->CreatePropertyValueWidget();
 
-	auto AddPropertySlot = [PropertyValueBox, InPropertyHandle](FName InPropertyName)
-	{
-		TSharedPtr<IPropertyHandle> ChildPropertyHandle = InPropertyHandle->GetChildHandle(InPropertyName);
+			PropertyValueWidget->SetVisibility(TAttribute<EVisibility>::Create(TAttribute<EVisibility>::FGetter::CreateLambda(
+				[ChildPropertyHandle]
+				{
+					return ChildPropertyHandle->IsEditable()
+						? EVisibility::SelfHitTestInvisible
+						: EVisibility::Collapsed;
+				})));
 
-		TSharedRef<SWidget> PropertyValueWidget = ChildPropertyHandle->CreatePropertyValueWidget();
+			return PropertyValueWidget;
+		};
 
-		TAttribute<EVisibility>::FGetter VisibilityGetter = TAttribute<EVisibility>::FGetter::CreateLambda([ChildPropertyHandle]
-			{
-				return ChildPropertyHandle->IsEditable()
-					? EVisibility::SelfHitTestInvisible
-					: EVisibility::Collapsed;
-			});
-
-		PropertyValueWidget->SetVisibility(TAttribute<EVisibility>::Create(VisibilityGetter));
-
-		PropertyValueBox->AddSlot()
-			.VAlign(VAlign_Center)
-			.HAlign(HAlign_Fill)
-			[
-				PropertyValueWidget
-			];
-	};
-
-	AddPropertySlot(GET_MEMBER_NAME_CHECKED(FAvaSequenceTime, Frame));
-	AddPropertySlot(GET_MEMBER_NAME_CHECKED(FAvaSequenceTime, Seconds));
-	AddPropertySlot(GET_MEMBER_NAME_CHECKED(FAvaSequenceTime, MarkLabel));
-
-	PropertyValueBox->AddSlot()
-		.AutoWidth()
-		[
-			InPropertyHandle->CreateDefaultPropertyButtonWidgets()
-		];
+	constexpr bool bDisplayDefaultPropertyButtons = false;
 
+	TSharedPtr<IPropertyHandle> TimeTypeHandle = InPropertyHandle->GetChildHandle(GET_MEMBER_NAME_CHECKED(FAvaSequenceTime, TimeType));
 	TSharedPtr<IPropertyHandle> HasTimeHandle = InPropertyHandle->GetChildHandle(GET_MEMBER_NAME_CHECKED(FAvaSequenceTime, bHasTimeConstraint));
 
-	static const bool bDisplayDefaultPropertyButtons = false;
-
 	InHeaderRow
 		.NameContent()
 		[
-			SNew(SBox)
-			.MinDesiredWidth(300.f)
-			.MaxDesiredWidth(300.f)
+			InPropertyHandle->CreatePropertyNameWidget()
+		]
+		.ValueContent()
+		[
+			SNew(SVerticalBox)
+			+ SVerticalBox::Slot()
+			.VAlign(VAlign_Center)
+			.HAlign(HAlign_Fill)
+			.AutoHeight()
 			[
 				SNew(SHorizontalBox)
 				+ SHorizontalBox::Slot()
+				.VAlign(VAlign_Center)
+				.HAlign(HAlign_Left)
 				.AutoWidth()
 				[
 					HasTimeHandle->CreatePropertyValueWidget(bDisplayDefaultPropertyButtons)
 				]
 				+ SHorizontalBox::Slot()
-				.AutoWidth()
-				.Padding(2.f, 0.f, 0.f, 0.f)
+				.VAlign(VAlign_Center)
+				.HAlign(HAlign_Fill)
+				.FillWidth(1.f)
 				[
-					InPropertyHandle->CreatePropertyNameWidget()
+					TimeTypeHandle->CreatePropertyValueWidget(bDisplayDefaultPropertyButtons)
 				]
 				+ SHorizontalBox::Slot()
-				.FillWidth(1.f)
-				.Padding(0.f, 0.f, 5.f, 0.f)
+				.VAlign(VAlign_Center)
 				.HAlign(HAlign_Right)
+				.AutoWidth()
 				[
-					PositionTypeHandle->CreatePropertyValueWidget(bDisplayDefaultPropertyButtons)
+					InPropertyHandle->CreateDefaultPropertyButtonWidgets()
+				]
+			]
+			+ SVerticalBox::Slot()
+			.VAlign(VAlign_Center)
+			.HAlign(HAlign_Fill)
+			.AutoHeight()
+			[
+				SNew(SOverlay)
+				+ SOverlay::Slot()
+				.VAlign(VAlign_Center)
+				.HAlign(HAlign_Fill)
+				[
+					CreatePropertySlotWidget(GET_MEMBER_NAME_CHECKED(FAvaSequenceTime, Frame))
+				]
+				+ SOverlay::Slot()
+				.VAlign(VAlign_Center)
+				.HAlign(HAlign_Fill)
+				[
+					CreatePropertySlotWidget(GET_MEMBER_NAME_CHECKED(FAvaSequenceTime, Seconds))
+				]
+				+ SOverlay::Slot()
+				.VAlign(VAlign_Center)
+				.HAlign(HAlign_Fill)
+				[
+					CreatePropertySlotWidget(GET_MEMBER_NAME_CHECKED(FAvaSequenceTime, MarkLabel))
 				]
 			]
-		]
-		.ValueContent()
-		[
-			PropertyValueBox
 		];
 }
