[Backout] - CL34065553
Original CL Desc
-----------------------------------------------------------------
Sequencer- Dynamic Binding Refactor
Refactor the 'Dynamic Binding' blueprint system to be Custom Bindings rather than tied directly to FMovieScenePossessable and FMovieSceneSpawnable. As part of this:
? Created UMovieSceneReplaceableDirectorBlueprintBinding and UMovieSceneSpawnableDirectorBlueprintBinding to replace the Possessable and Spawnable versions of DynamicBindings
? Deprecated the DynamicBinding references on possessable and spawnable and created Postload fixup code to automatically create custom bindings instead
? Note that because the previous Dynamic Bindings operated as overrides and the new ones don't, there is a no longer supported path which is using Possessables as a preview in editor for Possessable Dynamic Bindings. Now if you don't mark 'Call in Editor' for a Replaceable Director Blueprint binding, you'll need to define a spawnable preview binding for the binding in order to work with it in Sequencer.
? Since UMG also used Dynamic Bindings on possessables, temporarily add DynamicBinding as a member variable on FWidgetAnimationBinding and create postload logic to move DynamicBindings to live on there for WidgetAnimations.
? USD has custom code to create dynamic bindings automatically for its use case. Convert this also to use the new custom binding system.
? Some cleanup on object binding right-click menus was done as part of this.
? Some further MovieSceneSequence API deprecation and refactoring has been done for the creation of bindings to add SharedPlaybackState where possible which was necessary for dynamic binding setup.

FMovieSceneSpawnable deprecation in Level Sequences
Since I was having to write postload fixup logic for FMovieSceneSpawnables with Dynamic Bindings that converted these to FMovieScenePossessables with custom bindings, it felt like a good chance to deprecate all FMovieSceneSpawnable use in Level Sequences.

I chose at this time not to deprecate FMovieSceneSpawnable use for non-Level-Sequence Movie Scene Sequences. This includes Day Sequences, Template Sequences, Actor Sequences, Widget Animations, Niagara Sequences, and some Meta Human sequences. I made this choice because to do so would require also refactoring those sequence types to use MovieSceneBindingReferences in order to enable custom bindings for those sequence types. This will be a future refactor.

Some notes on this deprecation.
? I kept MovieScene specific code that references or creates FMovieSceneSpawnables due to not refactoring all MovieSceneSequences
? Binding creation code in tools such as Sequencer will check for the presence of Binding References and create custom bindings if present, otherwise will still create FMovieSceneSpawnables
? AvaSequences required some extra work to make this change, including some changes and bugfixes to copy/pasting bindings to move away from custom avalanche copy/paste code, as the copy/paste code for custom bindings is a bit intricate and I didn't want to duplicate it.
? Chaos Cache has been moved away from using an object spawner and instead uses a custom spawnable actor binding inheriting from the regular spawnable actor binding.
? Actor and Take Recording has been updated to support custom spawnable bindings
? Meta Human specific code has been updated to support custom spawnable bindings when using Level Sequences
? FortniteMainBranchObjectVersion has been bumped

Crash fix
As part of the USD work, I found a crash bug where the Binding Properties menu would hold a reference to a level sequence. In the case of USD, the outer of the sequences can be the level, and this caused issues when switching levels. To fix this I ended up needing to add some events when menus are being dismissed so I could clean up these references.

diff --git a/Engine/Plugins/Animation/ControlRig/Source/ControlRigEditor/Private/ControlRigEditorModule.cpp b/Engine/Plugins/Animation/ControlRig/Source/ControlRigEditor/Private/ControlRigEditorModule.cpp
--- a/Engine/Plugins/Animation/ControlRig/Source/ControlRigEditor/Private/ControlRigEditorModule.cpp
+++ b/Engine/Plugins/Animation/ControlRig/Source/ControlRigEditor/Private/ControlRigEditorModule.cpp
@@ -97,8 +97,6 @@
 #include "Constraints/TransformConstraintChannelInterface.h"
 #include "Async/TaskGraphInterfaces.h"
 #include "AssetRegistry/IAssetRegistry.h"
-#include "SequencerUtilities.h"
-#include "Bindings/MovieSceneSpawnableActorBinding.h"
 
 #define LOCTEXT_NAMESPACE "ControlRigEditorModule"
 
@@ -1165,10 +1163,10 @@ void FControlRigEditorModule::BakeToControlRig(UClass* ControlRigClass, UAnimSeq
 			// By default, convert this to a spawnable and delete the existing actor. If for some reason, 
 			// the spawnable couldn't be generated, use the existing actor as a possessable (this could 
 			// eventually be an option)
-
-			if (FMovieScenePossessable* Possessable = FSequencerUtilities::ConvertToCustomBinding(SequencerPtr.ToSharedRef(), ActorTrackGuid, UMovieSceneSpawnableActorBinding::StaticClass(), 0))
-			{ 
-				ActorTrackGuid = Possessable->GetGuid();
+			TArray<FGuid> SpawnableGuids = SequencerPtr->ConvertToSpawnable(ActorTrackGuid);
+			if (SpawnableGuids.Num())
+			{	
+				ActorTrackGuid = SpawnableGuids[0];
 
 				UObject* SpawnedMesh = SequencerPtr->FindSpawnedObjectOrTemplate(ActorTrackGuid);
 
diff --git a/Engine/Plugins/Animation/ControlRig/Source/ControlRigEditor/Private/Sequencer/ControlRigParameterTrackEditor.cpp b/Engine/Plugins/Animation/ControlRig/Source/ControlRigEditor/Private/Sequencer/ControlRigParameterTrackEditor.cpp
--- a/Engine/Plugins/Animation/ControlRig/Source/ControlRigEditor/Private/Sequencer/ControlRigParameterTrackEditor.cpp
+++ b/Engine/Plugins/Animation/ControlRig/Source/ControlRigEditor/Private/Sequencer/ControlRigParameterTrackEditor.cpp
@@ -4184,12 +4184,9 @@ bool FControlRigParameterTrackEditor::HandleAssetAdded(UObject* Asset, const FGu
 
 	const FScopedTransaction Transaction(LOCTEXT("AddControlRigAsset", "Add Control Rig"));
 
-	UE::Sequencer::FCreateBindingParams CreateBindingParams;
-	CreateBindingParams.bSpawnable = true;
-	CreateBindingParams.bAllowCustomBinding = true;
-	FGuid NewGuid = GetSequencer()->CreateBinding(*ASkeletalMeshActor::StaticClass(), CreateBindingParams);
+	FGuid NewGuid = GetSequencer()->MakeNewSpawnable(*ASkeletalMeshActor::StaticClass());
 
-	// CreateBinding can fail if spawnables are not allowed
+	// MakeNewSpawnable can fail if spawnables are not allowed
 	if (!NewGuid.IsValid())
 	{
 		return false;
@@ -4208,7 +4205,8 @@ bool FControlRigParameterTrackEditor::HandleAssetAdded(UObject* Asset, const FGu
 	
 	// Save Spawnable state as the default (with new name and skeletal mesh asset)
 	{
-		GetSequencer()->GetSpawnRegister().SaveDefaultSpawnableState(NewGuid, GetSequencer()->GetFocusedTemplateID(), GetSequencer()->GetSharedPlaybackState());
+		FMovieSceneSpawnable* Spawnable = MovieScene->FindSpawnable(NewGuid);
+		GetSequencer()->GetSpawnRegister().SaveDefaultSpawnableState(*Spawnable, GetSequencer()->GetFocusedTemplateID(), *GetSequencer());
 	}
 
 	UMovieSceneControlRigParameterTrack* Track = Cast<UMovieSceneControlRigParameterTrack>(MovieScene->FindTrack(UMovieSceneControlRigParameterTrack::StaticClass(), NewGuid, NAME_None));
diff --git a/Engine/Plugins/Experimental/Avalanche/Source/AvalancheSequence/Private/AvaSequence.cpp b/Engine/Plugins/Experimental/Avalanche/Source/AvalancheSequence/Private/AvaSequence.cpp
--- a/Engine/Plugins/Experimental/Avalanche/Source/AvalancheSequence/Private/AvaSequence.cpp
+++ b/Engine/Plugins/Experimental/Avalanche/Source/AvalancheSequence/Private/AvaSequence.cpp
@@ -429,7 +429,7 @@ TArray<UObject*> UAvaSequence::GetBoundObjects(UObject* InPlaybackContext) const
 	for (const FMovieSceneBinding& Binding : MovieScene->GetBindings())
 	{
 		TArray<UObject*, TInlineAllocator<1>> BoundObject;
-		LocateBoundObjects(Binding.GetObjectGuid(), UE::UniversalObjectLocator::FResolveParams(InPlaybackContext), MovieSceneHelpers::CreateTransientSharedPlaybackState(InPlaybackContext, MovieScene->GetTypedOuter<UMovieSceneSequence>()), BoundObject);
+		LocateBoundObjects(Binding.GetObjectGuid(), UE::UniversalObjectLocator::FResolveParams(InPlaybackContext), BoundObject);
 		OutObjects.Append(BoundObject);
 	}
 
diff --git a/Engine/Plugins/Experimental/Avalanche/Source/AvalancheSequencer/Private/AvaSequencer.cpp b/Engine/Plugins/Experimental/Avalanche/Source/AvalancheSequencer/Private/AvaSequencer.cpp
--- a/Engine/Plugins/Experimental/Avalanche/Source/AvalancheSequencer/Private/AvaSequencer.cpp
+++ b/Engine/Plugins/Experimental/Avalanche/Source/AvalancheSequencer/Private/AvaSequencer.cpp
@@ -1199,7 +1199,7 @@ UObject* FAvaSequencer::FindResolutionContext(UAvaSequence& InSequence
 		{
 			TArray<UObject*, TInlineAllocator<1>> BoundObjects;
 
-			InSequence.LocateBoundObjects(InGuid, UE::UniversalObjectLocator::FResolveParams(InContextChecked), MovieSceneHelpers::CreateTransientSharedPlaybackState(InContextChecked, &InSequence), BoundObjects);
+			InSequence.LocateBoundObjects(InGuid, UE::UniversalObjectLocator::FResolveParams(InContextChecked), BoundObjects);
 			return BoundObjects;
 		};
 
diff --git a/Engine/Plugins/Experimental/Avalanche/Source/AvalancheSequencer/Private/Clipboard/AvaSequenceCopyableBinding.h b/Engine/Plugins/Experimental/Avalanche/Source/AvalancheSequencer/Private/Clipboard/AvaSequenceCopyableBinding.h
new file mode 100644
--- /dev/null
+++ b/Engine/Plugins/Experimental/Avalanche/Source/AvalancheSequencer/Private/Clipboard/AvaSequenceCopyableBinding.h
@@ -0,0 +1,51 @@
+// Copyright Epic Games, Inc. All Rights Reserved.
+
+#pragma once
+
+#include "UObject/ObjectMacros.h"
+#include "MovieScene.h"
+#include "AvaSequenceCopyableBinding.generated.h"
+
+UCLASS()
+class UAvaSequenceCopyableBinding : public UObject
+{
+	GENERATED_BODY()
+
+public:
+	/**
+	 * Spawnables need to know about their Object Template but we cannot rely on automatic serialization due to the object
+	 * template belonging to the Movie Scene (it gets serialized as a reference). Instead we manually serialize the object
+	 * so that we can duplicate it into a new object (which is stored in this variable) but we don't want this exported with
+	 * the rest of the text as it'll fall back to the same reference issue. Marking this as TextExportTransient solves this.
+	 */
+	UPROPERTY(TextExportTransient)
+	TObjectPtr<UObject> SpawnableObjectTemplate;
+
+	/**
+	 * Tracks are also owned by the owning Movie Sequence. We manually copy the tracks out of a binding when we copy,
+	 * because the binding stores them as a reference to a privately owned object. We store these copied tracks here,
+	 * and then restore them upon paste to re-create the tracks with the correct owner.
+	 */
+	UPROPERTY()
+	TArray<TObjectPtr<UMovieSceneTrack>> Tracks;
+
+	UPROPERTY()
+	FMovieSceneBinding Binding;
+
+	UPROPERTY()
+	FMovieSceneSpawnable Spawnable;
+
+	UPROPERTY()
+	FMovieScenePossessable Possessable;
+
+	/** FNames of the Actor Names bound. Used for lookup on Import */
+	UPROPERTY()
+	TArray<FName> BoundActorNames;
+
+	/** Path of the Subobjects bound relative to the Resolution Context. Used only if Resolution Context is NOT the Playback Context */
+	UPROPERTY()
+	TArray<FString> BoundObjectPaths;
+
+	UPROPERTY()
+	TArray<FName> FolderPath;
+};
diff --git a/Engine/Plugins/Experimental/Avalanche/Source/AvalancheSequencer/Private/Clipboard/AvaSequenceExporter.cpp b/Engine/Plugins/Experimental/Avalanche/Source/AvalancheSequencer/Private/Clipboard/AvaSequenceExporter.cpp
--- a/Engine/Plugins/Experimental/Avalanche/Source/AvalancheSequencer/Private/Clipboard/AvaSequenceExporter.cpp
+++ b/Engine/Plugins/Experimental/Avalanche/Source/AvalancheSequencer/Private/Clipboard/AvaSequenceExporter.cpp
@@ -1,6 +1,7 @@
 // Copyright Epic Games, Inc. All Rights Reserved.
 
 #include "AvaSequenceExporter.h"
+#include "AvaSequenceCopyableBinding.h"
 #include "AvaSequencer.h"
 #include "GameFramework/Actor.h"
 #include "MovieScene.h"
@@ -8,15 +9,14 @@
 #include "AvaSequence.h"
 #include "UnrealExporter.h"
 #include "UObject/Package.h"
-#include "SequencerUtilities.h"
 
 namespace UE::AvaSequencer::Private
 {
 	class FAvaSequenceExportObjectInnerContext : public FExportObjectInnerContext
 	{
 	public:
-		explicit FAvaSequenceExportObjectInnerContext(const TSharedRef<FAvaSequencer>& InAvaSequencer)
-			: AvaSequencerWeak(InAvaSequencer)
+		explicit FAvaSequenceExportObjectInnerContext(UObject* InPlaybackContext)
+			: PlaybackContext(InPlaybackContext)
 		{
 		}
 		virtual ~FAvaSequenceExportObjectInnerContext() override = default;
@@ -32,14 +32,12 @@ namespace UE::AvaSequencer::Private
 
 		TConstArrayView<AActor*> GetBoundActors() const { return BoundActors; }
 
-		UObject* GetPlaybackContext() const { return AvaSequencerWeak.IsValid() ? AvaSequencerWeak.Pin()->GetPlaybackContext() : nullptr; }
-
-		TSharedPtr<FAvaSequencer> GetSequencer() const { return AvaSequencerWeak.Pin(); }
+		UObject* GetPlaybackContext() const { return PlaybackContext; }
 
 	private:
-		TArray<TObjectPtr<AActor>> BoundActors;
+		TObjectPtr<UObject> PlaybackContext;
 
-		TWeakPtr<FAvaSequencer> AvaSequencerWeak;
+		TArray<TObjectPtr<AActor>> BoundActors;
 	};
 }
 
@@ -75,7 +73,7 @@ void FAvaSequenceExporter::ExportText(FString& InOutCopiedData, TConstArrayView<
 	}
 
 	UObject* const PlaybackContext = AvaSequencer->GetPlaybackContext();
-	UE::AvaSequencer::Private::FAvaSequenceExportObjectInnerContext ExportContext(AvaSequencer.ToSharedRef());
+	UE::AvaSequencer::Private::FAvaSequenceExportObjectInnerContext ExportContext(PlaybackContext);
 
 	const TCHAR* const Filetype = TEXT("copy");
 	const uint32 PortFlags = PPF_DeepCompareInstances | PPF_ExportsNotFullyQualified;
@@ -120,14 +118,6 @@ bool UAvaSequenceExporter::ExportText(const FExportObjectInnerContext* InContext
 	const UE::AvaSequencer::Private::FAvaSequenceExportObjectInnerContext& Context
 		= static_cast<const UE::AvaSequencer::Private::FAvaSequenceExportObjectInnerContext&>(*InContext);
 
-	TSharedPtr<FAvaSequencer> Sequencer = Context.GetSequencer();
-	 
-	if (!Sequencer)
-	{
-		return false;
-	}
-
-
 	UObject* const PlaybackContext = Context.GetPlaybackContext();
 
 	// Gather guids for the object nodes and any child object nodes
@@ -172,10 +162,123 @@ bool UAvaSequenceExporter::ExportText(const FExportObjectInnerContext* InContext
 		, FCString::Spc(TextIndent)
 		, *Sequence->GetLabel().ToString());
 
-	TArray<UMovieSceneFolder*> Folders;
-	FSequencerUtilities::CopyBindings(Sequencer->GetSequencer(), Bindings, Folders, Ar);
+	UObject* const CopyableBindingOuter = GetTransientPackage();
+
+	TArray<UAvaSequenceCopyableBinding*> ObjectsToExport;
+	ObjectsToExport.Reserve(Bindings.Num());
+
+	for (const FMovieSceneBindingProxy& ObjectBinding : Bindings)
+	{
+		UMovieScene* const MovieScene = ObjectBinding.GetMovieScene();
+		if (!IsValid(MovieScene))
+		{
+			continue;
+		}
+
+		UAvaSequenceCopyableBinding* const CopyableBinding = NewObject<UAvaSequenceCopyableBinding>(CopyableBindingOuter, NAME_None, RF_Transient);
+
+		ObjectsToExport.Add(CopyableBinding);
+
+		if (FMovieScenePossessable* const Possessable = MovieScene->FindPossessable(ObjectBinding.BindingID))
+		{
+			UObject* ResolutionContext = FAvaSequencer::FindResolutionContext(*Sequence
+				, *MovieScene
+				, Possessable->GetParent()
+				, PlaybackContext);
+
+			CopyableBinding->Possessable = *Possessable;
+
+			TArray<UObject*, TInlineAllocator<1>> BoundObjects;
+			Sequence->LocateBoundObjects(CopyableBinding->Possessable.GetGuid(), UE::UniversalObjectLocator::FResolveParams(ResolutionContext), BoundObjects);
+
+			// Store the names of the bound objects so that they can be found on paste
+			for (UObject* const BoundObject : BoundObjects)
+			{
+				if (!BoundObject)
+				{
+					continue;
+				}
+
+				if (AActor* const Actor = Cast<AActor>(BoundObject))
+				{
+					CopyableBinding->BoundActorNames.Add(Actor->GetFName());
+				}
+				else if (ResolutionContext != PlaybackContext)
+				{
+					CopyableBinding->BoundObjectPaths.Add(BoundObject->GetPathName(ResolutionContext));
+				}
+			}
+		}
+		else
+		{
+			FMovieSceneSpawnable* const Spawnable = MovieScene->FindSpawnable(ObjectBinding.BindingID);
+			if (Spawnable)
+			{
+				CopyableBinding->Spawnable = *Spawnable;
+
+				// We manually serialize the spawnable object template so that it's not a reference to a privately owned object. Spawnables all have unique copies of their template objects anyways.
+				// Object Templates are re-created on paste (based on these templates) with the correct ownership set up.
+				CopyableBinding->SpawnableObjectTemplate = Spawnable->GetObjectTemplate();
+			}
+		}
+
+		if (const FMovieSceneBinding* const Binding = MovieScene->FindBinding(ObjectBinding.BindingID))
+		{
+			CopyableBinding->Binding = *Binding;
+			for (UMovieSceneTrack* const Track : Binding->GetTracks())
+			{
+				// Tracks suffer from the same issues as Spawnable's Object Templates (reference to a privately owned object). We'll manually serialize the tracks to copy them,
+				// and then restore them on paste.
+				UMovieSceneTrack* const DuplicatedTrack = Cast<UMovieSceneTrack>(StaticDuplicateObject(Track, CopyableBinding));
+				CopyableBinding->Tracks.Add(DuplicatedTrack);
+			}
+		}
+	}
+
+	UAvaSequenceExporter::ExportBindings(ObjectsToExport, TEXT("copy"), TextIndent + 3, Ar, InWarn);
 
 	Ar.Logf(TEXT("%sEnd Sequence\r\n"), FCString::Spc(TextIndent));
 
 	return true;
 }
+
+void UAvaSequenceExporter::ExportBindings(const TArray<UAvaSequenceCopyableBinding*>& InObjectsToExport, const TCHAR* InType
+	, int32 InTextIndent
+	, FOutputDevice& Ar
+	, FFeedbackContext* InWarn)
+{
+	// Clear the mark state for saving.
+	UnMarkAllObjects(static_cast<EObjectMark>(EObjectMark::OBJECTMARK_TagExp | EObjectMark::OBJECTMARK_TagImp));
+
+	UObject* const Outer = GetTransientPackage();
+
+	const uint32 PortFlags = PPF_ExportsNotFullyQualified | PPF_Copy | PPF_Delimited;
+
+	FExportObjectInnerContext EmptyContext;
+
+	for (UAvaSequenceCopyableBinding* const ObjectToExport : InObjectsToExport)
+	{
+		// We can't use TextExportTransient on USTRUCTS (which our object contains) so we're going to manually null out some references before serializing them. These references are
+		// serialized manually into the archive, as the auto-serialization will only store a reference (to a privately owned object) which creates issues on deserialization. Attempting 
+		// to deserialize these private objects throws a superflous error in the console that makes it look like things went wrong when they're actually OK and expected.
+		TArray<UMovieSceneTrack*> OldTracks = ObjectToExport->Binding.StealTracks(nullptr);
+
+		UObject* const OldSpawnableTemplate = ObjectToExport->Spawnable.GetObjectTemplate();
+
+		ObjectToExport->Spawnable.SetObjectTemplate(nullptr);
+
+		UExporter::ExportToOutputDevice(&EmptyContext, ObjectToExport, nullptr, Ar, InType
+			, InTextIndent, PortFlags, false, Outer);
+
+		// Restore the references (as we don't want to modify the original in the event of a copy operation!)
+		ObjectToExport->Binding.SetTracks(MoveTemp(OldTracks), nullptr);
+		ObjectToExport->Spawnable.SetObjectTemplate(OldSpawnableTemplate);
+
+		// We manually export the object template for the same private-ownership reason as above. Templates need to be re-created anyways as each Spawnable contains its own copy of the template.
+		if (ObjectToExport->SpawnableObjectTemplate)
+		{
+			UExporter::ExportToOutputDevice(&EmptyContext, ObjectToExport->SpawnableObjectTemplate, nullptr, Ar, InType
+				, InTextIndent, PortFlags);
+		}
+	}
+}
diff --git a/Engine/Plugins/Experimental/Avalanche/Source/AvalancheSequencer/Private/Clipboard/AvaSequenceExporter.h b/Engine/Plugins/Experimental/Avalanche/Source/AvalancheSequencer/Private/Clipboard/AvaSequenceExporter.h
--- a/Engine/Plugins/Experimental/Avalanche/Source/AvalancheSequencer/Private/Clipboard/AvaSequenceExporter.h
+++ b/Engine/Plugins/Experimental/Avalanche/Source/AvalancheSequencer/Private/Clipboard/AvaSequenceExporter.h
@@ -35,4 +35,8 @@ public:
 	virtual bool ExportText(const FExportObjectInnerContext* InContext, UObject* InObject, const TCHAR* InType
 		, FOutputDevice& Ar, FFeedbackContext* InWarn, uint32 InPortFlags) override;
 	//~ End UExporter
+
+private:
+	static void ExportBindings(const TArray<UAvaSequenceCopyableBinding*>& InObjectsToExport, const TCHAR* InType
+		, int32 InTextIndent, FOutputDevice& Ar, FFeedbackContext* InWarn);
 };
\ No newline at end of file
diff --git a/Engine/Plugins/Experimental/Avalanche/Source/AvalancheSequencer/Private/Clipboard/AvaSequenceImporter.cpp b/Engine/Plugins/Experimental/Avalanche/Source/AvalancheSequencer/Private/Clipboard/AvaSequenceImporter.cpp
--- a/Engine/Plugins/Experimental/Avalanche/Source/AvalancheSequencer/Private/Clipboard/AvaSequenceImporter.cpp
+++ b/Engine/Plugins/Experimental/Avalanche/Source/AvalancheSequencer/Private/Clipboard/AvaSequenceImporter.cpp
@@ -1,6 +1,7 @@
 // Copyright Epic Games, Inc. All Rights Reserved.
 
 #include "AvaSequenceImporter.h"
+#include "AvaSequenceCopyableBinding.h"
 #include "AvaSequence.h"
 #include "AvaSequencer.h"
 #include "EngineUtils.h"
@@ -10,8 +11,6 @@
 #include "MovieScenePossessable.h"
 #include "SequencerUtilities.h"
 #include "UObject/Package.h"
-#include "Widgets/Notifications/SNotificationList.h"
-#include "Algo/Transform.h"
 
 FAvaSequenceImporter::FAvaSequenceImporter(const TSharedRef<FAvaSequencer>& InAvaSequencer)
 	: AvaSequencerWeak(InAvaSequencer)
@@ -68,11 +67,8 @@ void FAvaSequenceImporter::ImportText(FStringView InPastedData, const TMap<FName
 			BindingsString += TEXT("\r\n");
 		}
 
-		TArray<FNotificationInfo> PasteErrors; 
-		TArray<FMovieSceneBindingProxy> OutBindings;
-		FMovieScenePasteBindingsParams PasteParams;
-		Algo::Transform(InPastedActors, PasteParams.PastedActors, [](TTuple<FName, AActor*> Pair) { return TTuple<FName, TObjectPtr<AActor>>(Pair.Key, Pair.Value);});
-		FSequencerUtilities::PasteBindings(BindingsString, Sequencer, PasteParams, OutBindings, PasteErrors);
+		const TArray<UAvaSequenceCopyableBinding*> ImportedBindings = ImportBindings(*Sequencer, BindingsString);
+		ProcessImportedBindings(ImportedBindings, InPastedActors);
 	}
 
 	AvaSequencer->SetViewedSequence(OriginallyViewedSequence);
@@ -137,3 +133,427 @@ UAvaSequence* FAvaSequenceImporter::GetOrCreateSequence(FAvaSequencer& InAvaSequ
 
 	return NewSequence;
 }
+
+TArray<UAvaSequenceCopyableBinding*> FAvaSequenceImporter::ImportBindings(ISequencer& InSequencerRef, const FString& InBindingsString)
+{
+	class FAvaObjectBindingTextFactory : public FCustomizableTextObjectFactory
+	{
+	public:
+		FAvaObjectBindingTextFactory(ISequencer& InSequencer) : FCustomizableTextObjectFactory(GWarn), Sequencer(InSequencer)
+		{
+		}
+
+		virtual bool CanCreateClass(UClass* InObjectClass, bool& bOmitSubObjs) const override
+		{
+			if (InObjectClass->IsChildOf<UAvaSequenceCopyableBinding>())
+			{
+				return true;
+			}
+			return Sequencer.GetSpawnRegister().CanSpawnObject(InObjectClass);
+		}
+
+		virtual void ProcessConstructedObject(UObject* InObject) override
+		{
+			check(InObject);
+
+			if (InObject->IsA<UAvaSequenceCopyableBinding>())
+			{
+				UAvaSequenceCopyableBinding* CopyableBinding = Cast<UAvaSequenceCopyableBinding>(InObject);
+				ImportedBindings.Add(CopyableBinding);
+			}
+			else
+			{
+				NewSpawnableObjectTemplates.Add(InObject);
+			}
+		}
+
+		TArray<UAvaSequenceCopyableBinding*> ImportedBindings;
+		TArray<UObject*> NewSpawnableObjectTemplates;
+
+	private:
+		ISequencer& Sequencer;
+	};
+
+	UPackage* const TempPackage = NewObject<UPackage>(nullptr, TEXT("/Engine/Editor/AvalancheSequencer/Transient"), RF_Transient);
+	TempPackage->AddToRoot();
+
+	// Turn the text buffer into objects
+	FAvaObjectBindingTextFactory Factory(InSequencerRef);
+	Factory.ProcessBuffer(TempPackage, RF_Transactional, InBindingsString);
+
+	// We had to explicitly serialize object templates due to them being a reference to a privately owned object.
+	// We now deserialize these object template copies
+	// and match them up with their MovieSceneCopyableBinding again.
+	int32 SpawnableObjectTemplateIndex = 0;
+	for (UAvaSequenceCopyableBinding* const ImportedObject : Factory.ImportedBindings)
+	{
+		if (ImportedObject->Spawnable.GetGuid().IsValid() && SpawnableObjectTemplateIndex < Factory.NewSpawnableObjectTemplates.Num())
+		{
+			// This Spawnable Object Template is owned by our transient package, so you'll need to change the owner if you want to keep it later.
+			ImportedObject->SpawnableObjectTemplate = Factory.NewSpawnableObjectTemplates[SpawnableObjectTemplateIndex++];
+		}
+	}
+
+	// Remove the temp package from the root now that it has served its purpose
+	TempPackage->RemoveFromRoot();
+
+	return Factory.ImportedBindings;
+}
+
+void FAvaSequenceImporter::ProcessImportedBindings(const TArray<UAvaSequenceCopyableBinding*>& InImportedBindings
+	, const TMap<FName, AActor*>& InPastedActors)
+{
+	if (InImportedBindings.IsEmpty())
+	{
+		return;
+	}
+
+	TSharedPtr<FAvaSequencer> AvaSequencer = AvaSequencerWeak.Pin();
+	if (!ensureAlways(AvaSequencer.IsValid()))
+	{
+		return;
+	}
+
+	TSharedRef<ISequencer> Sequencer = AvaSequencer->GetSequencer();
+	UAvaSequence* const Sequence     = AvaSequencer->GetViewedSequence();
+	UObject* const PlaybackContext   = Sequencer->GetPlaybackContext();
+
+	if (!ensureAlways(Sequence && PlaybackContext))
+	{
+		return;
+	}
+
+	UMovieScene* const MovieScene = Sequence->GetMovieScene();
+	if (!ensureAlways(MovieScene))
+	{
+		return;
+	}
+
+	struct FImportedPossessable
+	{
+		TArray<FString> BoundObjectPaths;
+	};
+	TMap<FGuid, FImportedPossessable> ImportedPossessableMap;
+	TMap<FGuid, FGuid> OldToNewGuidMap;
+	TArray<FMovieSceneBinding> PastedBindings;
+
+	ImportedPossessableMap.Reserve(InImportedBindings.Num());
+	OldToNewGuidMap.Reserve(InImportedBindings.Num());
+	PastedBindings.Reserve(InImportedBindings.Num());
+
+	FImportContext ImportContext { Sequencer, *Sequence, *MovieScene, Sequencer->GetFocusedTemplateID(), InPastedActors, nullptr };
+
+	for (UAvaSequenceCopyableBinding* const CopyableBinding : InImportedBindings)
+	{
+		// Clear transient flags on the imported tracks
+		for (UMovieSceneTrack* const CopiedTrack : CopyableBinding->Tracks)
+		{
+			FAvaSequenceImporter::ResetCopiedTracksFlags(CopiedTrack);
+		}
+
+		ImportContext.CopyableBinding = CopyableBinding;
+
+		// Possessable Imported Binding
+		if (CopyableBinding->Possessable.GetGuid().IsValid())
+		{
+			FMovieSceneBinding NewBinding = BindPossessable(ImportContext);
+
+			FImportedPossessable ImportedPossessable;
+			ImportedPossessable.BoundObjectPaths = CopyableBinding->BoundObjectPaths;
+			ImportedPossessableMap.Add(NewBinding.GetObjectGuid(), MoveTemp(ImportedPossessable));
+
+			OldToNewGuidMap.Add(CopyableBinding->Possessable.GetGuid(), NewBinding.GetObjectGuid());
+			PastedBindings.Add(MoveTemp(NewBinding));
+		}
+		// Spawnable Imported Binding
+		else if (CopyableBinding->Spawnable.GetGuid().IsValid())
+		{
+			FMovieSceneBinding NewBinding = BindSpawnable(ImportContext);
+			OldToNewGuidMap.Add(CopyableBinding->Spawnable.GetGuid(), NewBinding.GetObjectGuid());
+			PastedBindings.Add(MoveTemp(NewBinding));
+		}
+	}
+
+	// Fix up possessables' parent guids
+	for (const TPair<FGuid, FImportedPossessable>& Pair : ImportedPossessableMap)
+	{
+		FMovieScenePossessable* const Possessable = MovieScene->FindPossessable(Pair.Key);
+		if (!Possessable)
+		{
+			continue;
+		}
+
+		if (FGuid* const NewParentGuid = OldToNewGuidMap.Find(Possessable->GetParent()))
+		{
+			Possessable->SetParent(*NewParentGuid, MovieScene);
+		}
+	}
+
+	// Fix possessables that have bound object paths
+	for (const TPair<FGuid, FImportedPossessable>& Pair : ImportedPossessableMap)
+	{
+		const FGuid& PossessableGuid = Pair.Key;
+		const FImportedPossessable& ImportedPossessable = Pair.Value;
+
+		if (ImportedPossessable.BoundObjectPaths.IsEmpty())
+		{
+			continue;
+		}
+
+		FMovieScenePossessable* const Possessable = MovieScene->FindPossessable(PossessableGuid);
+		if (!Possessable)
+		{
+			continue;
+		}
+
+		auto FindObjectsFunc = [&ImportedPossessableMap, Sequence](const FGuid& InGuid, UObject* InContextChecked)
+			{
+				TArray<UObject*, TInlineAllocator<1>> FoundObjects;
+
+				// Attempt to use the bound object paths that were copied over through the copyable binding
+				// to have an easier time trying to resolve the object
+				if (const FImportedPossessable* const FoundEntry = ImportedPossessableMap.Find(InGuid))
+				{
+					for (const FString& BoundObjectPath : FoundEntry->BoundObjectPaths)
+					{
+						if (UObject* FoundObject = FindObject<UObject>(InContextChecked, *BoundObjectPath))
+						{
+							FoundObjects.Add(FoundObject);
+						}
+					}
+				}
+
+				// If nothing was found, use default method to locate the object
+				if (FoundObjects.IsEmpty() || !FoundObjects[0])
+				{
+					Sequence->LocateBoundObjects(InGuid, UE::UniversalObjectLocator::FResolveParams(InContextChecked), FoundObjects);
+				}
+
+				return FoundObjects;
+			};
+
+		UObject* const ResolutionContext = FAvaSequencer::FindResolutionContext(*Sequence
+			, *MovieScene
+			, Possessable->GetParent()
+			, PlaybackContext
+			, FindObjectsFunc);
+
+		TArray<UObject*> ObjectsToBind;
+		for (const FString& BoundObjectPath : ImportedPossessable.BoundObjectPaths)
+		{
+			if (UObject* FoundObject = FindObject<UObject>(ResolutionContext, *BoundObjectPath))
+			{
+				ObjectsToBind.Add(FoundObject);
+			}
+		}
+
+		AddObjectsToBinding(Sequencer
+			, ObjectsToBind
+			, FMovieSceneBindingProxy(PossessableGuid, Sequence)
+			, ResolutionContext);
+	}
+
+	Sequencer->RestorePreAnimatedState();
+
+	Sequencer->OnMovieSceneBindingsPasted().Broadcast(PastedBindings);
+
+	Sequencer->NotifyMovieSceneDataChanged(EMovieSceneDataChangeType::MovieSceneStructureItemsChanged);
+
+	// Refresh all immediately so that spawned actors will be generated immediately
+	Sequencer->ForceEvaluate();
+
+	// Remap bindings in sections (ie. attach tracks)
+	for (const TPair<FGuid, FGuid>& Pair : OldToNewGuidMap)
+	{
+		FSequencerUtilities::UpdateBindingIDs(Sequencer, Pair.Key, Pair.Value);
+	}
+}
+
+FMovieSceneBinding FAvaSequenceImporter::BindPossessable(const FImportContext& InContext)
+{
+	const FGuid NewGuid = FGuid::NewGuid();
+
+	FMovieSceneBinding NewBinding(NewGuid, InContext.CopyableBinding->Binding.GetName(), InContext.CopyableBinding->Tracks);
+	{
+		FMovieScenePossessable NewPossessable = InContext.CopyableBinding->Possessable;
+		NewPossessable.SetGuid(NewGuid);
+		InContext.MovieScene.AddPossessable(NewPossessable, NewBinding);
+	}
+
+	// Find the actors that this pasted binding should bind to
+	TArray<UObject*> ActorsToBind;
+
+	// Only attempt to bind actor names here, as this is the only thing that can be rebound at this point
+	// Subobjects first need their parents (actors) to resolve correctly, so wait on second pass
+	for (const FName& BoundActorName : InContext.CopyableBinding->BoundActorNames)
+	{
+		if (AActor* const * const FoundActor = InContext.PastedActors.Find(BoundActorName))
+		{
+			AActor* const Actor = *FoundActor;
+			if (!Actor)
+			{
+				continue;
+			}
+
+			// Don't bind if this actor is already bound
+			if (InContext.Sequencer->FindObjectId(*Actor, InContext.SequenceIDRef).IsValid() || ActorsToBind.Contains(Actor))
+			{
+				continue;
+			}
+
+			ActorsToBind.Add(Actor);
+		}
+	}
+
+	// Bind the actors
+	if (!ActorsToBind.IsEmpty())
+	{
+		AddObjectsToBinding(InContext.Sequencer
+			, ActorsToBind
+			, FMovieSceneBindingProxy(NewGuid, &InContext.Sequence)
+			, InContext.Sequencer->GetPlaybackContext());
+	}
+
+	return NewBinding;
+}
+
+FMovieSceneBinding FAvaSequenceImporter::BindSpawnable(const FImportContext& InContext)
+{
+	// We need to let the sequence create the spawnable so that it has everything set up properly internally.
+	// This is required to get spawnables with the correct references to object templates, object templates with
+	// correct owners, etc. However, making a new spawnable also creates the binding for us - this is a problem
+	// because we need to use our binding (which has tracks associated with it). To solve this, we let it create
+	// an object template based off of our (transient package owned) template, then find the newly created binding
+	// and update it.
+
+	FGuid NewGuid;
+	if (InContext.CopyableBinding->SpawnableObjectTemplate)
+	{
+		NewGuid = InContext.Sequencer->MakeNewSpawnable(*InContext.CopyableBinding->SpawnableObjectTemplate
+			, nullptr
+			, false);
+	}
+	else
+	{
+		FMovieSceneSpawnable NewSpawnable{};
+		NewSpawnable.SetGuid(FGuid::NewGuid());
+		NewSpawnable.SetName(InContext.CopyableBinding->Spawnable.GetName());
+
+		InContext.MovieScene.AddSpawnable(NewSpawnable, FMovieSceneBinding(NewSpawnable.GetGuid(), NewSpawnable.GetName()));
+
+		NewGuid = NewSpawnable.GetGuid();
+	}
+
+	FMovieSceneBinding NewBinding(NewGuid, InContext.CopyableBinding->Binding.GetName(), InContext.CopyableBinding->Tracks);
+
+	FMovieSceneSpawnable* const Spawnable = InContext.MovieScene.FindSpawnable(NewGuid);
+
+	// Copy the name of the original spawnable too.
+	Spawnable->SetName(InContext.CopyableBinding->Spawnable.GetName());
+
+	// Clear the transient flags on the copyable binding before assigning to the new spawnable
+	for (UMovieSceneTrack* const Track : NewBinding.GetTracks())
+	{
+		FAvaSequenceImporter::ResetCopiedTracksFlags(Track);
+	}
+
+	// Replace the auto-generated binding with our deserialized bindings (which has our tracks)
+	InContext.MovieScene.ReplaceBinding(NewGuid, NewBinding);
+
+	return NewBinding;
+}
+
+void FAvaSequenceImporter::AddObjectsToBinding(const TSharedRef<ISequencer>& InSequencer
+	, const TArray<UObject*>& InObjectsToAdd
+	, const FMovieSceneBindingProxy& InObjectBinding
+	, UObject* InResolutionContext)
+{
+	UAvaSequence* Sequence = CastChecked<UAvaSequence>(InObjectBinding.Sequence);
+	if (!Sequence)
+	{
+		return;
+	}
+
+	UMovieScene* const MovieScene = Sequence->GetMovieScene();
+	if (!MovieScene || InObjectsToAdd.IsEmpty())
+	{
+		return;
+	}
+
+	UClass* ObjectClass = nullptr;
+	int32 ValidObjectCount = 0;
+
+	FGuid Guid = InObjectBinding.BindingID;
+
+	TArrayView<TWeakObjectPtr<>> WeakObjectsInSequence = Sequence->FindObjectsFromGuid(Guid);
+	for (TWeakObjectPtr<> ObjectWeak : WeakObjectsInSequence)
+	{
+		if (UObject* Object = ObjectWeak.Get())
+		{
+			ObjectClass = Object->GetClass();
+			++ValidObjectCount;
+		}
+	}
+
+	Sequence->Modify();
+	MovieScene->Modify();
+
+	TArray<UObject*> AddedObjects;
+	AddedObjects.Reserve(InObjectsToAdd.Num());
+
+	for (UObject* ObjectToAdd : InObjectsToAdd)
+	{
+		// Skip invalid objects or objects already in the sequence
+		if (!ObjectToAdd || WeakObjectsInSequence.Contains(ObjectToAdd))
+		{
+			continue;
+		}
+
+		// Skip if the object has no common class with the objects already in the binding
+		if (ObjectClass && !UClass::FindCommonBase(ObjectToAdd->GetClass(), ObjectClass))
+		{
+			continue;
+		}
+
+		// if no objects are in the binding, set the class to this object's
+		if (!ObjectClass)
+		{
+			ObjectClass = ObjectToAdd->GetClass();
+		}
+
+		FMovieScenePossessable* Possessable = MovieScene->FindPossessable(Guid);
+		if (!ensureAlways(Possessable))
+		{
+			continue;
+		}
+
+		ObjectToAdd->Modify();
+		if (!Possessable->BindSpawnableObject(InSequencer->GetFocusedTemplateID(), ObjectToAdd, InSequencer->GetSharedPlaybackState()))
+		{
+			Sequence->BindPossessableObject(Guid, *ObjectToAdd, InResolutionContext);
+		}
+		AddedObjects.Add(ObjectToAdd);
+	}
+
+	// Update Labels
+	if (ValidObjectCount + AddedObjects.Num() > 0)
+	{
+		FMovieScenePossessable* const Possessable = MovieScene->FindPossessable(Guid);
+		if (Possessable && ObjectClass)
+		{
+			// If there are multiple objects within the same possessable, name possessable as "ClassName (Count)"
+			if (ValidObjectCount + AddedObjects.Num() > 1)
+			{
+				Possessable->SetName(FString::Printf(TEXT("%s (%d)")
+					, *ObjectClass->GetName()
+					, ValidObjectCount + AddedObjects.Num()));
+			}
+			else if (!AddedObjects.IsEmpty())
+			{
+				Possessable->SetName(FAvaSequencer::GetObjectName(AddedObjects[0]));
+			}
+			Possessable->SetPossessedObjectClass(ObjectClass);
+		}
+	}
+}
+
diff --git a/Engine/Plugins/Experimental/Avalanche/Source/AvalancheSequencer/Private/ObjectBinding/AvaSequencerObjectBinding.cpp b/Engine/Plugins/Experimental/Avalanche/Source/AvalancheSequencer/Private/ObjectBinding/AvaSequencerObjectBinding.cpp
--- a/Engine/Plugins/Experimental/Avalanche/Source/AvalancheSequencer/Private/ObjectBinding/AvaSequencerObjectBinding.cpp
+++ b/Engine/Plugins/Experimental/Avalanche/Source/AvalancheSequencer/Private/ObjectBinding/AvaSequencerObjectBinding.cpp
@@ -75,11 +75,9 @@ void FAvaSequencerObjectBinding::AddPossessActorMenuExtensions(FMenuBuilder& Out
 		{
 			FMovieScenePossessable& Possessable = MovieScene->GetPossessable(Index);
 
-			TArray<UObject*, TInlineAllocator<1>> BoundObjects;
-			static_cast<UMovieSceneSequence*>(Sequence)->LocateBoundObjects(Possessable.GetGuid(), UE::UniversalObjectLocator::FResolveParams(Sequencer->GetPlaybackContext()), Sequencer->GetSharedPlaybackState(), BoundObjects);
-
 			// A possession guid can apply to more than one object, so we get all bound objects for the GUID and add them to our set.
-			ExistingPossessedObjects.Append(BoundObjects);
+			ExistingPossessedObjects.Append(static_cast<UMovieSceneSequence*>(Sequence)->LocateBoundObjects(Possessable.GetGuid()
+				, Sequencer->GetPlaybackContext()));
 		}
 	}
 	
diff --git a/Engine/Plugins/Experimental/ChaosCaching/Source/ChaosCaching/Private/Chaos/Sequencer/ChaosCacheObjectSpawner.cpp b/Engine/Plugins/Experimental/ChaosCaching/Source/ChaosCaching/Private/Chaos/Sequencer/ChaosCacheObjectSpawner.cpp
--- a/Engine/Plugins/Experimental/ChaosCaching/Source/ChaosCaching/Private/Chaos/Sequencer/ChaosCacheObjectSpawner.cpp
+++ b/Engine/Plugins/Experimental/ChaosCaching/Source/ChaosCaching/Private/Chaos/Sequencer/ChaosCacheObjectSpawner.cpp
@@ -10,9 +10,7 @@ FChaosCacheObjectSpawner::FChaosCacheObjectSpawner() : FLevelSequenceActorSpawne
 
 TSharedRef<IMovieSceneObjectSpawner> FChaosCacheObjectSpawner::CreateObjectSpawner()
 {
-	PRAGMA_DISABLE_DEPRECATION_WARNINGS
 	return MakeShareable(new FChaosCacheObjectSpawner);
-	PRAGMA_ENABLE_DEPRECATION_WARNINGS
 }
 
 UClass* FChaosCacheObjectSpawner::GetSupportedTemplateType() const
diff --git a/Engine/Plugins/Experimental/ChaosCaching/Source/ChaosCaching/Public/Chaos/Sequencer/ChaosCacheObjectSpawner.h b/Engine/Plugins/Experimental/ChaosCaching/Source/ChaosCaching/Public/Chaos/Sequencer/ChaosCacheObjectSpawner.h
--- a/Engine/Plugins/Experimental/ChaosCaching/Source/ChaosCaching/Public/Chaos/Sequencer/ChaosCacheObjectSpawner.h
+++ b/Engine/Plugins/Experimental/ChaosCaching/Source/ChaosCaching/Public/Chaos/Sequencer/ChaosCacheObjectSpawner.h
@@ -8,21 +8,14 @@
 class FChaosCacheObjectSpawner : public FLevelSequenceActorSpawner
 {
 public:
-	UE_DEPRECATED(5.5, "This has been deprecated as it is no longer necessary. We use UMovieSceneSpawnableChaosCache custom binding instead to handle this logic")
 	CHAOSCACHING_API FChaosCacheObjectSpawner();
 
 	/** Static method to create the object spawner */
-	UE_DEPRECATED(5.5, "This has been deprecated as it is no longer necessary. We use UMovieSceneSpawnableChaosCache custom binding instead to handle this logic")
 	static CHAOSCACHING_API TSharedRef<IMovieSceneObjectSpawner> CreateObjectSpawner();
 
 	// IMovieSceneObjectSpawner interface
-	UE_DEPRECATED(5.5, "This has been deprecated as it is no longer necessary. We use UMovieSceneSpawnableChaosCache custom binding instead to handle this logic")
 	CHAOSCACHING_API virtual UClass* GetSupportedTemplateType() const override;
-
-	UE_DEPRECATED(5.5, "This has been deprecated as it is no longer necessary. We use UMovieSceneSpawnableChaosCache custom binding instead to handle this logic")
 	CHAOSCACHING_API virtual UObject* SpawnObject(FMovieSceneSpawnable& Spawnable, FMovieSceneSequenceIDRef TemplateID, TSharedRef<const UE::MovieScene::FSharedPlaybackState> SharedPlaybackState) override;
-	
-	UE_DEPRECATED(5.5, "This has been deprecated as it is no longer necessary. We use UMovieSceneSpawnableChaosCache custom binding instead to handle this logic")
 	CHAOSCACHING_API virtual void DestroySpawnedObject(UObject& Object) override;
 	virtual bool IsEditor() const override { return true; }
 	virtual int32 GetSpawnerPriority() const override { return 1; }
diff --git a/Engine/Plugins/Experimental/CineCameraRigs/Source/CineCameraRigs/Private/CineCameraRigRail.cpp b/Engine/Plugins/Experimental/CineCameraRigs/Source/CineCameraRigs/Private/CineCameraRigRail.cpp
--- a/Engine/Plugins/Experimental/CineCameraRigs/Source/CineCameraRigs/Private/CineCameraRigRail.cpp
+++ b/Engine/Plugins/Experimental/CineCameraRigs/Source/CineCameraRigs/Private/CineCameraRigRail.cpp
@@ -570,7 +570,7 @@ UMovieSceneFloatTrack* ACineCameraRigRail::FindPositionTrack(const UMovieSceneSe
 	for (const FMovieSceneBinding& Binding : Bindings)
 	{
 		TArray<UObject*, TInlineAllocator<1>> BoundObjects;
-		InSequence->LocateBoundObjects(Binding.GetObjectGuid(), UE::UniversalObjectLocator::FResolveParams(GetWorld()), MovieSceneHelpers::CreateTransientSharedPlaybackState(GetWorld(), const_cast<UMovieSceneSequence*>(InSequence)), BoundObjects);
+		InSequence->LocateBoundObjects(Binding.GetObjectGuid(), UE::UniversalObjectLocator::FResolveParams(GetWorld()), BoundObjects);
 		if (BoundObjects.IsEmpty() || BoundObjects[0] != this)
 		{
 			continue;
diff --git a/Engine/Plugins/Importers/USDImporter/Source/USDExporter/Private/LevelSequenceExporterUSD.cpp b/Engine/Plugins/Importers/USDImporter/Source/USDExporter/Private/LevelSequenceExporterUSD.cpp
--- a/Engine/Plugins/Importers/USDImporter/Source/USDExporter/Private/LevelSequenceExporterUSD.cpp
+++ b/Engine/Plugins/Importers/USDImporter/Source/USDExporter/Private/LevelSequenceExporterUSD.cpp
@@ -50,8 +50,6 @@
 #include "Tracks/MovieSceneSpawnTrack.h"
 #include "Tracks/MovieSceneSubTrack.h"
 #include "UObject/UObjectGlobals.h"
-#include "Bindings/MovieSceneSpawnableDirectorBlueprintBinding.h"
-#include "Bindings/MovieSceneReplaceableDirectorBlueprintBinding.h"
 
 #define LOCTEXT_NAMESPACE "LevelSequenceExporterUSD"
 
@@ -575,6 +573,16 @@ namespace UE::LevelSequenceExporterUSD::Private
 			// Spawn everything for this instance
 			for (FMovieSceneSequenceID SequenceInstance : SequenceInstancePair.Value)
 			{
+				int32 NumSpawnables = MovieScene->GetSpawnableCount();
+				for (int32 Index = 0; Index < NumSpawnables; ++Index)
+				{
+					const FMovieSceneSpawnable& Spawnable = MovieScene->GetSpawnable(Index);
+					const FGuid& Guid = Spawnable.GetGuid();
+
+					StaticCastSharedRef<FMovieSceneSpawnRegister>(Context.SpawnRegister)
+						->SpawnObject(Guid, *MovieScene, SequenceInstance, Context.Sequencer->GetSharedPlaybackState(), 0);
+				}
+
 				if (const FMovieSceneBindingReferences* BindingReferences = Sequence->GetBindingReferences())
 				{
 					int32 BindingIndex = 0;
@@ -849,38 +857,26 @@ namespace UE::LevelSequenceExporterUSD::Private
 							LastGuid = BindingReference.ID;
 							BindingIndex = 0;
 						}
-						if (BindingReference.CustomBinding)
+						if (BindingReference.CustomBinding
+							&& BindingReference.CustomBinding->WillSpawnObject(Context.Sequencer->GetSharedPlaybackState()))
 						{
-							if (BindingReference.CustomBinding->WillSpawnObject(Context.Sequencer->GetSharedPlaybackState()))
+							BoundObject = Context.SpawnRegister
+											  ->GetExistingSpawn(*RootSequence, SequenceInstance, FSpawnedInstanceKey(Guid, BindingIndex++));
+							if (!BoundObject)
 							{
-								BoundObject = Context.SpawnRegister->GetExistingSpawn(*RootSequence, SequenceInstance, FSpawnedInstanceKey(Guid, BindingIndex++));
-								if (!BoundObject)
-								{
-									// This should never happen as we preemptively spawn everything:
-									// At this point all our spawns should be spawned, but invisible
-									UE_LOG(LogUsd, Warning, TEXT("Failed to find spawned object for spawnable with Guid '%s'"), *Guid.ToString());
-									continue;
-								}
-
-								if (BoundObject && (BoundObject->IsA<USceneComponent>() || BoundObject->IsA<AActor>()))
-								{
-									BoundObjects.Add(BoundObject, FSpawnedInstanceKey(Guid, BindingIndex));
-								}
+								// This should never happen as we preemptively spawn everything:
+								// At this point all our spawns should be spawned, but invisible
+								UE_LOG(LogUsd, Warning, TEXT("Failed to find spawned object for spawnable with Guid '%s'"), *Guid.ToString());
+								continue;
 							}
 
-							if (UMovieSceneSpawnableDirectorBlueprintBinding* SpawnableDirectorBlueprintBinding = Cast<UMovieSceneSpawnableDirectorBlueprintBinding>(BindingReference.CustomBinding->AsSpawnable(Context.Sequencer->GetSharedPlaybackState())))
+							if (BoundObject && (BoundObject->IsA<USceneComponent>() || BoundObject->IsA<AActor>()))
 							{
-								if (SpawnableDirectorBlueprintBinding->DynamicBinding.Function)
-								{
-									DynamicBindings.Add(FSpawnedInstanceKey(Guid, BindingIndex), &SpawnableDirectorBlueprintBinding->DynamicBinding);
-								}
-							}
+								BoundObjects.Add(BoundObject, FSpawnedInstanceKey(Guid, BindingIndex));
 
-							if (UMovieSceneReplaceableDirectorBlueprintBinding* ReplaceableDirectorBlueprintBinding = Cast<UMovieSceneReplaceableDirectorBlueprintBinding>(BindingReference.CustomBinding))
-							{
-								if (ReplaceableDirectorBlueprintBinding->DynamicBinding.Function)
+								if (Possessable.DynamicBinding.Function)
 								{
-									DynamicBindings.Add(FSpawnedInstanceKey(Guid, BindingIndex), &ReplaceableDirectorBlueprintBinding->DynamicBinding);
+									DynamicBindings.Add(FSpawnedInstanceKey(Guid, BindingIndex), &Possessable.DynamicBinding);
 								}
 							}
 						}
@@ -908,6 +904,45 @@ namespace UE::LevelSequenceExporterUSD::Private
 				if (BoundObject && (BoundObject->IsA<USceneComponent>() || BoundObject->IsA<AActor>()))
 				{
 					BoundObjects.Add(BoundObject, FSpawnedInstanceKey(Guid, 0));
+
+					if (Possessable.DynamicBinding.Function)
+					{
+						DynamicBindings.Add(FSpawnedInstanceKey(Guid, 0), &Possessable.DynamicBinding);
+					}
+				}
+			}
+
+			// Spawnables
+			int32 NumSpawnables = MovieScene->GetSpawnableCount();
+			for (int32 Index = 0; Index < NumSpawnables; ++Index)
+			{
+				FMovieSceneSpawnable& Spawnable = MovieScene->GetSpawnable(Index);
+				const FGuid& Guid = Spawnable.GetGuid();
+
+				// We won't have spawned ASphereReflectionCapture here.
+				// See the comment inside FLevelSequenceHidingSpawnRegister::SpawnObject and UE-167593 for more info
+				if (ASphereReflectionCapture* ReflectionCapture = Cast<ASphereReflectionCapture>(Spawnable.GetObjectTemplate()))
+				{
+					continue;
+				}
+
+				UObject* BoundObject = Context.SpawnRegister->GetExistingSpawn(*RootSequence, SequenceInstance, FSpawnedInstanceKey(Guid, 0));
+				if (!BoundObject)
+				{
+					// This should never happen as we preemptively spawn everything:
+					// At this point all our spawns should be spawned, but invisible
+					UE_LOG(LogUsd, Warning, TEXT("Failed to find spawned object for spawnable with Guid '%s'"), *Guid.ToString());
+					continue;
+				}
+
+				if (BoundObject && (BoundObject->IsA<USceneComponent>() || BoundObject->IsA<AActor>()))
+				{
+					BoundObjects.Add(BoundObject, FSpawnedInstanceKey(Guid, 0));
+
+					if (Spawnable.DynamicBinding.Function)
+					{
+						DynamicBindings.Add(FSpawnedInstanceKey(Guid, 0), &Spawnable.DynamicBinding);
+					}
 				}
 			}
 		}
@@ -1734,7 +1769,7 @@ bool ULevelSequenceExporterUsd::ExportBinary(
 	// Spawn (but hide) all spawnables so that they will also show up on the level export if we need them to.
 	// We have to traverse the template IDs when spawning spawnables, because we'll want to force each individual spawnable of each
 	// FMovieSceneSequenceID to spawn a separate object, so that they can become separate prims. Without doing this, if we used the same
-	// subsequence with spawnables multiple times within a parent sequence we'd only get one prim out, as the spawnable bindings
+	// subsequence with spawnables multiple times within a parent sequence we'd only get one prim out, as the FMovieSceneSpawnable objects
 	// would be the exact same between all instances of the child sequence (same FGuid)
 	LevelSequenceExporterImpl::PreSpawnSpawnables(Context, *LevelSequence);
 
diff --git a/Engine/Plugins/Importers/USDImporter/Source/USDStage/Private/USDLevelSequenceHelper.cpp b/Engine/Plugins/Importers/USDImporter/Source/USDStage/Private/USDLevelSequenceHelper.cpp
--- a/Engine/Plugins/Importers/USDImporter/Source/USDStage/Private/USDLevelSequenceHelper.cpp
+++ b/Engine/Plugins/Importers/USDImporter/Source/USDStage/Private/USDLevelSequenceHelper.cpp
@@ -218,8 +218,7 @@ namespace UsdLevelSequenceHelperImpl
 			}
 		}
 
-		TArray<UObject*, TInlineAllocator<1>> Objects;
-		MovieSceneSequence.LocateBoundObjects(Guid, UE::UniversalObjectLocator::FResolveParams(ParentContext), nullptr, Objects);
+		TArray<UObject*, TInlineAllocator<1>> Objects = MovieSceneSequence.LocateBoundObjects(Guid, ParentContext);
 		if (Objects.Num() > 0)
 		{
 			return Objects[0];
diff --git a/Engine/Plugins/MovieScene/LevelSequenceEditor/Source/LevelSequenceEditor/Private/LevelSequenceEditorSubsystem.cpp b/Engine/Plugins/MovieScene/LevelSequenceEditor/Source/LevelSequenceEditor/Private/LevelSequenceEditorSubsystem.cpp
--- a/Engine/Plugins/MovieScene/LevelSequenceEditor/Source/LevelSequenceEditor/Private/LevelSequenceEditorSubsystem.cpp
+++ b/Engine/Plugins/MovieScene/LevelSequenceEditor/Source/LevelSequenceEditor/Private/LevelSequenceEditorSubsystem.cpp
@@ -67,7 +67,6 @@
 #include "DetailLayoutBuilder.h"
 #include "IPropertyTypeCustomization.h"
 #include "IPropertyUtilities.h"
-#include "MovieSceneDynamicBindingCustomization.h"
 
 #include UE_INLINE_GENERATED_CPP_BY_NAME(LevelSequenceEditorSubsystem)
 
@@ -138,15 +137,6 @@ void FMovieSceneBindingPropertyInfoDetailCustomization::CustomizeChildren(TShare
 		FMovieSceneBindingPropertyInfo* BindingPropertyInfo = (StructPtrs.Num() == 1) ? reinterpret_cast<FMovieSceneBindingPropertyInfo*>(StructPtrs[0]) : nullptr;
 		if (BindingPropertyInfo)
 		{
-			bool bShowConvert = true;
-			if (FMovieScenePossessable* Possessable = MovieScene->FindPossessable(BindingGuid))
-			{
-				if (Possessable && Possessable->GetParent().IsValid())
-				{
-					bShowConvert = false;
-				}
-			}
-
 			BindingIndex = InStructPropertyHandle->GetArrayIndex();
 			int32 InitialIndex = 0;
 			BindingTypeNames.Add(MakeShared<FText>(LOCTEXT("BindingType_Possessable", "Possessable")));
@@ -178,7 +168,6 @@ void FMovieSceneBindingPropertyInfoDetailCustomization::CustomizeChildren(TShare
 					SNew(SComboBox<TSharedPtr<FText>>)
 					.OptionsSource(&BindingTypeNames)
 					.InitiallySelectedItem(BindingTypeNames[InitialIndex])
-					.IsEnabled(bShowConvert)
 					.OnSelectionChanged_Lambda([this, InStructPropertyHandle, SupportedBindingTypes, BindingPropertyInfo, &StructBuilder, &CustomizationUtils](TSharedPtr<FText> NewValue, ESelectInfo::Type)
 					{
 						int32 CurrentIndex = BindingTypeNames.IndexOfByPredicate([&](const TSharedPtr<FText>& In) {
@@ -222,7 +211,8 @@ void FMovieSceneBindingPropertyInfoDetailCustomization::CustomizeChildren(TShare
 			];
 
 			// Only show certain menus if we have a currently bound object
-			if (bShowConvert && MovieSceneHelpers::GetSingleBoundObject(Sequence, BindingGuid, Sequencer->GetSharedPlaybackState(), BindingIndex) != nullptr)
+			
+			if (MovieSceneHelpers::GetSingleBoundObject(Sequence, BindingGuid, Sequencer->GetSharedPlaybackState(), BindingIndex) != nullptr)
 			{
 				FDetailWidgetRow& ConvertToRow = StructBuilder.AddCustomRow(LOCTEXT("ConvertBindingTo", "Convert Binding To..."));
 				ConvertToRow.WholeRowContent()
@@ -235,40 +225,40 @@ void FMovieSceneBindingPropertyInfoDetailCustomization::CustomizeChildren(TShare
 						SNew(STextBlock).Text(LOCTEXT("ConvertBindingTo", "Convert Binding To..."))
 					]
 					];
-			}
 
-			// Show Change class and save default state menus for spawnables
-			if (MovieSceneHelpers::SupportsObjectTemplate(Sequence, BindingGuid, Sequencer->GetSharedPlaybackState(), BindingIndex))
-			{
-				FDetailWidgetRow& ChangeClassRow = StructBuilder.AddCustomRow(LOCTEXT("ChangeClass", "Change Class..."));
-				ChangeClassRow.WholeRowContent()
-					[
-						SNew(SComboButton)
-						.OnGetMenuContent(FOnGetContent::CreateLambda([this, &StructBuilder, &CustomizationUtils]() { return OnGetChangeClassMenuContent(StructBuilder, CustomizationUtils); }))
-					.ComboButtonStyle(FAppStyle::Get(), "SimpleComboButton")
-					.ButtonContent()
-					[
-						SNew(STextBlock).Text(LOCTEXT("ChangeClass", "Change Class..."))
-					]
-					];
+				// Show Change class and save default state menus for spawnables
+				if (MovieSceneHelpers::IsBoundToSpawnable(Sequence, BindingGuid, Sequencer->GetSharedPlaybackState(), BindingIndex))
+				{
+					FDetailWidgetRow& ChangeClassRow = StructBuilder.AddCustomRow(LOCTEXT("ChangeClass", "Change Class..."));
+					ChangeClassRow.WholeRowContent()
+						[
+							SNew(SComboButton)
+							.OnGetMenuContent(FOnGetContent::CreateLambda([this, &StructBuilder, &CustomizationUtils]() { return OnGetChangeClassMenuContent(StructBuilder, CustomizationUtils); }))
+						.ComboButtonStyle(FAppStyle::Get(), "SimpleComboButton")
+						.ButtonContent()
+						[
+							SNew(STextBlock).Text(LOCTEXT("ChangeClass", "Change Class..."))
+						]
+						];
 
-				// Save Default State
-				FDetailWidgetRow& SaveDefaultStateRow = StructBuilder.AddCustomRow(LOCTEXT("SaveDefaultState", "Save Default State"));
-				SaveDefaultStateRow.WholeRowContent()
-					[
-						SNew(SButton)
-						.Text(LOCTEXT("SaveDefaultState", "Save Default State"))
-					.ToolTipText(LOCTEXT("SaveDefaultState_Tooltip", "Save the current state of this spawnable as default properties"))
-					.OnClicked_Lambda([this]()
-						{
-							TSharedPtr<ISequencer> Sequencer = SequencerPtr.Pin();
-							if (Sequencer.IsValid())
+					// Save Default State
+					FDetailWidgetRow& SaveDefaultStateRow = StructBuilder.AddCustomRow(LOCTEXT("SaveDefaultState", "Save Default State"));
+					SaveDefaultStateRow.WholeRowContent()
+						[
+							SNew(SButton)
+							.Text(LOCTEXT("SaveDefaultState", "Save Default State"))
+						.ToolTipText(LOCTEXT("SaveDefaultState_Tooltip", "Save the current state of this spawnable as default properties"))
+						.OnClicked_Lambda([this]()
 							{
-								Sequencer->GetSpawnRegister().SaveDefaultSpawnableState(BindingGuid, BindingIndex, Sequencer->GetFocusedTemplateID(), Sequencer->GetSharedPlaybackState());
-							}
-							return FReply::Handled();
-						})
-					];
+								TSharedPtr<ISequencer> Sequencer = SequencerPtr.Pin();
+								if (Sequencer.IsValid())
+								{
+									Sequencer->GetSpawnRegister().SaveDefaultSpawnableState(BindingGuid, BindingIndex, Sequencer->GetFocusedTemplateID(), Sequencer->GetSharedPlaybackState());
+								}
+								return FReply::Handled();
+							})
+						];
+				}
 			}
 
 
@@ -453,6 +443,11 @@ void ULevelSequenceEditorSubsystem::Initialize(FSubsystemCollectionBase& Collect
 	BindingPropertiesMenuExtender = MakeShareable(new FExtender);
 	BindingPropertiesMenuExtender->AddMenuExtension("Possessable", EExtensionHook::First, CommandList, FMenuExtensionDelegate::CreateLambda([this](FMenuBuilder& MenuBuilder) {
 		
+		if (!IsSelectedBindingRootPossessable())
+		{
+			return;
+		}
+		
 		FFormatNamedArguments Args;
 		MenuBuilder.AddSubMenu(
 			FText::Format(LOCTEXT("BindingProperties", "Binding Properties"), Args),
@@ -484,7 +479,11 @@ void ULevelSequenceEditorSubsystem::Initialize(FSubsystemCollectionBase& Collect
 
 	BindingPropertiesMenuExtender->AddMenuExtension("CustomBinding", EExtensionHook::First, CommandList, FMenuExtensionDelegate::CreateLambda([this](FMenuBuilder& MenuBuilder) {
 		// Only add menu entries where the focused sequence is a ULevelSequence
-		
+		if (!IsSelectedBindingRootPossessable())
+		{
+			return;
+		}
+
 		// Add instanced detail customizations
 
 		FFormatNamedArguments Args;
@@ -507,11 +506,6 @@ void ULevelSequenceEditorSubsystem::Deinitialize()
 		SequencerModulePtr->UnregisterOnSequencerCreated(OnSequencerCreatedHandle);
 	}
 
-	if (FSlateApplication::IsInitialized())
-	{
-		FSlateApplication::Get().OnMenuBeingDestroyed().RemoveAll(this);
-	}
-
 }
 
 void ULevelSequenceEditorSubsystem::OnSequencerCreated(TSharedRef<ISequencer> InSequencer)
@@ -535,31 +529,12 @@ void ULevelSequenceEditorSubsystem::AddBindingDetailCustomizations(TSharedRef<ID
 				{
 					return MakeShared<FMovieSceneBindingPropertyInfoDetailCustomization>(InSequencer, InMovieScene, InBindingGuid, LevelSequenceEditorSubsystem);
 				}, ActiveSequencer.ToWeakPtr(), MovieScene, BindingGuid, this));
-
-			DetailsView->RegisterInstancedCustomPropertyTypeLayout(FMovieSceneDynamicBinding::StaticStruct()->GetFName(), FOnGetPropertyTypeCustomizationInstance::CreateStatic(&FMovieSceneDynamicBindingCustomization::MakeInstance, MovieScene, BindingGuid));
 			
 			DetailsView->RegisterInstancedCustomPropertyLayout(UMovieSceneSpawnableActorBinding::StaticClass(), FOnGetDetailCustomizationInstance::CreateStatic(&FMovieSceneSpawnableActorBindingBaseCustomization::MakeInstance, ActiveSequencer.ToWeakPtr(), MovieScene, BindingGuid));
 		}
 	}
 }
 
-void ULevelSequenceEditorSubsystem::OnMenuBeingDestroyed(const TSharedRef<IMenu>& Menu, TSharedRef<IDetailsView> DetailsView)
-{
-	TSharedPtr<SWidget> ContentWidget = Menu->GetContent();
-	TSharedPtr<SWidget> ParentWidget = DetailsView;
-	while (ParentWidget)
-	{
-		if (ParentWidget == ContentWidget)
-		{
-			// Binding Properties Menu has closed, clear the binding property list
-			BindingPropertyInfoList = nullptr;
-			FSlateApplication::Get().OnMenuBeingDestroyed().RemoveAll(this);
-			break;
-		}
-		ParentWidget = ParentWidget->GetParentWidget();
-	}
-}
-
 TSharedPtr<ISequencer> ULevelSequenceEditorSubsystem::GetActiveSequencer()
 {
 	for (TWeakPtr<ISequencer> Ptr : Sequencers)
@@ -676,22 +651,15 @@ TArray<FMovieSceneBindingProxy> ULevelSequenceEditorSubsystem::ConvertToSpawnabl
 		return SpawnableProxies;
 	}
 
-	FMovieScenePossessable* NewPossessable = nullptr;
-
-	if (FMovieSceneBindingReferences* BindingReferences = Sequence->GetBindingReferences())
+	TArray<FMovieSceneSpawnable*> Spawnables = FSequencerUtilities::ConvertToSpawnable(Sequencer.ToSharedRef(), ObjectBinding.BindingID);
+	for (FMovieSceneSpawnable* Spawnable : Spawnables)
 	{
-		int32 NumBindings = BindingReferences->GetReferences(ObjectBinding.BindingID).Num();
-		for (int32 BindingIndex = 0; BindingIndex < NumBindings; BindingIndex++)
+		if (Spawnable)
 		{
-			NewPossessable = FSequencerUtilities::ConvertToCustomBinding(Sequencer.ToSharedRef(), ObjectBinding.BindingID, UMovieSceneSpawnableActorBinding::StaticClass(), BindingIndex);
+			SpawnableProxies.Add(FMovieSceneBindingProxy(Spawnable->GetGuid(), Sequence));
 		}
 	}
 
-	if (NewPossessable)
-	{
-		SpawnableProxies.Add(FMovieSceneBindingProxy(NewPossessable->GetGuid(), Sequence));
-	}
-
 	return SpawnableProxies;
 }
 
@@ -2200,10 +2168,6 @@ void ULevelSequenceEditorSubsystem::AddBindingPropertiesMenu(FMenuBuilder& MenuB
 		RefreshBindingDetails(&DetailsView.Get(), ObjectBindings[0]);
 		DetailsView->OnFinishedChangingProperties().AddUObject(this, &ULevelSequenceEditorSubsystem::OnFinishedChangingLocators, DetailsView, ObjectBindings[0]);
 
-		if (FSlateApplication::IsInitialized())
-		{
-			FSlateApplication::Get().OnMenuBeingDestroyed().AddUObject(this, &ULevelSequenceEditorSubsystem::OnMenuBeingDestroyed, DetailsView);
-		}
 		MenuBuilder.AddWidget(DetailsView, FText::GetEmpty(), true);
 	}
 }
@@ -2267,10 +2231,6 @@ void ULevelSequenceEditorSubsystem::OnFinishedChangingLocators(const FPropertyCh
 				CopiedBinding = Cast<UMovieSceneCustomBinding>(StaticDuplicateObject(LocatorInfo.CustomBinding, MovieScene));
 			}
 			BindingReferences->AddBinding(ObjectBindingID, MoveTemp(LocatorInfo.Locator), LocatorInfo.ResolveFlags, CopiedBinding);
-			if (CopiedBinding)
-			{
-				CopiedBinding->OnBindingAddedOrChanged(*MovieScene);
-			}
 		}
 
 		Sequencer->State.Invalidate(ObjectBindingID, Sequencer->GetFocusedTemplateID());
@@ -2320,12 +2280,6 @@ void ULevelSequenceEditorSubsystem::OnFinishedChangingLocators(const FPropertyCh
 
 		Sequencer->NotifyMovieSceneDataChanged(EMovieSceneDataChangeType::MovieSceneStructureItemsChanged);
 
-		// Destroy any previous spawnables- they'll get recreated on the force evaluate below
-		for (int32 BindingIndex = 0; BindingIndex < BindingPropertyInfoList->Bindings.Num(); ++BindingIndex)
-		{
-			Sequencer->GetSpawnRegister().DestroySpawnedObject(ObjectBindingID, Sequencer->GetFocusedTemplateID(), Sequencer->GetSharedPlaybackState(), BindingIndex);
-		}
-
 		// Force evaluate the Sequencer after clearing the cache (which the above will do) so that any newly loaded actors will be loaded as part of the transaction
 		Sequencer->ForceEvaluate();
 
@@ -2339,10 +2293,10 @@ void ULevelSequenceEditorSubsystem::OnFinishedChangingLocators(const FPropertyCh
 				UMovieSceneCustomBinding* CopiedBinding = nullptr;
 				if (Reference.CustomBinding)
 				{
-					CopiedBinding = Cast<UMovieSceneCustomBinding>(StaticDuplicateObject(Reference.CustomBinding, this));
+					CopiedBinding = Cast<UMovieSceneCustomBinding>(StaticDuplicateObject(Reference.CustomBinding, BindingPropertyInfoList.Get()));
 				}
 
-				return FMovieSceneBindingPropertyInfo{ Reference.Locator, Reference.ResolveFlags, CopiedBinding };
+				return FMovieSceneBindingPropertyInfo{ Reference.Locator, Reference.ResolveFlags, Reference.CustomBinding };
 			});
 
 
@@ -2555,7 +2509,7 @@ bool ULevelSequenceEditorSubsystem::IsSelectedBindingRootPossessable()
 			{
 				if (FMovieScenePossessable* Possessable = MovieScene->FindPossessable(ObjectBindings[0]))
 				{
-					if (!Possessable->GetParent().IsValid() && !Possessable->GetSpawnableObjectBindingID().IsValid())
+					if (!Possessable->GetParent().IsValid())
 					{
 						return true;
 					}
@@ -2585,15 +2539,15 @@ void ULevelSequenceEditorSubsystem::RefreshBindingDetails(IDetailsView* DetailsV
 	}
 	if (FMovieSceneBindingReferences* BindingReferences = Sequence->GetBindingReferences())
 	{
-		BindingPropertyInfoList = NewObject<UMovieSceneBindingPropertyInfoList>(this);
+		BindingPropertyInfoList = NewObject<UMovieSceneBindingPropertyInfoList>();
 		Algo::Transform(BindingReferences->GetReferences(ObjectBindingID), BindingPropertyInfoList->Bindings, [this](const FMovieSceneBindingReference& Reference)
 			{
 				UMovieSceneCustomBinding* CopiedBinding = nullptr;
 				if (Reference.CustomBinding)
 				{
-					CopiedBinding = Cast<UMovieSceneCustomBinding>(StaticDuplicateObject(Reference.CustomBinding, this));
+					CopiedBinding = Cast<UMovieSceneCustomBinding>(StaticDuplicateObject(Reference.CustomBinding, BindingPropertyInfoList.Get()));
 				}
-				return FMovieSceneBindingPropertyInfo{ Reference.Locator, Reference.ResolveFlags, CopiedBinding };
+				return FMovieSceneBindingPropertyInfo{ Reference.Locator, Reference.ResolveFlags, Reference.CustomBinding };
 			});
 
 		DetailsView->SetObject(BindingPropertyInfoList.Get(), true);
diff --git a/Engine/Plugins/MovieScene/LevelSequenceEditor/Source/LevelSequenceEditor/Private/Misc/LevelSequenceEditorActorBinding.cpp b/Engine/Plugins/MovieScene/LevelSequenceEditor/Source/LevelSequenceEditor/Private/Misc/LevelSequenceEditorActorBinding.cpp
--- a/Engine/Plugins/MovieScene/LevelSequenceEditor/Source/LevelSequenceEditor/Private/Misc/LevelSequenceEditorActorBinding.cpp
+++ b/Engine/Plugins/MovieScene/LevelSequenceEditor/Source/LevelSequenceEditor/Private/Misc/LevelSequenceEditorActorBinding.cpp
@@ -61,9 +61,7 @@ void FLevelSequenceEditorActorBinding::AddPossessActorMenuExtensions(FMenuBuilde
 				if (!Possessable.GetParent().IsValid())
 				{
 					// A possession guid can apply to more than one object, so we get all bound objects for the GUID and add them to our set.
-					TArray<UObject*, TInlineAllocator<1>> OutObjects;
-					MovieSceneSequence->LocateBoundObjects(Possessable.GetGuid(), UE::UniversalObjectLocator::FResolveParams(Sequencer.Pin()->GetPlaybackContext()), Sequencer.Pin()->FindSharedPlaybackState(), OutObjects);
-					ExistingPossessedObjects.Append(OutObjects);
+					ExistingPossessedObjects.Append(MovieSceneSequence->LocateBoundObjects(Possessable.GetGuid(), Sequencer.Pin()->GetPlaybackContext()));
 				}
 			}
 		}
diff --git a/Engine/Plugins/MovieScene/MovieRenderPipeline/Source/MovieRenderPipelineCore/Private/MoviePipelineBlueprintLibrary.cpp b/Engine/Plugins/MovieScene/MovieRenderPipeline/Source/MovieRenderPipelineCore/Private/MoviePipelineBlueprintLibrary.cpp
--- a/Engine/Plugins/MovieScene/MovieRenderPipeline/Source/MovieRenderPipelineCore/Private/MoviePipelineBlueprintLibrary.cpp
+++ b/Engine/Plugins/MovieScene/MovieRenderPipeline/Source/MovieRenderPipelineCore/Private/MoviePipelineBlueprintLibrary.cpp
@@ -6,6 +6,7 @@
 #include "MovieRenderPipelineDataTypes.h"
 #include "MovieScene.h"
 #include "MovieSceneSequence.h"
+#include "MovieSceneSpawnable.h"
 #include "MovieScenePossessable.h"
 #include "MovieSceneBinding.h"
 #include "LevelSequence.h"
@@ -511,8 +512,6 @@ void UMoviePipelineBlueprintLibrary::UpdateJobShotListFromSequence(ULevelSequenc
 			// camera components and track them as sidecar cameras. We do this even if sidecar data isn't being used.
 			if (LeafNode->CameraCutSection.IsValid() && LeafNode->MovieScene.IsValid())
 			{
-				UMovieSceneSequence* Sequence = LeafNode->MovieScene->GetTypedOuter<UMovieSceneSequence>();
-				TSharedRef<UE::MovieScene::FSharedPlaybackState> TransientPlaybackState = MovieSceneHelpers::CreateTransientSharedPlaybackState(GWorld, Sequence);
 				int32 MainCameraIndex = INDEX_NONE;
 				FMovieSceneObjectBindingID MainBinding = LeafNode->CameraCutSection->GetCameraBindingID();
 
@@ -562,28 +561,33 @@ void UMoviePipelineBlueprintLibrary::UpdateJobShotListFromSequence(ULevelSequenc
 								LastValidGuid = ParentGuid;
 								ParentGuid = ParentAsPossessable->GetParent();
 								SidecarCamera.Name = ParentAsPossessable->GetName();
-
-								if (MovieSceneHelpers::IsBoundToAnySpawnable(LeafNode->MovieScene->GetTypedOuter<UMovieSceneSequence>(), ParentAsPossessable->GetGuid(), TransientPlaybackState))
+							}
+							else if (FMovieSceneSpawnable* ParentAsSpawnable = LeafNode->MovieScene->FindSpawnable(ParentGuid))
+							{
+								FMovieSceneBinding* Binding = LeafNode->MovieScene->FindBinding(ParentAsSpawnable->GetGuid());
+								if (Binding)
 								{
-									FMovieSceneBinding* Binding = LeafNode->MovieScene->FindBinding(ParentAsPossessable->GetGuid());
-									if (Binding)
+									const TArray<UMovieSceneTrack*> AllBindingTracks = Binding->GetTracks();
+									for (const UMovieSceneTrack* Track : AllBindingTracks)
 									{
-										const TArray<UMovieSceneTrack*> AllBindingTracks = Binding->GetTracks();
-										for (const UMovieSceneTrack* Track : AllBindingTracks)
+										if (const UMovieSceneSpawnTrack* SpawnTrack = Cast<UMovieSceneSpawnTrack>(Track))
 										{
-											if (const UMovieSceneSpawnTrack* SpawnTrack = Cast<UMovieSceneSpawnTrack>(Track))
+											for (const UMovieSceneSection* Section : SpawnTrack->GetAllSections())
 											{
-												for (const UMovieSceneSection* Section : SpawnTrack->GetAllSections())
+												if (const UMovieSceneSpawnSection* BoolSection = Cast<UMovieSceneSpawnSection>(Section))
 												{
-													if (const UMovieSceneSpawnSection* BoolSection = Cast<UMovieSceneSpawnSection>(Section))
-													{
-														bAddCamera = BoolSection->GetChannel().GetDefault().Get(true);
-													}
+													bAddCamera = BoolSection->GetChannel().GetDefault().Get(true);
 												}
 											}
 										}
 									}
 								}
+
+								// Spawnables will never have a parent
+								LastValidGuid = ParentGuid;
+								ParentGuid.Invalidate(); 
+								
+								SidecarCamera.Name = ParentAsSpawnable->GetName();
 							}
 						}
 
diff --git a/Engine/Plugins/MovieScene/SequencerScripting/Source/SequencerScripting/Private/ExtensionLibraries/MovieSceneBindingExtensions.cpp b/Engine/Plugins/MovieScene/SequencerScripting/Source/SequencerScripting/Private/ExtensionLibraries/MovieSceneBindingExtensions.cpp
--- a/Engine/Plugins/MovieScene/SequencerScripting/Source/SequencerScripting/Private/ExtensionLibraries/MovieSceneBindingExtensions.cpp
+++ b/Engine/Plugins/MovieScene/SequencerScripting/Source/SequencerScripting/Private/ExtensionLibraries/MovieSceneBindingExtensions.cpp
@@ -260,9 +260,16 @@ UObject* UMovieSceneBindingExtensions::GetObjectTemplate(const FMovieSceneBindin
 	UMovieScene* MovieScene = InBinding.GetMovieScene();
 	if (MovieScene)
 	{
+		using namespace UE::MovieScene;
+		FSharedPlaybackStateCreateParams CreateParams;
+		// Ideally we would have a proper world context here, but we don't so use GWorld
+		CreateParams.PlaybackContext = GWorld;
 		UMovieSceneSequence* ThisSequence = MovieScene->GetTypedOuter<UMovieSceneSequence>();
-		TSharedRef<UE::MovieScene::FSharedPlaybackState> TransientPlaybackState = MovieSceneHelpers::CreateTransientSharedPlaybackState(GWorld, ThisSequence);
+		TSharedRef<FSharedPlaybackState> TransientPlaybackState = MakeShared<FSharedPlaybackState>(*ThisSequence, CreateParams);
 
+		FMovieSceneEvaluationState State;
+		TransientPlaybackState->AddCapabilityRaw(&State);
+		State.AssignSequence(MovieSceneSequenceID::Root, *ThisSequence, TransientPlaybackState);
 
 		// TODO: Technically this assumes only one spawnable- do we need to upgrade script to handle multiple binding indices?
 		return MovieSceneHelpers::GetObjectTemplate(MovieScene->GetTypedOuter<UMovieSceneSequence>(), InBinding.BindingID, TransientPlaybackState);
diff --git a/Engine/Plugins/VirtualProduction/Takes/Source/TakeRecorder/Private/SerializedRecorder.cpp b/Engine/Plugins/VirtualProduction/Takes/Source/TakeRecorder/Private/SerializedRecorder.cpp
--- a/Engine/Plugins/VirtualProduction/Takes/Source/TakeRecorder/Private/SerializedRecorder.cpp
+++ b/Engine/Plugins/VirtualProduction/Takes/Source/TakeRecorder/Private/SerializedRecorder.cpp
@@ -32,7 +32,6 @@
 #include "UObject/UObjectIterator.h"
 #include "LevelSequenceActor.h"
 #include "LevelEditorViewport.h"
-#include "SequencerUtilities.h"
 
 #define LOCTEXT_NAMESPACE "SerializedRecorder"
 
@@ -424,18 +423,11 @@ AActor* FSerializedRecorder::SetActorPossesableOrSpawnable(UMovieSceneSequence*
 
 		Actor->SetActorLabel(ActorHeader.Label); //spawnable has samen Name used for matching.
 
-		FMovieScenePossessable Possessable(ActorHeader.Label, SpawnableClass);
+		FMovieSceneSpawnable Spawnable(ActorHeader.Label, *Actor);
+		Spawnable.SetGuid(ActorHeader.Guid);
 
-		Possessable.SetGuid(ActorHeader.Guid);
 		FMovieSceneBinding NewBinding(ActorHeader.Guid, ActorHeader.Label);
-		MovieScene->AddPossessable(Possessable, NewBinding);
-
-		UE::Sequencer::FCreateBindingParams CreateBindingParams;
-		CreateBindingParams.bAllowCustomBinding = true;
-		CreateBindingParams.bSpawnable = true;
-		CreateBindingParams.BindingNameOverride = ActorHeader.Label;
-		CreateBindingParams.ReplacementGuid = ActorHeader.Guid;
-		FSequencerUtilities::CreateOrReplaceBinding(nullptr, InMovieSceneSequence, Actor, CreateBindingParams);
+		MovieScene->AddSpawnable(Spawnable, NewBinding);
 
 		if (ActorHeader.Guid.IsValid())
 		{
@@ -483,6 +475,15 @@ AActor* FSerializedRecorder::SetActorPossesableOrSpawnable(UMovieSceneSequence*
 		Possessable->SetName(ActorHeader.Label);
 	}
 
+	if (FMovieSceneSpawnable* Spawnable = MovieScene->FindSpawnable(ActorHeader.Guid))
+	{
+		if (!Spawnable->Tags.Contains(*(ActorHeader.Label)))
+		{
+			Spawnable->Tags.AddUnique(FName(*(ActorHeader.Label)));
+		}
+		Spawnable->SetName(ActorHeader.Label);
+	}
+
 	// look for a folder to put us in
 	if (ActorHeader.FolderName.IsValid())
 	{
@@ -525,6 +526,12 @@ void FSerializedRecorder::SetComponentPossessable(UMovieSceneSequence* InMovieSc
 		// Set up parent/child guids for possessables within spawnables
 		FMovieScenePossessable* ChildPossessablePtr = InMovieScene->FindPossessable(ActorProperty.Guid);
 		ChildPossessablePtr->SetParent(ActorHeader.Guid, InMovieScene);
+		
+		FMovieSceneSpawnable* ParentSpawnable = InMovieScene->FindSpawnable(ActorHeader.Guid);
+		if (ParentSpawnable)
+		{
+			ParentSpawnable->AddChildPossessable(ActorProperty.Guid);
+		}
 
 		for (UActorComponent* Component : Actor->GetComponents())
 		{
diff --git a/Engine/Plugins/VirtualProduction/Takes/Source/TakeRecorderSources/Private/TakeRecorderActorSource.cpp b/Engine/Plugins/VirtualProduction/Takes/Source/TakeRecorderSources/Private/TakeRecorderActorSource.cpp
--- a/Engine/Plugins/VirtualProduction/Takes/Source/TakeRecorderSources/Private/TakeRecorderActorSource.cpp
+++ b/Engine/Plugins/VirtualProduction/Takes/Source/TakeRecorderSources/Private/TakeRecorderActorSource.cpp
@@ -46,7 +46,6 @@
 #if WITH_EDITOR
 #include "Tracks/MovieSceneSpawnTrack.h"
 #endif
-#include "SequencerUtilities.h"
 
 #include UE_INLINE_GENERATED_CPP_BY_NAME(TakeRecorderActorSource)
 
@@ -257,16 +256,13 @@ TArray<UTakeRecorderSource*> UTakeRecorderActorSource::PreRecording(ULevelSequen
 	{
 		// We need to store the object template in the Movie Scene (because it's a complex UObject)
 		// instead of trying to place this data into the non-UObject safe data stream.
-		UE::Sequencer::FCreateBindingParams CreateBindingParams;
-		CreateBindingParams.bSpawnable = true;
-		CreateBindingParams.bAllowCustomBinding = true;
-		CreateBindingParams.BindingNameOverride = ActorToRecord->GetActorLabel();
-		CachedObjectBindingGuid = FSequencerUtilities::CreateOrReplaceBinding(nullptr, InSequence, ActorToRecord, CreateBindingParams);
-		CachedObjectTemplate = CastChecked<AActor>(MovieSceneHelpers::GetObjectTemplate(InSequence, CachedObjectBindingGuid, MovieSceneHelpers::CreateTransientSharedPlaybackState(ActorToRecord, InSequence)));
+		FName UniqueTemplateName = MakeUniqueObjectName(TargetLevelSequence, ActorToRecord->GetClass(), NAME_None);
+		Header.TemplateName = UniqueTemplateName.ToString();
+		CachedObjectTemplate = CastChecked<AActor>(TargetLevelSequence->MakeSpawnableTemplateFromInstance(*ActorToRecord, UniqueTemplateName));
+		CachedObjectBindingGuid = MovieScene->AddSpawnable(ActorToRecord->GetActorLabel(), *CachedObjectTemplate);
 		
 		if (CachedObjectTemplate.IsValid())
 		{
-			Header.TemplateName = CachedObjectTemplate->GetName();
 			PostProcessCreatedObjectTemplateImpl(CachedObjectTemplate.Get());
 		}
 	}
@@ -420,6 +416,12 @@ void UTakeRecorderActorSource::CreateSectionRecordersRecursive(UObject* ObjectTo
 			ChildPossessable->SetParent(CachedObjectBindingGuid, MovieScene);
 		}
 
+		FMovieSceneSpawnable* ParentSpawnable = MovieScene->FindSpawnable(CachedObjectBindingGuid);
+		if (ParentSpawnable)
+		{
+			ParentSpawnable->AddChildPossessable(Guid);
+		}
+
 		// Bindings are stored relative to their context outer. Newly duplicated components have a different outer
 		TargetLevelSequence->BindPossessableObject(Guid, *ComponentToRecord, bNewComponentAdded ? CachedObjectTemplate.Get() : Target.Get());
 
diff --git a/Engine/Plugins/VirtualProduction/Takes/Source/TakeTrackRecorders/Private/MovieScene3DTransformTrackRecorder.cpp b/Engine/Plugins/VirtualProduction/Takes/Source/TakeTrackRecorders/Private/MovieScene3DTransformTrackRecorder.cpp
--- a/Engine/Plugins/VirtualProduction/Takes/Source/TakeTrackRecorders/Private/MovieScene3DTransformTrackRecorder.cpp
+++ b/Engine/Plugins/VirtualProduction/Takes/Source/TakeTrackRecorders/Private/MovieScene3DTransformTrackRecorder.cpp
@@ -15,7 +15,6 @@
 #include "Components/SkeletalMeshComponent.h"
 #include "SequenceRecorderUtils.h"
 #include "Animation/AnimData/IAnimationDataModel.h"
-#include "Evaluation/MovieSceneEvaluationState.h"
 
 #include UE_INLINE_GENERATED_CPP_BY_NAME(MovieScene3DTransformTrackRecorder)
 
@@ -299,24 +298,15 @@ void UMovieScene3DTransformTrackRecorder::FinalizeTrackImpl()
  
  	SlowTask.EnterProgressFrame();
  
-	// If recording a spawnable, update the spawnable object template to the first keyframe
-	if (MovieScene.IsValid() && ObjectGuid.IsValid())
-	{
-		UMovieSceneSequence* ThisSequence = MovieScene->GetTypedOuter<UMovieSceneSequence>();
-		UObject* WorldContext = ObjectToRecord.Get() ? ObjectToRecord.Get() : GWorld;
-		if (WorldContext)
-		{ 
-			TSharedRef<UE::MovieScene::FSharedPlaybackState> TransientPlaybackState = MovieSceneHelpers::CreateTransientSharedPlaybackState(WorldContext, ThisSequence);
-
-			if (AActor* ActorTemplatePtr = Cast<AActor>(MovieSceneHelpers::GetObjectTemplate(ThisSequence, ObjectGuid, TransientPlaybackState, 0)))
-			{
-				if (USceneComponent* RootComponent = ActorTemplatePtr->GetRootComponent())
-				{
-					RootComponent->SetRelativeTransform(FirstTransform);
-				}
-			}
-		}
-	}
+ 	// If recording a spawnable, update the spawnable object template to the first keyframe
+ 	if (MovieScene.IsValid() && ObjectGuid.IsValid())
+ 	{
+ 		FMovieSceneSpawnable* Spawnable = MovieScene->FindSpawnable(ObjectGuid);
+ 		if (Spawnable)
+ 		{
+ 			Spawnable->SpawnTransform = FirstTransform;
+ 		}
+ 	}
 }
 
 void UMovieScene3DTransformTrackRecorder::RecordSampleImpl(const FQualifiedFrameTime& CurrentTime)
diff --git a/Engine/Plugins/VirtualProduction/Takes/Source/TakeTrackRecorders/Private/MovieSceneAnimationTrackRecorder.cpp b/Engine/Plugins/VirtualProduction/Takes/Source/TakeTrackRecorders/Private/MovieSceneAnimationTrackRecorder.cpp
--- a/Engine/Plugins/VirtualProduction/Takes/Source/TakeTrackRecorders/Private/MovieSceneAnimationTrackRecorder.cpp
+++ b/Engine/Plugins/VirtualProduction/Takes/Source/TakeTrackRecorders/Private/MovieSceneAnimationTrackRecorder.cpp
@@ -92,6 +92,16 @@ static FGuid GetActorInSequence(AActor* InActor, UMovieScene* MovieScene)
 {
 	FString ActorTargetName = InActor->GetActorLabel();
 
+
+	for (int32 SpawnableCount = 0; SpawnableCount < MovieScene->GetSpawnableCount(); ++SpawnableCount)
+	{
+		const FMovieSceneSpawnable& Spawnable = MovieScene->GetSpawnable(SpawnableCount);
+		if (Spawnable.GetName() == ActorTargetName || Spawnable.Tags.Contains(*ActorTargetName))
+		{
+			return Spawnable.GetGuid();
+		}
+	}
+
 	for (int32 PossessableCount = 0; PossessableCount < MovieScene->GetPossessableCount(); ++PossessableCount)
 	{
 		const FMovieScenePossessable& Possessable = MovieScene->GetPossessable(PossessableCount);
diff --git a/Engine/Plugins/VirtualProduction/Takes/Source/TakeTrackRecorders/Private/MovieSceneSpawnTrackRecorder.cpp b/Engine/Plugins/VirtualProduction/Takes/Source/TakeTrackRecorders/Private/MovieSceneSpawnTrackRecorder.cpp
--- a/Engine/Plugins/VirtualProduction/Takes/Source/TakeTrackRecorders/Private/MovieSceneSpawnTrackRecorder.cpp
+++ b/Engine/Plugins/VirtualProduction/Takes/Source/TakeTrackRecorders/Private/MovieSceneSpawnTrackRecorder.cpp
@@ -7,7 +7,6 @@
 #include "Sections/MovieSceneBoolSection.h"
 #include "Tracks/MovieSceneSpawnTrack.h"
 #include "Channels/MovieSceneChannelProxy.h"
-#include "MovieSceneSequence.h"
 
 #include UE_INLINE_GENERATED_CPP_BY_NAME(MovieSceneSpawnTrackRecorder)
 
@@ -26,12 +25,7 @@ bool FMovieSceneSpawnTrackRecorderFactory::CanRecordObject(UObject* InObjectToRe
 
 void UMovieSceneSpawnTrackRecorder::CreateTrackImpl()
 {
-	UMovieSceneSequence* Sequence = MovieScene->GetTypedOuter<UMovieSceneSequence>();
-	if (!Sequence)
-	{
-		return;
-	}
-	if (!MovieSceneHelpers::IsBoundToAnySpawnable(Sequence, ObjectGuid, MovieSceneHelpers::CreateTransientSharedPlaybackState(this, Sequence)))
+	if (MovieScene->FindPossessable(ObjectGuid))
 	{
 		return;
 	}
diff --git a/Engine/Source/Editor/MovieSceneTools/Private/MovieSceneDynamicBindingBlueprintExtension.cpp b/Engine/Source/Editor/MovieSceneTools/Private/MovieSceneDynamicBindingBlueprintExtension.cpp
--- a/Engine/Source/Editor/MovieSceneTools/Private/MovieSceneDynamicBindingBlueprintExtension.cpp
+++ b/Engine/Source/Editor/MovieSceneTools/Private/MovieSceneDynamicBindingBlueprintExtension.cpp
@@ -113,7 +113,7 @@ void UMovieSceneDynamicBindingBlueprintExtension::HandleGenerateFunctionGraphs(F
 	};
 	FMovieSceneDynamicBindingUtils::IterateDynamicBindings(MovieScene, GenerateFunctionGraphs);
 
-	// This callback sets the generated function calls back onto the custom director blueprint bindings in the sequence,
+	// This callback sets the generated function calls back onto the spawnables/possessables in the sequence,
 	// and keeps a pointer to any needed arguments on this function, so that we can pass special values later
 	// at runtime.
 	TWeakObjectPtr<UMovieSceneSequence> WeakMovieSceneSequence(MovieSceneSequence);
diff --git a/Engine/Source/Editor/MovieSceneTools/Private/MovieSceneDynamicBindingCustomization.cpp b/Engine/Source/Editor/MovieSceneTools/Private/MovieSceneDynamicBindingCustomization.cpp
--- a/Engine/Source/Editor/MovieSceneTools/Private/MovieSceneDynamicBindingCustomization.cpp
+++ b/Engine/Source/Editor/MovieSceneTools/Private/MovieSceneDynamicBindingCustomization.cpp
@@ -41,23 +41,6 @@ TSharedRef<IPropertyTypeCustomization> FMovieSceneDynamicBindingCustomization::M
 	return Instance;
 }
 
-void FMovieSceneDynamicBindingCustomization::CustomizeChildren(TSharedRef<IPropertyHandle> InPropertyHandle, IDetailChildrenBuilder& ChildBuilder, IPropertyTypeCustomizationUtils& CustomizationUtils)
-{
-	FMovieSceneDirectorBlueprintEndpointCustomization::CustomizeChildren(InPropertyHandle, ChildBuilder, CustomizationUtils);
-	FStringView PropertyPath = InPropertyHandle->GetPropertyPath();
-	BindingIndex = 0;
-	if (PropertyPath.StartsWith(TEXT("Bindings[")))
-	{
-		int32 RightBounds;
-		if (PropertyPath.FindChar(']', RightBounds))
-		{
-			FStringView BindingIndexStr = PropertyPath.Mid(9, RightBounds);
-			// Grab the BindingIndex out of the PropertyPath
-			BindingIndex = FCString::Atoi(BindingIndexStr.GetData());
-		}
-	}
-}
-
 void FMovieSceneDynamicBindingCustomization::GetPayloadVariables(UObject* EditObject, void* RawData, FPayloadVariableMap& OutPayloadVariables) const
 {
 	const FMovieSceneDynamicBinding* DynamicBinding = static_cast<FMovieSceneDynamicBinding*>(RawData);
@@ -165,10 +148,10 @@ FMovieSceneDirectorBlueprintEndpointDefinition FMovieSceneDynamicBindingCustomiz
 	if (FMovieScenePossessable* Possessable = MovieScene->FindPossessable(ObjectBinding))
 	{
 		Definition.EndpointName = Possessable->GetName() + "_DynamicBinding";
-		if (BindingIndex > 0)
-		{
-			Definition.EndpointName += "_" + FString::FromInt(BindingIndex);
-		}
+	}
+	else if (FMovieSceneSpawnable* Spawnable = MovieScene->FindSpawnable(ObjectBinding))
+	{
+		Definition.EndpointName = Spawnable->GetName() + "_DynamicBinding";
 	}
 
 	return Definition;
@@ -188,12 +171,6 @@ void FMovieSceneDynamicBindingCustomization::OnCreateEndpoint(UMovieSceneSequenc
 
 		FMovieSceneDynamicBinding* DynamicBinding = static_cast<FMovieSceneDynamicBinding*>(RawData[Index]);
 
-		// Default call in editor to true
-		if (UK2Node_FunctionEntry* CallFunction = Cast<UK2Node_FunctionEntry>(NewEndpoint))
-		{
-			CallFunction->MetaData.bCallInEditor = true;
-		}
-
 		FMovieSceneDynamicBindingUtils::SetEndpoint(MovieScene, DynamicBinding, NewEndpoint);
 	}
 
diff --git a/Engine/Source/Editor/MovieSceneTools/Private/MovieSceneToolsModule.cpp b/Engine/Source/Editor/MovieSceneTools/Private/MovieSceneToolsModule.cpp
--- a/Engine/Source/Editor/MovieSceneTools/Private/MovieSceneToolsModule.cpp
+++ b/Engine/Source/Editor/MovieSceneTools/Private/MovieSceneToolsModule.cpp
@@ -117,11 +117,6 @@
 #include "TransformableHandle.h"
 #include "Constraints/ComponentConstraintChannelInterface.h"
 #include "Constraints/TransformConstraintChannelInterface.h"
-#include "Bindings/MovieSceneReplaceableDirectorBlueprintBinding.h"
-#include "Bindings/MovieSceneSpawnableDirectorBlueprintBinding.h"
-#include "EntitySystem/MovieSceneSharedPlaybackState.h"
-#include "MovieSceneDynamicBindingUtils.h"
-#include "Kismet2/KismetEditorUtilities.h"
 
 #define LOCTEXT_NAMESPACE "FMovieSceneToolsModule"
 
@@ -255,8 +250,6 @@ void FMovieSceneToolsModule::StartupModule()
 	UMovieSceneEventSectionBase::RemoveForCookEvent.BindStatic(RemoveForCookEventSection);
 	UMovieScene::IsTrackClassAllowedEvent.BindStatic(IsTrackClassAllowed);
 	ULevelSequence::PostDuplicateEvent.BindStatic(PostDuplicateEvent);
-	FixupDynamicBindingsHandle = ULevelSequence::FixupDynamicBindingsEvent.AddStatic(FixupDynamicBindingsEvent);
-
 
 	auto OnObjectsReplaced = [](const TMap<UObject*, UObject*>& ReplacedObjects)
 	{
@@ -294,7 +287,6 @@ void FMovieSceneToolsModule::ShutdownModule()
 	UMovieSceneEventSectionBase::RemoveForCookEvent = UMovieSceneEventSectionBase::FRemoveForCookEvent();
 	UMovieScene::IsTrackClassAllowedEvent = UMovieScene::FIsTrackClassAllowedEvent();
 	ULevelSequence::PostDuplicateEvent = ULevelSequence::FPostDuplicateEvent();
-	ULevelSequence::FixupDynamicBindingsEvent.Remove(FixupDynamicBindingsHandle);
 
 	if (ICurveEditorModule* CurveEditorModule = FModuleManager::GetModulePtr<ICurveEditorModule>("CurveEditor"))
 	{
@@ -448,19 +440,6 @@ void FMovieSceneToolsModule::PostDuplicateEvent(ULevelSequence* LevelSequence)
 	}
 }
 
-
-void FMovieSceneToolsModule::FixupDynamicBindingsEvent(ULevelSequence* LevelSequence)
-{
-	if (LevelSequence)
-	{
-		if (UBlueprint* DirectorBlueprint = LevelSequence->GetDirectorBlueprint())
-		{
-			FMovieSceneDynamicBindingUtils::EnsureBlueprintExtensionCreated(LevelSequence, DirectorBlueprint);
-			FKismetEditorUtilities::CompileBlueprint(DirectorBlueprint);
-		}
-	}
-}
-
 bool FMovieSceneToolsModule::UpgradeLegacyEventEndpointForSection(UMovieSceneEventSectionBase* Section)
 {
 	UMovieSceneSequence*       Sequence           = Section->GetTypedOuter<UMovieSceneSequence>();
@@ -604,8 +583,6 @@ void FMovieSceneToolsModule::FixupPayloadParameterNameForSection(UMovieSceneEven
 
 void FMovieSceneToolsModule::FixupPayloadParameterNameForDynamicBinding(UMovieScene* MovieScene, UK2Node* InNode, FName OldPinName, FName NewPinName)
 {
-	using namespace UE::MovieScene;
-
 	check(MovieScene);
 
 	auto FixupPayloadParameterName = [InNode, OldPinName, NewPinName](FMovieSceneDynamicBinding& DynamicBinding)
@@ -620,26 +597,14 @@ void FMovieSceneToolsModule::FixupPayloadParameterNameForDynamicBinding(UMovieSc
 		}
 	};
 
-	UMovieSceneSequence* ThisSequence = MovieScene->GetTypedOuter<UMovieSceneSequence>();
-	TSharedRef<UE::MovieScene::FSharedPlaybackState> TransientPlaybackState = MovieSceneHelpers::CreateTransientSharedPlaybackState(GEditor->GetEditorWorldContext().World(), ThisSequence);
-
-	if (FMovieSceneBindingReferences* BindingReferences = MovieScene->GetTypedOuter<UMovieSceneSequence>()->GetBindingReferences())
+	for (int32 Index = 0, PossessableCount = MovieScene->GetPossessableCount(); Index < PossessableCount; ++Index)
 	{
-		for (FMovieSceneBindingReference& BindingReference : BindingReferences->GetAllReferences())
-		{
-			if (BindingReference.CustomBinding)
-			{
-				if (UMovieSceneReplaceableDirectorBlueprintBinding* ReplaceableDirectorBlueprintBinding = Cast<UMovieSceneReplaceableDirectorBlueprintBinding>(BindingReference.CustomBinding))
-				{
-					FixupPayloadParameterName(ReplaceableDirectorBlueprintBinding->DynamicBinding);
-				}
+		FixupPayloadParameterName(MovieScene->GetPossessable(Index).DynamicBinding);
+	}
 
-				if (UMovieSceneSpawnableDirectorBlueprintBinding* SpawnableDirectorBlueprintBinding = Cast<UMovieSceneSpawnableDirectorBlueprintBinding>(BindingReference.CustomBinding->AsSpawnable(TransientPlaybackState)))
-				{
-					FixupPayloadParameterName(SpawnableDirectorBlueprintBinding->DynamicBinding);
-				}
-			}
-		}
+	for (int32 Index = 0, SpawnableCount = MovieScene->GetSpawnableCount(); Index < SpawnableCount; ++Index)
+	{
+		FixupPayloadParameterName(MovieScene->GetSpawnable(Index).DynamicBinding);
 	}
 }
 
diff --git a/Engine/Source/Editor/MovieSceneTools/Public/MovieSceneDynamicBindingCustomization.h b/Engine/Source/Editor/MovieSceneTools/Public/MovieSceneDynamicBindingCustomization.h
--- a/Engine/Source/Editor/MovieSceneTools/Public/MovieSceneDynamicBindingCustomization.h
+++ b/Engine/Source/Editor/MovieSceneTools/Public/MovieSceneDynamicBindingCustomization.h
@@ -28,8 +28,6 @@ public:
 	static TSharedRef<IPropertyTypeCustomization> MakeInstance();
 	static TSharedRef<IPropertyTypeCustomization> MakeInstance(UMovieScene* InMovieScene, FGuid InObjectBinding);
 
-	virtual void CustomizeChildren(TSharedRef<IPropertyHandle> InPropertyHandle, IDetailChildrenBuilder& ChildBuilder, IPropertyTypeCustomizationUtils& CustomizationUtils) override;
-
 protected:
 
 	virtual void GetPayloadVariables(UObject* EditObject, void* RawData, FPayloadVariableMap& OutPayloadVariables) const override;
@@ -55,6 +53,5 @@ private:
 
 	UMovieScene* EditedMovieScene;
 	FGuid ObjectBinding;
-	int32 BindingIndex = 0;
 };
 
diff --git a/Engine/Source/Editor/MovieSceneTools/Public/MovieSceneDynamicBindingUtils.h b/Engine/Source/Editor/MovieSceneTools/Public/MovieSceneDynamicBindingUtils.h
--- a/Engine/Source/Editor/MovieSceneTools/Public/MovieSceneDynamicBindingUtils.h
+++ b/Engine/Source/Editor/MovieSceneTools/Public/MovieSceneDynamicBindingUtils.h
@@ -5,13 +5,8 @@
 #include "CoreTypes.h"
 #include "MovieScene.h"
 #include "MovieSceneDynamicBinding.h"
-#include "Bindings/MovieSceneReplaceableDirectorBlueprintBinding.h"
-#include "Bindings/MovieSceneSpawnableDirectorBlueprintBinding.h"
-#include "EntitySystem/MovieSceneSharedPlaybackState.h"
-#include "Evaluation/MovieSceneEvaluationState.h"
-#include "Editor.h"
-#include "MovieSceneBindingReferences.h"
-#include "MovieSceneCommonHelpers.h"
+#include "MovieScenePossessable.h"
+#include "MovieSceneSpawnable.h"
 
 #include "MovieSceneDynamicBindingUtils.generated.h"
 
@@ -20,8 +15,6 @@
  */
 struct MOVIESCENETOOLS_API FMovieSceneDynamicBindingUtils
 {
-	using FSharedPlaybackState = UE::MovieScene::FSharedPlaybackState;
-
 	/**
 	 * Set an endpoint on the given dynamic binding.
 	 */
@@ -38,9 +31,19 @@ struct MOVIESCENETOOLS_API FMovieSceneDynamicBindingUtils
 	template<typename Callback>
 	static void IterateDynamicBindings(UMovieScene* InMovieScene, Callback&& InCallback)
 	{
-		UMovieSceneSequence* ThisSequence = InMovieScene->GetTypedOuter<UMovieSceneSequence>();
-		TSharedRef<UE::MovieScene::FSharedPlaybackState> TransientPlaybackState = MovieSceneHelpers::CreateTransientSharedPlaybackState(GEditor->GetEditorWorldContext().World(), ThisSequence);
-		ThisSequence->IterateDynamicBindings(TransientPlaybackState, InCallback);
+		for (int32 Index = 0, PossessableCount = InMovieScene->GetPossessableCount(); Index < PossessableCount; ++Index)
+		{
+			FMovieScenePossessable& Possessable = InMovieScene->GetPossessable(Index);
+			FMovieSceneDynamicBinding& DynamicBinding = Possessable.DynamicBinding;
+			InCallback(Possessable.GetGuid(), DynamicBinding);
+		}
+
+		for (int32 Index = 0, SpawnableCount = InMovieScene->GetSpawnableCount(); Index < SpawnableCount; ++Index)
+		{
+			FMovieSceneSpawnable& Spawnable = InMovieScene->GetSpawnable(Index);
+			FMovieSceneDynamicBinding& DynamicBinding = Spawnable.DynamicBinding;
+			InCallback(Spawnable.GetGuid(), DynamicBinding);
+		}
 	}
 
 	/**
diff --git a/Engine/Source/Editor/MovieSceneTools/Public/MovieSceneToolsModule.h b/Engine/Source/Editor/MovieSceneTools/Public/MovieSceneToolsModule.h
--- a/Engine/Source/Editor/MovieSceneTools/Public/MovieSceneToolsModule.h
+++ b/Engine/Source/Editor/MovieSceneTools/Public/MovieSceneToolsModule.h
@@ -96,7 +96,6 @@ private:
 	static void RemoveForCookEventSection(UMovieSceneEventSectionBase* Section);
 	static bool IsTrackClassAllowed(UClass* InClass);
 	static void PostDuplicateEvent(ULevelSequence* LevelSequence);
-	static void FixupDynamicBindingsEvent(ULevelSequence* LevelSequence);
 
 private:
 
@@ -149,7 +148,6 @@ private:
 	FDelegateHandle FixupDynamicBindingPayloadParameterNameHandle;
 	FDelegateHandle FixupEventSectionPayloadParameterNameHandle;
 	FDelegateHandle UpgradeLegacyEventEndpointHandle;
-	FDelegateHandle FixupDynamicBindingsHandle;
 
 	FDelegateHandle OnObjectsReplacedHandle;
 
diff --git a/Engine/Source/Editor/SequenceRecorder/Private/ActorRecording.cpp b/Engine/Source/Editor/SequenceRecorder/Private/ActorRecording.cpp
--- a/Engine/Source/Editor/SequenceRecorder/Private/ActorRecording.cpp
+++ b/Engine/Source/Editor/SequenceRecorder/Private/ActorRecording.cpp
@@ -24,10 +24,6 @@
 #include "IAssetTools.h"
 #include "AssetToolsModule.h"
 #include "Subsystems/AssetEditorSubsystem.h"
-#include "SequencerUtilities.h"
-#include "EntitySystem/MovieSceneSharedPlaybackState.h"
-#include "Evaluation/MovieSceneEvaluationState.h"
-#include "MovieSceneCommonHelpers.h"
 
 static const FName SequencerActorTag(TEXT("SequencerActor"));
 static const FName MovieSceneSectionRecorderFactoryName("MovieSceneSectionRecorderFactory");
@@ -178,6 +174,26 @@ bool UActorRecording::StartRecording(ULevelSequence* CurrentSequence, float Curr
 	return true;
 }
 
+static FString GetUniqueSpawnableName(UMovieScene* MovieScene, const FString& BaseName)
+{
+	FString BlueprintName = BaseName;
+	auto DuplName = [&](FMovieSceneSpawnable& InSpawnable)
+	{
+		return InSpawnable.GetName() == BlueprintName;
+	};
+
+	int32 Index = 2;
+	FString UniqueString;
+	while (MovieScene->FindSpawnable(DuplName))
+	{
+		BlueprintName.RemoveFromEnd(UniqueString);
+		UniqueString = FString::Printf(TEXT(" (%d)"), Index++);
+		BlueprintName += UniqueString;
+	}
+
+	return BlueprintName;
+}
+
 void UActorRecording::GetNonSceneActorComponents(TArray<UActorComponent*>& OutArray)
 {
 	if (GetActorToRecord() != nullptr)
@@ -419,6 +435,15 @@ FGuid UActorRecording::GetActorInSequence(AActor* InActor, ULevelSequence* Curre
 
 	UMovieScene* MovieScene = CurrentSequence->GetMovieScene();
 	
+	for (int32 SpawnableCount = 0; SpawnableCount < MovieScene->GetSpawnableCount(); ++SpawnableCount)
+	{
+		const FMovieSceneSpawnable& Spawnable = MovieScene->GetSpawnable(SpawnableCount);
+		if (Spawnable.GetName() == ActorTargetName || Spawnable.Tags.Contains(*InActor->GetActorLabel()))
+		{
+			return Spawnable.GetGuid();
+		}
+	}
+
 	for (int32 PossessableCount = 0; PossessableCount < MovieScene->GetPossessableCount(); ++PossessableCount)
 	{
 		const FMovieScenePossessable& Possessable = MovieScene->GetPossessable(PossessableCount);
@@ -457,7 +482,7 @@ void UActorRecording::StartRecordingActorProperties(ULevelSequence* CurrentSeque
 			}
 			else
 			{
-				FString TemplateName = MovieSceneHelpers::MakeUniqueBindingName(MovieScene, Actor->GetName());
+				FString TemplateName = GetUniqueSpawnableName(MovieScene, Actor->GetName());
 
 				AActor* ObjectTemplate = CastChecked<AActor>(CurrentSequence->MakeSpawnableTemplateFromInstance(*Actor, *TemplateName));
 
@@ -481,11 +506,7 @@ void UActorRecording::StartRecordingActorProperties(ULevelSequence* CurrentSeque
 						Pawn->AutoPossessPlayer = EAutoReceiveInput::Disabled;
 					}
 
-					UE::Sequencer::FCreateBindingParams CreateBindingParams;
-					CreateBindingParams.bAllowCustomBinding = true;
-					CreateBindingParams.bSpawnable = true;
-					CreateBindingParams.BindingNameOverride = ObjectBindingName;
-					Guid = FSequencerUtilities::CreateOrReplaceBinding(nullptr, CurrentSequence, ObjectTemplate, CreateBindingParams);
+					Guid = MovieScene->AddSpawnable(ObjectBindingName, *ObjectTemplate);
 				}
 			}
 		}
@@ -503,6 +524,15 @@ void UActorRecording::StartRecordingActorProperties(ULevelSequence* CurrentSeque
 				Possessable->SetName(ObjectBindingName);
 			}
 
+			if (FMovieSceneSpawnable* Spawnable = MovieScene->FindSpawnable(Guid))
+			{
+				if (!Spawnable->Tags.Contains(*Actor->GetActorLabel()))
+				{
+					Spawnable->Tags.AddUnique(FName(*Actor->GetActorLabel()));
+				}
+				Spawnable->SetName(ObjectBindingName);
+			}
+
 			// add our folder
 			FindOrAddFolder(MovieScene);
 
@@ -616,6 +646,12 @@ TSharedPtr<FMovieSceneAnimationSectionRecorder> UActorRecording::StartRecordingC
 		ChildPossessable->SetParent(Guid, OwnerMovieScene);
 	}
 
+	FMovieSceneSpawnable* ParentSpawnable = OwnerMovieScene->FindSpawnable(Guid);
+	if (ParentSpawnable)
+	{
+		ParentSpawnable->AddChildPossessable(PossessableGuid);
+	}
+
 	CurrentSequence->BindPossessableObject(PossessableGuid, *ActorComponent, BindingContext);
 
 	const USequenceRecorderSettings* Settings = GetDefault<USequenceRecorderSettings>();
@@ -713,10 +749,10 @@ bool UActorRecording::StopRecording(ULevelSequence* OriginalSequence, float Curr
 		UMovieScene* MovieScene = CurrentSequence->GetMovieScene();
 		check(MovieScene);
 
-		FMovieScenePossessable* Possessable = MovieScene->FindPossessable(Guid);
-		if(Possessable)
+		FMovieSceneSpawnable* Spawnable = MovieScene->FindSpawnable(Guid);
+		if(Spawnable)
 		{
-			ActorName = Possessable->GetName();
+			ActorName = Spawnable->GetName();
 		}
 	}
 
@@ -1009,9 +1045,10 @@ void UActorRecording::StartRecordingNewComponents(ULevelSequence* CurrentSequenc
 		const USequenceRecorderSettings* Settings = GetDefault<USequenceRecorderSettings>();
 		if (!bRecordToPossessable)
 		{
-			TSharedRef<UE::MovieScene::FSharedPlaybackState> TransientPlaybackState = MovieSceneHelpers::CreateTransientSharedPlaybackState(GetActorToRecord()->GetWorld(), CurrentSequence);
+			FMovieSceneSpawnable* Spawnable = MovieScene->FindSpawnable(Guid);
+			check(Spawnable);
 
-			AActor* ObjectTemplate = Cast<AActor>(MovieSceneHelpers::GetObjectTemplate(CurrentSequence, Guid, TransientPlaybackState, 0));
+			AActor* ObjectTemplate = CastChecked<AActor>(Spawnable->GetObjectTemplate());
 
 			for (UActorComponent* ActorComponent : NewComponents)
 			{
diff --git a/Engine/Source/Editor/SequenceRecorder/Private/Sections/MovieScene3DTransformSectionRecorder.cpp b/Engine/Source/Editor/SequenceRecorder/Private/Sections/MovieScene3DTransformSectionRecorder.cpp
--- a/Engine/Source/Editor/SequenceRecorder/Private/Sections/MovieScene3DTransformSectionRecorder.cpp
+++ b/Engine/Source/Editor/SequenceRecorder/Private/Sections/MovieScene3DTransformSectionRecorder.cpp
@@ -14,7 +14,6 @@
 #include "Channels/MovieSceneChannelProxy.h"
 #include "Animation/AnimData/IAnimationDataModel.h"
 #include "Animation/AnimSequence.h"
-#include "Evaluation/MovieSceneEvaluationState.h"
 
 TSharedPtr<IMovieSceneSectionRecorder> FMovieScene3DTransformSectionRecorderFactory::CreateSectionRecorder(const FActorRecordingSettings& InActorRecordingSettings) const
 {
@@ -319,15 +318,10 @@ void FMovieScene3DTransformSectionRecorder::FinalizeSection(float CurrentTime)
 	// If recording a spawnable, update the spawnable object template to the first keyframe
 	if (MovieScene.IsValid() && Guid.IsValid())
 	{
-		UMovieSceneSequence* ThisSequence = MovieScene->GetTypedOuter<UMovieSceneSequence>();
-		TSharedRef<UE::MovieScene::FSharedPlaybackState> TransientPlaybackState = MovieSceneHelpers::CreateTransientSharedPlaybackState(GetSourceObject(), ThisSequence);
-
-		if (AActor* ActorTemplatePtr = Cast<AActor>(MovieSceneHelpers::GetObjectTemplate(ThisSequence, Guid, TransientPlaybackState, 0)))
+		FMovieSceneSpawnable* Spawnable = MovieScene->FindSpawnable(Guid);
+		if (Spawnable)
 		{
-			if (USceneComponent* RootComponent = ActorTemplatePtr->GetRootComponent())
-			{
-				RootComponent->SetRelativeTransform(FirstTransform);
-			}
+			Spawnable->SpawnTransform = FirstTransform;
 		}
 	}
 
diff --git a/Engine/Source/Editor/SequenceRecorder/SequenceRecorder.Build.cs b/Engine/Source/Editor/SequenceRecorder/SequenceRecorder.Build.cs
--- a/Engine/Source/Editor/SequenceRecorder/SequenceRecorder.Build.cs
+++ b/Engine/Source/Editor/SequenceRecorder/SequenceRecorder.Build.cs
@@ -41,7 +41,6 @@ namespace UnrealBuildTool.Rules
                     "Kismet",
                     "LiveLinkInterface",
 					"SceneOutliner",
-					"Sequencer"
                 }
                 );
 
diff --git a/Engine/Source/Editor/Sequencer/Private/MVVM/ViewModels/PossessableModel.cpp b/Engine/Source/Editor/Sequencer/Private/MVVM/ViewModels/PossessableModel.cpp
--- a/Engine/Source/Editor/Sequencer/Private/MVVM/ViewModels/PossessableModel.cpp
+++ b/Engine/Source/Editor/Sequencer/Private/MVVM/ViewModels/PossessableModel.cpp
@@ -132,6 +132,16 @@ const FSlateBrush* FPossessableModel::GetIconOverlayBrush() const
 		}
 	}
 
+	UMovieScene*            MovieScene  = OwnerModel ? OwnerModel->GetMovieScene() : nullptr;
+	FMovieScenePossessable* Possessable = MovieScene ? MovieScene->FindPossessable(ObjectBindingID) : nullptr;
+	if (Possessable && Possessable->DynamicBinding.WeakEndpoint.IsValid())
+	{
+		if (Possessable->DynamicBinding.WeakEndpoint.IsValid())
+		{
+			return FAppStyle::GetBrush("Sequencer.DynamicBindingIconOverlay");
+		}
+	}
+
 	UMovieSceneSequence* Sequence = OwnerModel ? OwnerModel->GetSequence() : nullptr;
 	if (Sequence)
 	{
@@ -145,10 +155,6 @@ const FSlateBrush* FPossessableModel::GetIconOverlayBrush() const
 				}
 			}
 		}
-		else if (const FSlateBrush* CustomBrush = Sequence->GetCustomBrushForBinding(ObjectBindingID))
-		{
-			return CustomBrush;
-		}
 	}
 
 	return nullptr;
diff --git a/Engine/Source/Editor/Sequencer/Private/MovieSceneCopyableBinding.h b/Engine/Source/Editor/Sequencer/Private/MovieSceneCopyableBinding.h
--- a/Engine/Source/Editor/Sequencer/Private/MovieSceneCopyableBinding.h
+++ b/Engine/Source/Editor/Sequencer/Private/MovieSceneCopyableBinding.h
@@ -61,7 +61,4 @@ public:
 	
 	UPROPERTY()
 	int32 NumCustomBindings;
-
-	UPROPERTY()
-	TArray<int32> PreviewSpawnableBindings;
 };
diff --git a/Engine/Source/Editor/Sequencer/Private/Sequencer.cpp b/Engine/Source/Editor/Sequencer/Private/Sequencer.cpp
--- a/Engine/Source/Editor/Sequencer/Private/Sequencer.cpp
+++ b/Engine/Source/Editor/Sequencer/Private/Sequencer.cpp
@@ -1616,7 +1616,18 @@ void FSequencer::UpdateSequencerCustomizations(const UMovieSceneSequence* Previo
 	// TODO: Allow customization of which custom bindings are allowed? For now just iterate over all subclasses and sort by priority
 
 	// Cache custom spawnable types
-	MovieSceneHelpers::GetPrioritySortedCustomBindingTypes(SupportedCustomBindingTypes);
+	SupportedCustomBindingTypes.Empty();
+	for (TObjectIterator<UClass> It; It; ++It)
+	{
+		if (It->IsChildOf(UMovieSceneCustomBinding::StaticClass()) && !It->HasAnyClassFlags(CLASS_Abstract))
+		{
+			SupportedCustomBindingTypes.Add(*It);
+		}
+	}
+
+	// Sort by spawner priority to allow disambiguation for similar object types
+	SupportedCustomBindingTypes.Sort([](const TSubclassOf<UMovieSceneCustomBinding>& A, const TSubclassOf<UMovieSceneCustomBinding>& B) {
+		return A && B && A->GetDefaultObject<UMovieSceneCustomBinding>()->GetCustomBindingPriority() > B->GetDefaultObject<UMovieSceneCustomBinding>()->GetCustomBindingPriority(); });
 }
 
 void FSequencer::RerunConstructionScripts()
