Motion Design Transition: Added Node Interface implemented in Blueprints so that Function Libraries can retrieve the Behavior Instance Context

diff --git a/Engine/Plugins/Experimental/Avalanche/Source/AvalancheTransition/Private/AvaTransitionUtils.cpp b/Engine/Plugins/Experimental/Avalanche/Source/AvalancheTransition/Private/AvaTransitionUtils.cpp
new file mode 100644
--- /dev/null
+++ b/Engine/Plugins/Experimental/Avalanche/Source/AvalancheTransition/Private/AvaTransitionUtils.cpp
@@ -0,0 +1,20 @@
+// Copyright Epic Games, Inc. All Rights Reserved.
+
+#include "AvaTransitionUtils.h"
+#include "Behavior/AvaTransitionBehaviorActor.h"
+#include "Execution/AvaTransitionExecutionContext.h"
+
+namespace UE::AvaTransition
+{
+
+const FAvaTransitionBehaviorInstance* GetBehaviorInstance(const FStateTreeExecutionContext& InExecutionContext)
+{
+	UObject* Owner = InExecutionContext.GetOwner();
+	if (ensure(Owner && Owner->IsA<AAvaTransitionBehaviorActor>()))
+	{
+		return static_cast<const FAvaTransitionExecutionContext&>(InExecutionContext).GetBehaviorInstance();
+	}
+	return nullptr;
+}
+
+}
diff --git a/Engine/Plugins/Experimental/Avalanche/Source/AvalancheTransition/Private/Behavior/AvaTransitionBehaviorInstance.cpp b/Engine/Plugins/Experimental/Avalanche/Source/AvalancheTransition/Private/Behavior/AvaTransitionBehaviorInstance.cpp
--- a/Engine/Plugins/Experimental/Avalanche/Source/AvalancheTransition/Private/Behavior/AvaTransitionBehaviorInstance.cpp
+++ b/Engine/Plugins/Experimental/Avalanche/Source/AvalancheTransition/Private/Behavior/AvaTransitionBehaviorInstance.cpp
@@ -64,6 +64,11 @@ FAvaTransitionContext& FAvaTransitionBehaviorInstance::GetTransitionContext()
 	return TransitionContext;
 }
 
+FAvaTransitionSceneOwner FAvaTransitionBehaviorInstance::GetTransitionSceneOwner() const
+{
+	return TransitionSceneOwner;
+}
+
 void FAvaTransitionBehaviorInstance::SetTransitionType(EAvaTransitionType InTransitionType)
 {
 	TransitionContext.TransitionType = InTransitionType;
@@ -243,7 +248,7 @@ TOptional<FAvaTransitionExecutionContext> FAvaTransitionBehaviorInstance::MakeCo
 
 	const FStateTreeReference& StateTreeReference = InBehavior->GetStateTreeReference();
 
-	FAvaTransitionExecutionContext Context(InBehavior->AsUObject(), *StateTreeReference.GetStateTree(), InstanceData);
+	FAvaTransitionExecutionContext Context(*this, InBehavior->AsUObject(), *StateTreeReference.GetStateTree(), InstanceData);
 	if (!Context.IsValid())
 	{
 		return TOptional<FAvaTransitionExecutionContext>();
diff --git a/Engine/Plugins/Experimental/Avalanche/Source/AvalancheTransition/Private/Behavior/AvaTransitionInstanceCache.cpp b/Engine/Plugins/Experimental/Avalanche/Source/AvalancheTransition/Private/Behavior/AvaTransitionInstanceCache.cpp
new file mode 100644
--- /dev/null
+++ b/Engine/Plugins/Experimental/Avalanche/Source/AvalancheTransition/Private/Behavior/AvaTransitionInstanceCache.cpp
@@ -0,0 +1,33 @@
+// Copyright Epic Games, Inc. All Rights Reserved.
+
+#include "AvaTransitionUtils.h"
+#include "Behavior/AvaTransitionBehaviorInstance.h"
+#include "Behavior/AvaTransitionBehaviorInstanceCache.h"
+
+const FAvaTransitionContext* FAvaTransitionBehaviorInstanceCache::GetTransitionContext() const
+{
+	if (CachedBehaviorInstance && CachedSceneOwner.IsValid())
+	{
+		return &CachedBehaviorInstance->GetTransitionContext();
+	}
+	return nullptr;
+}
+
+void FAvaTransitionBehaviorInstanceCache::SetCachedInstanceDataFromContext(const FStateTreeExecutionContext& InExecutionContext)
+{
+	if (const FAvaTransitionBehaviorInstance* BehaviorInstance = UE::AvaTransition::GetBehaviorInstance(InExecutionContext))
+	{
+		CachedBehaviorInstance = BehaviorInstance;
+		CachedSceneOwner = BehaviorInstance->GetTransitionSceneOwner();
+	}
+	else
+	{
+		ClearCachedInstanceData();
+	}
+}
+
+void FAvaTransitionBehaviorInstanceCache::ClearCachedInstanceData()
+{
+	CachedBehaviorInstance = nullptr;
+	CachedSceneOwner = FAvaTransitionSceneOwner(nullptr);
+}
diff --git a/Engine/Plugins/Experimental/Avalanche/Source/AvalancheTransition/Private/Conditions/AvaTransitionConditionBlueprint.cpp b/Engine/Plugins/Experimental/Avalanche/Source/AvalancheTransition/Private/Conditions/AvaTransitionConditionBlueprint.cpp
new file mode 100644
--- /dev/null
+++ b/Engine/Plugins/Experimental/Avalanche/Source/AvalancheTransition/Private/Conditions/AvaTransitionConditionBlueprint.cpp
@@ -0,0 +1,23 @@
+// Copyright Epic Games, Inc. All Rights Reserved.
+
+#include "Conditions/AvaTransitionConditionBlueprint.h"
+#include "Misc/ScopeExit.h"
+
+bool UAvaTransitionConditionBlueprint::TestCondition(FStateTreeExecutionContext& InContext) const
+{
+	if (bHasTestCondition)
+	{
+		BehaviorInstanceCache.SetCachedInstanceDataFromContext(InContext);
+		ON_SCOPE_EXIT
+		{
+			BehaviorInstanceCache.ClearCachedInstanceData();
+		};
+		return Super::TestCondition(InContext);
+	}
+	return false;
+}
+
+const FAvaTransitionBehaviorInstanceCache& UAvaTransitionConditionBlueprint::GetBehaviorInstanceCache() const
+{
+	return BehaviorInstanceCache;
+}
diff --git a/Engine/Plugins/Experimental/Avalanche/Source/AvalancheTransition/Private/Execution/AvaTransitionExecutionContext.cpp b/Engine/Plugins/Experimental/Avalanche/Source/AvalancheTransition/Private/Execution/AvaTransitionExecutionContext.cpp
--- a/Engine/Plugins/Experimental/Avalanche/Source/AvalancheTransition/Private/Execution/AvaTransitionExecutionContext.cpp
+++ b/Engine/Plugins/Experimental/Avalanche/Source/AvalancheTransition/Private/Execution/AvaTransitionExecutionContext.cpp
@@ -2,8 +2,9 @@
 
 #include "AvaTransitionExecutionContext.h"
 
-FAvaTransitionExecutionContext::FAvaTransitionExecutionContext(UObject& InOwner, const UStateTree& InStateTree, FStateTreeInstanceData& InInstanceData)
+FAvaTransitionExecutionContext::FAvaTransitionExecutionContext(const FAvaTransitionBehaviorInstance& InBehaviorInstance, UObject& InOwner, const UStateTree& InStateTree, FStateTreeInstanceData& InInstanceData)
 	: FStateTreeExecutionContext(InOwner, InStateTree, InInstanceData)
+	, BehaviorInstance(InBehaviorInstance)
 {
 }
 
diff --git a/Engine/Plugins/Experimental/Avalanche/Source/AvalancheTransition/Private/Execution/AvaTransitionExecutionContext.h b/Engine/Plugins/Experimental/Avalanche/Source/AvalancheTransition/Private/Execution/AvaTransitionExecutionContext.h
--- a/Engine/Plugins/Experimental/Avalanche/Source/AvalancheTransition/Private/Execution/AvaTransitionExecutionContext.h
+++ b/Engine/Plugins/Experimental/Avalanche/Source/AvalancheTransition/Private/Execution/AvaTransitionExecutionContext.h
@@ -5,17 +5,26 @@
 #include "Containers/UnrealString.h"
 #include "StateTreeExecutionContext.h"
 
+struct FAvaTransitionBehaviorInstance;
+
 struct FAvaTransitionExecutionContext : FStateTreeExecutionContext
 {
-	FAvaTransitionExecutionContext(UObject& InOwner, const UStateTree& InStateTree, FStateTreeInstanceData& InInstanceData);
+	FAvaTransitionExecutionContext(const FAvaTransitionBehaviorInstance& InBehaviorInstance, UObject& InOwner, const UStateTree& InStateTree, FStateTreeInstanceData& InInstanceData);
 
 	void SetSceneDescription(FString&& InSceneDescription);
 
+	const FAvaTransitionBehaviorInstance* GetBehaviorInstance() const
+	{
+		return &BehaviorInstance;
+	}
+
 protected:
 	//~ Begin FStateTreeExecutionContext
 	virtual FString GetInstanceDescription() const override;
 	//~ End FStateTreeExecutionContext
 
 private:
+	const FAvaTransitionBehaviorInstance& BehaviorInstance;
+
 	FString SceneDescription;
 };
diff --git a/Engine/Plugins/Experimental/Avalanche/Source/AvalancheTransition/Private/Tasks/AvaTransitionTaskBlueprint.cpp b/Engine/Plugins/Experimental/Avalanche/Source/AvalancheTransition/Private/Tasks/AvaTransitionTaskBlueprint.cpp
new file mode 100644
--- /dev/null
+++ b/Engine/Plugins/Experimental/Avalanche/Source/AvalancheTransition/Private/Tasks/AvaTransitionTaskBlueprint.cpp
@@ -0,0 +1,20 @@
+// Copyright Epic Games, Inc. All Rights Reserved.
+
+#include "Tasks/AvaTransitionTaskBlueprint.h"
+
+EStateTreeRunStatus UAvaTransitionTaskBlueprint::EnterState(FStateTreeExecutionContext& InExecutionContext, const FStateTreeTransitionResult& InTransition)
+{
+	BehaviorInstanceCache.SetCachedInstanceDataFromContext(InExecutionContext);
+	return Super::EnterState(InExecutionContext, InTransition);
+}
+
+void UAvaTransitionTaskBlueprint::ExitState(FStateTreeExecutionContext& InExecutionContext, const FStateTreeTransitionResult& InTransition)
+{
+	Super::ExitState(InExecutionContext, InTransition);
+	BehaviorInstanceCache.ClearCachedInstanceData();
+}
+
+const FAvaTransitionBehaviorInstanceCache& UAvaTransitionTaskBlueprint::GetBehaviorInstanceCache() const
+{
+	return BehaviorInstanceCache;
+}
diff --git a/Engine/Plugins/Experimental/Avalanche/Source/AvalancheTransition/Public/AvaTransitionUtils.h b/Engine/Plugins/Experimental/Avalanche/Source/AvalancheTransition/Public/AvaTransitionUtils.h
--- a/Engine/Plugins/Experimental/Avalanche/Source/AvalancheTransition/Public/AvaTransitionUtils.h
+++ b/Engine/Plugins/Experimental/Avalanche/Source/AvalancheTransition/Public/AvaTransitionUtils.h
@@ -4,6 +4,9 @@
 
 #include "StateTreeTypes.h"
 
+struct FAvaTransitionBehaviorInstance;
+struct FStateTreeExecutionContext;
+
 namespace UE::AvaTransition
 {
 	/** Safely gets the Instance Data if the Struct type matches or nullptr if there is type mismatch */
@@ -18,4 +21,10 @@ namespace UE::AvaTransition
 		}
 		return nullptr;
 	}
+
+	/**
+	 * Retrieves the Behavior Instance from the Execution Context
+	 * The Execution Context *must* have originated from Motion Design to return a valid pointer
+	 */
+	const FAvaTransitionBehaviorInstance* GetBehaviorInstance(const FStateTreeExecutionContext& InExecutionContext);
 }
diff --git a/Engine/Plugins/Experimental/Avalanche/Source/AvalancheTransition/Public/Behavior/AvaTransitionBehaviorInstance.h b/Engine/Plugins/Experimental/Avalanche/Source/AvalancheTransition/Public/Behavior/AvaTransitionBehaviorInstance.h
--- a/Engine/Plugins/Experimental/Avalanche/Source/AvalancheTransition/Public/Behavior/AvaTransitionBehaviorInstance.h
+++ b/Engine/Plugins/Experimental/Avalanche/Source/AvalancheTransition/Public/Behavior/AvaTransitionBehaviorInstance.h
@@ -48,6 +48,8 @@ struct FAvaTransitionBehaviorInstance
 
 	FAvaTransitionContext& GetTransitionContext();
 
+	FAvaTransitionSceneOwner GetTransitionSceneOwner() const;
+
 	void SetTransitionType(EAvaTransitionType InTransitionType);
 
 	bool Setup();
diff --git a/Engine/Plugins/Experimental/Avalanche/Source/AvalancheTransition/Public/Behavior/AvaTransitionBehaviorInstanceCache.h b/Engine/Plugins/Experimental/Avalanche/Source/AvalancheTransition/Public/Behavior/AvaTransitionBehaviorInstanceCache.h
new file mode 100644
--- /dev/null
+++ b/Engine/Plugins/Experimental/Avalanche/Source/AvalancheTransition/Public/Behavior/AvaTransitionBehaviorInstanceCache.h
@@ -0,0 +1,23 @@
+// Copyright Epic Games, Inc. All Rights Reserved.
+
+#pragma once
+
+#include "AvaTransitionScene.h"
+
+struct FAvaTransitionBehaviorInstance;
+struct FAvaTransitionContext;
+struct FStateTreeExecutionContext;
+
+struct FAvaTransitionBehaviorInstanceCache
+{
+	AVALANCHETRANSITION_API const FAvaTransitionContext* GetTransitionContext() const;
+
+	void SetCachedInstanceDataFromContext(const FStateTreeExecutionContext& InExecutionContext);
+
+	void ClearCachedInstanceData();
+
+private:
+	FAvaTransitionSceneOwner CachedSceneOwner;
+
+	const FAvaTransitionBehaviorInstance* CachedBehaviorInstance = nullptr;
+};
diff --git a/Engine/Plugins/Experimental/Avalanche/Source/AvalancheTransition/Public/Conditions/AvaTransitionConditionBlueprint.h b/Engine/Plugins/Experimental/Avalanche/Source/AvalancheTransition/Public/Conditions/AvaTransitionConditionBlueprint.h
--- a/Engine/Plugins/Experimental/Avalanche/Source/AvalancheTransition/Public/Conditions/AvaTransitionConditionBlueprint.h
+++ b/Engine/Plugins/Experimental/Avalanche/Source/AvalancheTransition/Public/Conditions/AvaTransitionConditionBlueprint.h
@@ -2,11 +2,25 @@
 
 #pragma once
 
+#include "Behavior/AvaTransitionBehaviorInstanceCache.h"
 #include "Blueprint/StateTreeConditionBlueprintBase.h"
+#include "IAvaTransitionNodeInterface.h"
 #include "AvaTransitionConditionBlueprint.generated.h"
 
 UCLASS(MinimalAPI, Abstract, DisplayName="Motion Design Transition Condition Blueprint")
-class UAvaTransitionConditionBlueprint : public UStateTreeConditionBlueprintBase
+class UAvaTransitionConditionBlueprint : public UStateTreeConditionBlueprintBase, public IAvaTransitionNodeInterface
 {
 	GENERATED_BODY()
+
+protected:
+	//~ Begin UStateTreeConditionBlueprintBase
+	AVALANCHETRANSITION_API virtual bool TestCondition(FStateTreeExecutionContext& InContext) const override;
+	//~ End UStateTreeConditionBlueprintBase
+
+	//~ Begin IAvaTransitionNodeInterface
+	AVALANCHETRANSITION_API virtual const FAvaTransitionBehaviorInstanceCache& GetBehaviorInstanceCache() const override;
+	//~ End IAvaTransitionNodeInterface
+
+private:
+	mutable FAvaTransitionBehaviorInstanceCache BehaviorInstanceCache;
 };
diff --git a/Engine/Plugins/Experimental/Avalanche/Source/AvalancheTransition/Public/IAvaTransitionNodeInterface.h b/Engine/Plugins/Experimental/Avalanche/Source/AvalancheTransition/Public/IAvaTransitionNodeInterface.h
new file mode 100644
--- /dev/null
+++ b/Engine/Plugins/Experimental/Avalanche/Source/AvalancheTransition/Public/IAvaTransitionNodeInterface.h
@@ -0,0 +1,23 @@
+// Copyright Epic Games, Inc. All Rights Reserved.
+
+#pragma once
+
+#include "UObject/Interface.h"
+#include "IAvaTransitionNodeInterface.generated.h"
+
+struct FAvaTransitionBehaviorInstanceCache;
+
+/** Interface that enable nodes that store a Behavior Instance Cache to expose this execution cache for others like function libraries to retrieve it */
+UINTERFACE(MinimalAPI, meta=(DisplayName="Motion Design Transition Node Interface", CannotImplementInterfaceInBlueprint))
+class UAvaTransitionNodeInterface : public UInterface
+{
+	GENERATED_BODY()
+};
+
+class IAvaTransitionNodeInterface
+{
+	GENERATED_BODY()
+
+public:
+	virtual const FAvaTransitionBehaviorInstanceCache& GetBehaviorInstanceCache() const = 0;
+};
diff --git a/Engine/Plugins/Experimental/Avalanche/Source/AvalancheTransition/Public/Tasks/AvaTransitionTaskBlueprint.h b/Engine/Plugins/Experimental/Avalanche/Source/AvalancheTransition/Public/Tasks/AvaTransitionTaskBlueprint.h
--- a/Engine/Plugins/Experimental/Avalanche/Source/AvalancheTransition/Public/Tasks/AvaTransitionTaskBlueprint.h
+++ b/Engine/Plugins/Experimental/Avalanche/Source/AvalancheTransition/Public/Tasks/AvaTransitionTaskBlueprint.h
@@ -2,11 +2,26 @@
 
 #pragma once
 
+#include "Behavior/AvaTransitionBehaviorInstanceCache.h"
 #include "Blueprint/StateTreeTaskBlueprintBase.h"
+#include "IAvaTransitionNodeInterface.h"
 #include "AvaTransitionTaskBlueprint.generated.h"
 
 UCLASS(MinimalAPI, Abstract, DisplayName="Motion Design Transition Task Blueprint")
-class UAvaTransitionTaskBlueprint : public UStateTreeTaskBlueprintBase
+class UAvaTransitionTaskBlueprint : public UStateTreeTaskBlueprintBase, public IAvaTransitionNodeInterface
 {
 	GENERATED_BODY()
+
+protected:
+	//~ Begin UStateTreeTaskBlueprintBase
+	AVALANCHETRANSITION_API virtual EStateTreeRunStatus EnterState(FStateTreeExecutionContext& InExecutionContext, const FStateTreeTransitionResult& InTransition) override;
+	AVALANCHETRANSITION_API virtual void ExitState(FStateTreeExecutionContext& InExecutionContext, const FStateTreeTransitionResult& InTransition) override;
+	//~ End UStateTreeTaskBlueprintBase
+
+	//~ Begin IAvaTransitionNodeInterface
+	AVALANCHETRANSITION_API virtual const FAvaTransitionBehaviorInstanceCache& GetBehaviorInstanceCache() const override;
+	//~ End IAvaTransitionNodeInterface
+
+private:
+	FAvaTransitionBehaviorInstanceCache BehaviorInstanceCache;
 };
