Motion Design Transition: Added RC Library for TL

diff --git a/Engine/Plugins/Experimental/Avalanche/Source/Avalanche/Private/Transition/AvaTransitionRCLibrary.cpp b/Engine/Plugins/Experimental/Avalanche/Source/Avalanche/Private/Transition/AvaTransitionRCLibrary.cpp
new file mode 100644
--- /dev/null
+++ b/Engine/Plugins/Experimental/Avalanche/Source/Avalanche/Private/Transition/AvaTransitionRCLibrary.cpp
@@ -0,0 +1,133 @@
+// Copyright Epic Games, Inc. All Rights Reserved.
+
+#include "Transition/AvaTransitionRCLibrary.h"
+#include "AvaRCControllerId.h"
+#include "AvaSceneSubsystem.h"
+#include "AvaTransitionContext.h"
+#include "AvaTransitionLayer.h"
+#include "AvaTransitionLayerUtils.h"
+#include "AvaTransitionSubsystem.h"
+#include "Behavior/AvaTransitionBehaviorInstance.h"
+#include "Behavior/AvaTransitionBehaviorInstanceCache.h"
+#include "Engine/Level.h"
+#include "Engine/World.h"
+#include "IAvaSceneInterface.h"
+#include "IAvaTransitionNodeInterface.h"
+#include "RCVirtualProperty.h"
+#include "RemoteControlPreset.h"
+#include "Transition/Extensions/IAvaTransitionRCExtension.h"
+
+namespace UE::Ava::Private
+{
+	URCVirtualPropertyBase* GetController(const FAvaRCControllerId& InControllerId, const UAvaSceneSubsystem& InSceneSubsystem, ULevel* InLevel)
+	{
+		IAvaSceneInterface* SceneInterface = InSceneSubsystem.GetSceneInterface(InLevel);
+		if (!SceneInterface)
+		{
+			return nullptr;
+		}
+
+		if (URemoteControlPreset* RemoteControlPreset = SceneInterface->GetRemoteControlPreset())
+		{
+			return InControllerId.FindController(RemoteControlPreset);
+		}
+
+		return nullptr;
+	}
+}
+
+bool UAvaTransitionRCLibrary::CompareRCControllerValues(const FAvaTransitionContext& InTransitionContext, const FAvaRCControllerId& InControllerId, EAvaTransitionComparisonResult InValueComparisonType)
+{
+	const FAvaTransitionScene* TransitionScene = InTransitionContext.GetTransitionScene();
+	if (!TransitionScene)
+	{
+		return false;
+	}
+
+	ULevel* Level = TransitionScene->GetLevel();
+	if (!Level || !Level->OwningWorld)
+	{
+		return false;
+	}
+
+	UAvaTransitionSubsystem* TransitionSubsystem = Level->OwningWorld->GetSubsystem<UAvaTransitionSubsystem>();
+	UAvaSceneSubsystem* SceneSubsystem = Level->OwningWorld->GetSubsystem<UAvaSceneSubsystem>();
+	if (!TransitionSubsystem || !SceneSubsystem)
+	{
+		return false;
+	}
+
+	URCVirtualPropertyBase* Controller = UE::Ava::Private::GetController(InControllerId, *SceneSubsystem, Level);
+	if (!Controller)
+	{
+		return false;
+	}
+
+	// Get all the Behavior Instances in the same Layer
+	TArray<const FAvaTransitionBehaviorInstance*> BehaviorInstances;
+	{
+		FAvaTransitionLayerComparator Comparator = FAvaTransitionLayerUtils::BuildComparator(InTransitionContext, EAvaTransitionLayerCompareType::Same, FAvaTagHandle());
+		BehaviorInstances = FAvaTransitionLayerUtils::QueryBehaviorInstances(*TransitionSubsystem, Comparator);
+	}
+
+	if (BehaviorInstances.IsEmpty())
+	{
+		return false;
+	}
+
+	// Optional Extension to override Controller Comparison
+	IAvaRCTransitionExtension* const RCTransitionExtension = TransitionScene->FindExtension<IAvaRCTransitionExtension>();
+
+	for (const FAvaTransitionBehaviorInstance* BehaviorInstance : BehaviorInstances)
+	{
+		check(BehaviorInstance);
+
+		const FAvaTransitionScene* OtherTransitionScene = BehaviorInstance->GetTransitionContext().GetTransitionScene();
+		if (!OtherTransitionScene)
+		{
+			continue;
+		}
+
+		EAvaTransitionComparisonResult Result;
+		if (RCTransitionExtension)
+		{
+			Result = RCTransitionExtension->CompareControllers(Controller->Id
+				, *TransitionScene
+				, *OtherTransitionScene);
+		}
+		else if (URCVirtualPropertyBase* OtherController = UE::Ava::Private::GetController(InControllerId, *SceneSubsystem, OtherTransitionScene->GetLevel()))
+		{
+			Result = Controller->IsValueEqual(OtherController)
+				? EAvaTransitionComparisonResult::Same
+				: EAvaTransitionComparisonResult::Different;
+		}
+		else
+		{
+			Result = EAvaTransitionComparisonResult::None;
+		}
+
+		if (InValueComparisonType == Result)
+		{
+			return true;
+		}
+	}
+
+	return false;
+}
+
+bool UAvaTransitionRCLibrary::CompareRCControllerValues(UObject* InTransitionNode, const FAvaRCControllerId& InControllerId, EAvaTransitionComparisonResult InValueComparisonType)
+{
+	IAvaTransitionNodeInterface* NodeInterface = Cast<IAvaTransitionNodeInterface>(InTransitionNode);
+	if (!NodeInterface)
+	{
+		return false;
+	}
+
+	const FAvaTransitionContext* TransitionContext = NodeInterface->GetBehaviorInstanceCache().GetTransitionContext();
+	if (!TransitionContext)
+	{
+		return false;
+	}
+
+	return CompareRCControllerValues(*TransitionContext, InControllerId, InValueComparisonType);
+}
diff --git a/Engine/Plugins/Experimental/Avalanche/Source/Avalanche/Private/Transition/AvaTransitionRCLibrary.h b/Engine/Plugins/Experimental/Avalanche/Source/Avalanche/Private/Transition/AvaTransitionRCLibrary.h
new file mode 100644
--- /dev/null
+++ b/Engine/Plugins/Experimental/Avalanche/Source/Avalanche/Private/Transition/AvaTransitionRCLibrary.h
@@ -0,0 +1,26 @@
+// Copyright Epic Games, Inc. All Rights Reserved.
+
+#pragma once
+
+#include "AvaTransitionEnums.h"
+#include "Kismet/BlueprintFunctionLibrary.h"
+#include "AvaTransitionRCLibrary.generated.h"
+
+class IAvaTransitionNodeInterface;
+class UAvaSceneSubsystem;
+class URCVirtualPropertyBase;
+struct FAvaRCControllerId;
+struct FAvaTransitionContext;
+struct FAvaTransitionScene;
+
+UCLASS()
+class UAvaTransitionRCLibrary : public UBlueprintFunctionLibrary
+{
+	GENERATED_BODY()
+
+public:
+	static bool CompareRCControllerValues(const FAvaTransitionContext& InTransitionContext, const FAvaRCControllerId& InControllerId, EAvaTransitionComparisonResult InValueComparisonType);
+
+	UFUNCTION(BlueprintPure, DisplayName="Compare RC Controller Values", Category="Remote Control", meta=(DefaultToSelf="InTransitionNode"))
+	static bool CompareRCControllerValues(UObject* InTransitionNode, const FAvaRCControllerId& InControllerId, EAvaTransitionComparisonResult InValueComparisonType = EAvaTransitionComparisonResult::Different);
+};
diff --git a/Engine/Plugins/Experimental/Avalanche/Source/Avalanche/Private/Transition/Conditions/AvaTransitionRCControllerMatchCondition.cpp b/Engine/Plugins/Experimental/Avalanche/Source/Avalanche/Private/Transition/Conditions/AvaTransitionRCControllerMatchCondition.cpp
--- a/Engine/Plugins/Experimental/Avalanche/Source/Avalanche/Private/Transition/Conditions/AvaTransitionRCControllerMatchCondition.cpp
+++ b/Engine/Plugins/Experimental/Avalanche/Source/Avalanche/Private/Transition/Conditions/AvaTransitionRCControllerMatchCondition.cpp
@@ -3,19 +3,11 @@
 #include "Transition/Conditions/AvaTransitionRCControllerMatchCondition.h"
 #include "AvaSceneSubsystem.h"
 #include "AvaTransitionContext.h"
-#include "AvaTransitionLayer.h"
-#include "AvaTransitionLayerUtils.h"
 #include "AvaTransitionLog.h"
-#include "AvaTransitionScene.h"
-#include "AvaTransitionSubsystem.h"
 #include "AvaTransitionUtils.h"
-#include "Behavior/AvaTransitionBehaviorInstance.h"
-#include "IAvaSceneInterface.h"
-#include "RCVirtualProperty.h"
-#include "RemoteControlPreset.h"
 #include "StateTreeExecutionContext.h"
 #include "StateTreeLinker.h"
-#include "Transition/Extensions/IAvaTransitionRCExtension.h"
+#include "Transition/AvaTransitionRCLibrary.h"
 
 #define LOCTEXT_NAMESPACE "AvaTransitionRCControllerMatchCondition"
 
@@ -55,89 +47,11 @@ bool FAvaTransitionRCControllerMatchCondition::Link(FStateTreeLinker& InLinker)
 
 bool FAvaTransitionRCControllerMatchCondition::TestCondition(FStateTreeExecutionContext& InContext) const
 {
-	const FInstanceDataType& InstanceData          = InContext.GetInstanceData(*this);
-	const FAvaTransitionContext& TransitionContext = InContext.GetExternalData(TransitionContextHandle);
-	UAvaTransitionSubsystem& TransitionSubsystem   = InContext.GetExternalData(TransitionSubsystemHandle);
-	const UAvaSceneSubsystem& SceneSubsystem       = InContext.GetExternalData(SceneSubsystemHandle);
-	const FAvaTransitionScene* TransitionScene     = TransitionContext.GetTransitionScene();
+	const FInstanceDataType& InstanceData = InContext.GetInstanceData(*this);
 
-	URCVirtualPropertyBase* Controller = GetController(InstanceData.ControllerId, SceneSubsystem, TransitionScene);
-	if (!Controller)
-	{
-		return false;
-	}
-
-	// Get all the Behavior Instances in the same Layer
-	TArray<const FAvaTransitionBehaviorInstance*> BehaviorInstances;
-	{
-		FAvaTransitionLayerComparator Comparator = FAvaTransitionLayerUtils::BuildComparator(TransitionContext, EAvaTransitionLayerCompareType::Same, FAvaTagHandle());
-		BehaviorInstances = FAvaTransitionLayerUtils::QueryBehaviorInstances(TransitionSubsystem, Comparator);
-	}
-
-	if (BehaviorInstances.IsEmpty())
-	{
-		return false;
-	}
-
-	// Optional Extension to override Controller Comparison
-	IAvaRCTransitionExtension* const RCTransitionExtension = TransitionScene->FindExtension<IAvaRCTransitionExtension>();
-
-	for (const FAvaTransitionBehaviorInstance* BehaviorInstance : BehaviorInstances)
-	{
-		check(BehaviorInstance);
-
-		const FAvaTransitionScene* OtherTransitionScene = BehaviorInstance->GetTransitionContext().GetTransitionScene();
-		if (!OtherTransitionScene)
-		{
-			continue;
-		}
-
-		EAvaTransitionComparisonResult Result;
-		if (RCTransitionExtension)
-		{
-			Result = RCTransitionExtension->CompareControllers(Controller->Id
-				, *TransitionScene
-				, *OtherTransitionScene);
-		}
-		else if (URCVirtualPropertyBase* OtherController = GetController(InstanceData.ControllerId, SceneSubsystem, OtherTransitionScene))
-		{
-			Result = Controller->IsValueEqual(OtherController)
-				? EAvaTransitionComparisonResult::Same
-				: EAvaTransitionComparisonResult::Different;
-		}
-		else
-		{
-			Result = EAvaTransitionComparisonResult::None;
-		}
-
-		if (InstanceData.ValueComparisonType == Result)
-		{
-			return true;
-		}
-	}
-
-	return false;
-}
-
-URCVirtualPropertyBase* FAvaTransitionRCControllerMatchCondition::GetController(const FAvaRCControllerId& InControllerId, const UAvaSceneSubsystem& InSceneSubsystem, const FAvaTransitionScene* InTransitionScene) const
-{
-	if (!InTransitionScene)
-	{
-		return nullptr;
-	}
-
-	IAvaSceneInterface* SceneInterface = InSceneSubsystem.GetSceneInterface(InTransitionScene->GetLevel());
-	if (!SceneInterface)
-	{
-		return nullptr;
-	}
-
-	if (URemoteControlPreset* RemoteControlPreset = SceneInterface->GetRemoteControlPreset())
-	{
-		return InControllerId.FindController(RemoteControlPreset);
-	}
-
-	return nullptr;
+	return UAvaTransitionRCLibrary::CompareRCControllerValues(InContext.GetExternalData(TransitionContextHandle)
+		, InstanceData.ControllerId
+		, InstanceData.ValueComparisonType);
 }
 
 #undef LOCTEXT_NAMESPACE
diff --git a/Engine/Plugins/Experimental/Avalanche/Source/Avalanche/Public/Transition/Conditions/AvaTransitionRCControllerMatchCondition.h b/Engine/Plugins/Experimental/Avalanche/Source/Avalanche/Public/Transition/Conditions/AvaTransitionRCControllerMatchCondition.h
--- a/Engine/Plugins/Experimental/Avalanche/Source/Avalanche/Public/Transition/Conditions/AvaTransitionRCControllerMatchCondition.h
+++ b/Engine/Plugins/Experimental/Avalanche/Source/Avalanche/Public/Transition/Conditions/AvaTransitionRCControllerMatchCondition.h
@@ -52,8 +52,6 @@ struct AVALANCHE_API FAvaTransitionRCControllerMatchCondition : public FAvaTrans
 	virtual bool TestCondition(FStateTreeExecutionContext& InContext) const override;
 	//~ End FStateTreeConditionBase
 
-	URCVirtualPropertyBase* GetController(const FAvaRCControllerId& InControllerId, const UAvaSceneSubsystem& InSceneSubsystem, const FAvaTransitionScene* InTransitionScene) const;
-
 	UE_DEPRECATED(5.5, "ControllerId has been moved to Instance Data")
 	UPROPERTY(meta = (DeprecatedProperty, DeprecationMessage = "Use the Instance Data ControllerId instead"))
 	FAvaRCControllerId ControllerId_DEPRECATED;
diff --git a/Engine/Plugins/Experimental/Avalanche/Source/AvalancheRemoteControl/Public/AvaRCControllerId.h b/Engine/Plugins/Experimental/Avalanche/Source/AvalancheRemoteControl/Public/AvaRCControllerId.h
--- a/Engine/Plugins/Experimental/Avalanche/Source/AvalancheRemoteControl/Public/AvaRCControllerId.h
+++ b/Engine/Plugins/Experimental/Avalanche/Source/AvalancheRemoteControl/Public/AvaRCControllerId.h
@@ -9,7 +9,7 @@ class URCVirtualPropertyBase;
 class URemoteControlPreset;
 
 /** Struct describing data to identify/find a Controller in a given preset */
-USTRUCT()
+USTRUCT(BlueprintType)
 struct FAvaRCControllerId
 {
 	GENERATED_BODY()
@@ -18,6 +18,6 @@ struct FAvaRCControllerId
 
 	AVALANCHEREMOTECONTROL_API FText ToText() const;
 
-	UPROPERTY(EditAnywhere, Category = "Motion Design Remote Control")
+	UPROPERTY(EditAnywhere, BlueprintReadWrite, Category = "Motion Design Remote Control")
 	FName Name;
 };
