State Tree: Added extension handler create func

diff --git a/Engine/Plugins/Experimental/Avalanche/Source/AvalancheTransitionEditor/Private/Views/SAvaTransitionSelectionDetails.cpp b/Engine/Plugins/Experimental/Avalanche/Source/AvalancheTransitionEditor/Private/Views/SAvaTransitionSelectionDetails.cpp
--- a/Engine/Plugins/Experimental/Avalanche/Source/AvalancheTransitionEditor/Private/Views/SAvaTransitionSelectionDetails.cpp
+++ b/Engine/Plugins/Experimental/Avalanche/Source/AvalancheTransitionEditor/Private/Views/SAvaTransitionSelectionDetails.cpp
@@ -8,6 +8,7 @@
 #include "Modules/ModuleManager.h"
 #include "PropertyEditorModule.h"
 #include "StateTreeDelegates.h"
+#include "StateTreeEditorModule.h"
 #include "StateTreeState.h"
 #include "ViewModels/AvaTransitionViewModel.h"
 
@@ -38,6 +39,9 @@ void SAvaTransitionSelectionDetails::Construct(const FArguments& InArgs, const T
 		DetailsView->SetIsPropertyEditingEnabledDelegate(FIsPropertyEditingEnabled::CreateLambda([]{ return false; }));
 	}
 
+	FStateTreeEditorModule& StateTreeEditorModule = FModuleManager::LoadModuleChecked<FStateTreeEditorModule>("StateTreeEditorModule");
+	DetailsView->SetExtensionHandler(StateTreeEditorModule.CreateStateTreeDetailPropertyExtensionHandler());
+
 	OnSelectionChanged(InSelection->GetSelectedItems());
 
 	OnParametersChangedHandle = UE::StateTree::Delegates::OnParametersChanged.AddSP(this, &SAvaTransitionSelectionDetails::Refresh);
diff --git a/Engine/Plugins/Runtime/StateTree/Source/StateTreeEditorModule/Private/StateTreeEditorModule.cpp b/Engine/Plugins/Runtime/StateTree/Source/StateTreeEditorModule/Private/StateTreeEditorModule.cpp
--- a/Engine/Plugins/Runtime/StateTree/Source/StateTreeEditorModule/Private/StateTreeEditorModule.cpp
+++ b/Engine/Plugins/Runtime/StateTree/Source/StateTreeEditorModule/Private/StateTreeEditorModule.cpp
@@ -17,6 +17,7 @@
 #include "Customizations/StateTreeStateParametersDetails.h"
 #include "Customizations/StateTreeTransitionDetails.h"
 #include "Customizations/StateTreeEventDescDetails.h"
+#include "Customizations/StateTreeBindingExtension.h"
 #include "Customizations/StateTreeBlueprintPropertyRefDetails.h"
 #include "PropertyEditorModule.h"
 #include "StateTree.h"
@@ -148,6 +149,11 @@ TSharedRef<IStateTreeEditor> FStateTreeEditorModule::CreateStateTreeEditor(const
 	return NewEditor;
 }
 
+TSharedRef<IDetailPropertyExtensionHandler> FStateTreeEditorModule::CreateStateTreeDetailPropertyExtensionHandler()
+{
+	return MakeShared<FStateTreeBindingExtension>();
+}
+
 TSharedPtr<FStateTreeNodeClassCache> FStateTreeEditorModule::GetNodeClassCache()
 {
 	if (!NodeClassCache.IsValid())
diff --git a/Engine/Plugins/Runtime/StateTree/Source/StateTreeEditorModule/Public/StateTreeEditorModule.h b/Engine/Plugins/Runtime/StateTree/Source/StateTreeEditorModule/Public/StateTreeEditorModule.h
--- a/Engine/Plugins/Runtime/StateTree/Source/StateTreeEditorModule/Public/StateTreeEditorModule.h
+++ b/Engine/Plugins/Runtime/StateTree/Source/StateTreeEditorModule/Public/StateTreeEditorModule.h
@@ -26,6 +26,8 @@ public:
 	/** Creates an instance of StateTree editor. Only virtual so that it can be called across the DLL boundary. */
 	virtual TSharedRef<IStateTreeEditor> CreateStateTreeEditor(const EToolkitMode::Type Mode, const TSharedPtr< class IToolkitHost >& InitToolkitHost, UStateTree* StateTree);
 
+	TSharedRef<IDetailPropertyExtensionHandler> CreateStateTreeDetailPropertyExtensionHandler();
+
 	virtual TSharedPtr<FExtensibilityManager> GetMenuExtensibilityManager() override { return MenuExtensibilityManager; }
 	virtual TSharedPtr<FExtensibilityManager> GetToolBarExtensibilityManager() override { return ToolBarExtensibilityManager; }
 
