Motion Design
- Pressing 'c' in the viewport now targets the correct, active viewport.
- Pressing 'c' while not in motion design mode not longer causes an ensure fail.

diff --git a/Engine/Plugins/Experimental/Avalanche/Source/AvalancheEditor/Private/Viewport/AvaLevelViewportExtension.cpp b/Engine/Plugins/Experimental/Avalanche/Source/AvalancheEditor/Private/Viewport/AvaLevelViewportExtension.cpp
--- a/Engine/Plugins/Experimental/Avalanche/Source/AvalancheEditor/Private/Viewport/AvaLevelViewportExtension.cpp
+++ b/Engine/Plugins/Experimental/Avalanche/Source/AvalancheEditor/Private/Viewport/AvaLevelViewportExtension.cpp
@@ -288,25 +288,50 @@ TArray<TSharedPtr<IAvaViewportClient>> FAvaLevelViewportExtension::GetLevelEdito
 
 void FAvaLevelViewportExtension::SetDefaultViewportType()
 {
-	if (TSharedPtr<IAvaEditor> Editor = GetEditor())
-	{
-		Editor->GetCommandList()->ExecuteAction(FLevelViewportCommands::Get().SetDefaultViewportType.ToSharedRef());
-		UE::AvaEditor::Private::FixupInvalidFocusedLevelEditorViewport();
-	} 
+	SetViewportType(FLevelViewportCommands::Get().SetDefaultViewportType.ToSharedRef(), /* Set Active Camera */ false);
 }
 
 void FAvaLevelViewportExtension::SetMotionDesignViewportType()
 {
-	if (TSharedPtr<IAvaEditor> Editor = GetEditor())
+	SetViewportType(FAvaEditorCommands::Get().SetMotionDesignViewportType.ToSharedRef(), /* Set Active Camera */ true);
+}
+
+void FAvaLevelViewportExtension::SetViewportType(const TSharedRef<FUICommandInfo>& InViewportCommand, bool bInSetActiveCamera)
+{
+	FLevelEditorModule* LevelEditorModule = FAvaLevelEditorUtils::GetLevelEditorModule();
+
+	if (!LevelEditorModule)
 	{
-		TSharedPtr<FUICommandList> CommandList = Editor->GetCommandList();
-		if (ensure(CommandList.IsValid()))
-		{
-			CommandList->ExecuteAction(FAvaEditorCommands::Get().SetMotionDesignViewportType.ToSharedRef());
-		}
+		return;
+	}
 
-		UE::AvaEditor::Private::FixupInvalidFocusedLevelEditorViewport();
+	TSharedPtr<ILevelEditor> LevelEditor = LevelEditorModule->GetLevelEditorInstance().Pin();
 
+	if (!LevelEditor.IsValid())
+	{
+		return;
+	}
+
+	TSharedPtr<SLevelViewport> ActiveLevelViewport = LevelEditor->GetActiveViewportInterface();
+
+	if (!ActiveLevelViewport.IsValid())
+	{
+		return;
+	}
+
+	TSharedPtr<FUICommandList> CommandList = ActiveLevelViewport->GetCommandList();
+
+	if (!CommandList.IsValid())
+	{
+		return;
+	}
+
+	CommandList->ExecuteAction(InViewportCommand);
+
+	UE::AvaEditor::Private::FixupInvalidFocusedLevelEditorViewport();
+
+	if (bInSetActiveCamera)
+	{
 		if (AActor* LastCameraCutActor = LastCameraCutActorWeak.Get())
 		{
 			SetActiveCamera(LastCameraCutActor, true);
diff --git a/Engine/Plugins/Experimental/Avalanche/Source/AvalancheEditor/Private/Viewport/AvaLevelViewportExtension.h b/Engine/Plugins/Experimental/Avalanche/Source/AvalancheEditor/Private/Viewport/AvaLevelViewportExtension.h
--- a/Engine/Plugins/Experimental/Avalanche/Source/AvalancheEditor/Private/Viewport/AvaLevelViewportExtension.h
+++ b/Engine/Plugins/Experimental/Avalanche/Source/AvalancheEditor/Private/Viewport/AvaLevelViewportExtension.h
@@ -26,6 +26,7 @@ public:
 
 	void SetDefaultViewportType();
 	void SetMotionDesignViewportType();
+	void SetViewportType(const TSharedRef<FUICommandInfo>& InViewportCommand, bool bInSetActiveCamera);
 
 	virtual ~FAvaLevelViewportExtension() override;
 
