Motion Design: Moved Style set to styling directory

diff --git a/Engine/Plugins/Experimental/Avalanche/Source/AvalancheTransitionEditor/Private/AvaTransitionCommands.cpp b/Engine/Plugins/Experimental/Avalanche/Source/AvalancheTransitionEditor/Private/AvaTransitionCommands.cpp
--- a/Engine/Plugins/Experimental/Avalanche/Source/AvalancheTransitionEditor/Private/AvaTransitionCommands.cpp
+++ b/Engine/Plugins/Experimental/Avalanche/Source/AvalancheTransitionEditor/Private/AvaTransitionCommands.cpp
@@ -1,9 +1,8 @@
 // Copyright Epic Games, Inc. All Rights Reserved.
 
 #include "AvaTransitionCommands.h"
-
-#include "AvaTransitionEditorStyle.h"
 #include "Styling/AppStyle.h"
+#include "Styling/AvaTransitionEditorStyle.h"
 
 #define LOCTEXT_NAMESPACE "AvaTransitionEditorCommands"
 
diff --git a/Engine/Plugins/Experimental/Avalanche/Source/AvalancheTransitionEditor/Private/AvaTransitionEditorModule.cpp b/Engine/Plugins/Experimental/Avalanche/Source/AvalancheTransitionEditor/Private/AvaTransitionEditorModule.cpp
--- a/Engine/Plugins/Experimental/Avalanche/Source/AvalancheTransitionEditor/Private/AvaTransitionEditorModule.cpp
+++ b/Engine/Plugins/Experimental/Avalanche/Source/AvalancheTransitionEditor/Private/AvaTransitionEditorModule.cpp
@@ -4,7 +4,6 @@
 #include "AvaTransitionCommands.h"
 #include "AvaTransitionEditorEnums.h"
 #include "AvaTransitionEditorLog.h"
-#include "AvaTransitionEditorStyle.h"
 #include "AvaTransitionEditorUtils.h"
 #include "AvaTransitionTree.h"
 #include "AvaTransitionTreeEditorData.h"
@@ -12,6 +11,7 @@
 #include "Behavior/IAvaTransitionBehavior.h"
 #include "IAvaTransitionModule.h"
 #include "StateTreeDelegates.h"
+#include "Styling/AvaTransitionEditorStyle.h"
 #include "ToolMenu.h"
 #include "ToolMenuSection.h"
 #include "ViewModels/AvaTransitionEditorViewModel.h"
diff --git a/Engine/Plugins/Experimental/Avalanche/Source/AvalancheTransitionEditor/Private/AvaTransitionEditorStyle.cpp b/Engine/Plugins/Experimental/Avalanche/Source/AvalancheTransitionEditor/Private/AvaTransitionEditorStyle.cpp
deleted file mode 100644
diff --git a/Engine/Plugins/Experimental/Avalanche/Source/AvalancheTransitionEditor/Private/AvaTransitionEditorStyle.h b/Engine/Plugins/Experimental/Avalanche/Source/AvalancheTransitionEditor/Private/AvaTransitionEditorStyle.h
deleted file mode 100644
diff --git a/Engine/Plugins/Experimental/Avalanche/Source/AvalancheTransitionEditor/Private/Menu/Widgets/AvaTransitionTreeStatus.cpp b/Engine/Plugins/Experimental/Avalanche/Source/AvalancheTransitionEditor/Private/Menu/Widgets/AvaTransitionTreeStatus.cpp
--- a/Engine/Plugins/Experimental/Avalanche/Source/AvalancheTransitionEditor/Private/Menu/Widgets/AvaTransitionTreeStatus.cpp
+++ b/Engine/Plugins/Experimental/Avalanche/Source/AvalancheTransitionEditor/Private/Menu/Widgets/AvaTransitionTreeStatus.cpp
@@ -1,9 +1,9 @@
 // Copyright Epic Games, Inc. All Rights Reserved.
 
 #include "AvaTransitionTreeStatus.h"
-#include "AvaTransitionEditorStyle.h"
 #include "AvaTransitionTree.h"
 #include "ScopedTransaction.h"
+#include "Styling/AvaTransitionEditorStyle.h"
 #include "Widgets/Images/SImage.h"
 #include "Widgets/Input/SButton.h"
 
diff --git a/Engine/Plugins/Experimental/Avalanche/Source/AvalancheTransitionEditor/Private/Styling/AvaTransitionEditorStyle.cpp b/Engine/Plugins/Experimental/Avalanche/Source/AvalancheTransitionEditor/Private/Styling/AvaTransitionEditorStyle.cpp
new file mode 100644
--- /dev/null
+++ b/Engine/Plugins/Experimental/Avalanche/Source/AvalancheTransitionEditor/Private/Styling/AvaTransitionEditorStyle.cpp
@@ -0,0 +1,93 @@
+// Copyright Epic Games, Inc. All Rights Reserved.
+
+#include "AvaTransitionEditorStyle.h"
+#include "Brushes/SlateImageBrush.h"
+#include "AvaTransitionTree.h"
+#include "Interfaces/IPluginManager.h"
+#include "Layout/Margin.h"
+#include "Misc/Paths.h"
+#include "Styling/SlateStyleMacros.h"
+#include "Styling/SlateStyleRegistry.h"
+#include "Textures/SlateIcon.h"
+
+FAvaTransitionEditorStyle::FAvaTransitionEditorStyle()
+	: FSlateStyleSet(TEXT("AvaTransitionEditor"))
+{
+	const FVector2f Icon16(16.f);
+	const FVector2f Icon20(20.f);
+	const FVector2f Icon64(64.f);
+
+	const TSharedPtr<IPlugin> Plugin = IPluginManager::Get().FindPlugin(UE_PLUGIN_NAME);
+	check(Plugin.IsValid());
+
+	SetContentRoot(FPaths::EngineContentDir() / TEXT("Editor/Slate"));
+	SetCoreContentRoot(FPaths::EngineContentDir() / TEXT("Slate"));
+
+	// Re-use the Behavior Tree Icon
+	Set("ClassIcon.AvaTransitionTree"     , new IMAGE_BRUSH("Icons/AssetIcons/BehaviorTree_16x", Icon16));
+	Set("ClassThumbnail.AvaTransitionTree", new IMAGE_BRUSH("Icons/AssetIcons/BehaviorTree_64x", Icon64));
+
+	Set("Throbber.CircleChunk", new CORE_IMAGE_BRUSH("Common/Throbber_Piece", FVector2f(1.5f)));
+
+	// Editor Commands
+	Set("AvaTransitionEditor.AddSiblingState"       , new CORE_IMAGE_BRUSH_SVG("Starship/Common/plus"     , Icon20));
+	Set("AvaTransitionEditor.AddChildState"         , new CORE_IMAGE_BRUSH_SVG("Starship/Common/plus"     , Icon20));
+	Set("AvaTransitionEditor.ImportTransitionTree"  , new CORE_IMAGE_BRUSH_SVG("Starship/Common/import_20", Icon20));
+	Set("AvaTransitionEditor.ReimportTransitionTree", new CORE_IMAGE_BRUSH_SVG("Starship/Common/import_20", Icon20));
+	Set("AvaTransitionEditor.ExportTransitionTree"  , new CORE_IMAGE_BRUSH_SVG("Starship/Common/export_20", Icon20));
+	Set("AvaTransitionEditor.ToggleDebug"           , new IMAGE_BRUSH_SVG("Starship/Common/Debug"         , Icon20));
+
+	Set("DebugIndicatorBorder", new BOX_BRUSH("Images/NamespaceBorder", FMargin(0.25f)));
+
+	FSlateStyleRegistry::RegisterSlateStyle(*this);
+}
+
+FAvaTransitionEditorStyle::~FAvaTransitionEditorStyle()
+{
+	FSlateStyleRegistry::UnRegisterSlateStyle(*this);
+}
+
+FLinearColor FAvaTransitionEditorStyle::LerpColorSRGB(const FLinearColor& InA, const FLinearColor& InB, float InAlpha)
+{
+	const FColor A = InA.ToFColorSRGB();
+	const FColor B = InB.ToFColorSRGB();
+
+	return FLinearColor(FColor(
+		static_cast<uint8>(FMath::RoundToInt(static_cast<float>(A.R) * (1.f - InAlpha) + static_cast<float>(B.R) * InAlpha)),
+		static_cast<uint8>(FMath::RoundToInt(static_cast<float>(A.G) * (1.f - InAlpha) + static_cast<float>(B.G) * InAlpha)),
+		static_cast<uint8>(FMath::RoundToInt(static_cast<float>(A.B) * (1.f - InAlpha) + static_cast<float>(B.B) * InAlpha)),
+		static_cast<uint8>(FMath::RoundToInt(static_cast<float>(A.A) * (1.f - InAlpha) + static_cast<float>(B.A) * InAlpha))));
+}
+
+FSlateIcon FAvaTransitionEditorStyle::ParseIcon(FName InIconName)
+{
+	FString IconPath = InIconName.ToString();
+
+	constexpr int32 NumOfIconPathNames = 4;
+
+	FName IconPathNames[NumOfIconPathNames] = {
+		NAME_None, // StyleSetName
+		NAME_None, // StyleName
+		NAME_None, // SmallStyleName
+		NAME_None  // StatusOverlayStyleName
+	};
+
+	int32 NameIndex = 0;
+	while (!IconPath.IsEmpty() && NameIndex < NumOfIconPathNames)
+	{
+		FString Left;
+		FString Right;
+
+		if (!IconPath.Split(TEXT("|"), &Left, &Right))
+		{
+			Left = IconPath;
+		}
+
+		IconPathNames[NameIndex] = FName(*Left);
+
+		NameIndex++;
+		IconPath = Right;
+	}
+
+	return FSlateIcon(IconPathNames[0], IconPathNames[1], IconPathNames[2], IconPathNames[3]);
+}
diff --git a/Engine/Plugins/Experimental/Avalanche/Source/AvalancheTransitionEditor/Private/Styling/AvaTransitionEditorStyle.h b/Engine/Plugins/Experimental/Avalanche/Source/AvalancheTransitionEditor/Private/Styling/AvaTransitionEditorStyle.h
new file mode 100644
--- /dev/null
+++ b/Engine/Plugins/Experimental/Avalanche/Source/AvalancheTransitionEditor/Private/Styling/AvaTransitionEditorStyle.h
@@ -0,0 +1,24 @@
+// Copyright Epic Games, Inc. All Rights Reserved.
+
+#pragma once
+
+#include "Styling/SlateStyle.h"
+
+struct FSlateIcon;
+
+class FAvaTransitionEditorStyle final : public FSlateStyleSet
+{
+public:
+	static FAvaTransitionEditorStyle& Get()
+	{
+		static FAvaTransitionEditorStyle Instance;
+		return Instance;
+	}
+
+	FAvaTransitionEditorStyle();
+	virtual ~FAvaTransitionEditorStyle() override;
+
+	static FLinearColor LerpColorSRGB(const FLinearColor& InA, const FLinearColor& InB, float InAlpha);
+
+	static FSlateIcon ParseIcon(FName InIconName);
+};
diff --git a/Engine/Plugins/Experimental/Avalanche/Source/AvalancheTransitionEditor/Private/ViewModels/AvaTransitionNodeViewModel.cpp b/Engine/Plugins/Experimental/Avalanche/Source/AvalancheTransitionEditor/Private/ViewModels/AvaTransitionNodeViewModel.cpp
--- a/Engine/Plugins/Experimental/Avalanche/Source/AvalancheTransitionEditor/Private/ViewModels/AvaTransitionNodeViewModel.cpp
+++ b/Engine/Plugins/Experimental/Avalanche/Source/AvalancheTransitionEditor/Private/ViewModels/AvaTransitionNodeViewModel.cpp
@@ -1,7 +1,6 @@
 // Copyright Epic Games, Inc. All Rights Reserved.
 
 #include "AvaTransitionNodeViewModel.h"
-#include "AvaTransitionEditorStyle.h"
 #include "AvaTransitionEditorViewModel.h"
 #include "AvaTransitionTreeEditorData.h"
 #include "AvaTransitionViewModelSharedData.h"
@@ -9,6 +8,7 @@
 #include "State/AvaTransitionStateViewModel.h"
 #include "StateTreeEditorNode.h"
 #include "StateTreeEditorTypes.h"
+#include "Styling/AvaTransitionEditorStyle.h"
 
 FAvaTransitionNodeViewModel::FAvaTransitionNodeViewModel(const FStateTreeEditorNode& InEditorNode)
 	: NodeId(InEditorNode.ID)
diff --git a/Engine/Plugins/Experimental/Avalanche/Source/AvalancheTransitionEditor/Private/ViewModels/State/AvaTransitionStateViewModel.cpp b/Engine/Plugins/Experimental/Avalanche/Source/AvalancheTransitionEditor/Private/ViewModels/State/AvaTransitionStateViewModel.cpp
--- a/Engine/Plugins/Experimental/Avalanche/Source/AvalancheTransitionEditor/Private/ViewModels/State/AvaTransitionStateViewModel.cpp
+++ b/Engine/Plugins/Experimental/Avalanche/Source/AvalancheTransitionEditor/Private/ViewModels/State/AvaTransitionStateViewModel.cpp
@@ -2,11 +2,11 @@
 
 #include "AvaTransitionStateViewModel.h"
 #include "Actions/AvaTransitionStateActions.h"
-#include "AvaTransitionEditorStyle.h"
 #include "AvaTransitionSelection.h"
 #include "AvaTransitionTreeEditorData.h"
 #include "DragDrop/AvaTransitionStateDragDropOp.h"
 #include "StateTreeState.h"
+#include "Styling/AvaTransitionEditorStyle.h"
 #include "ViewModels/AvaTransitionEditorViewModel.h"
 #include "ViewModels/AvaTransitionViewModelSharedData.h"
 #include "ViewModels/AvaTransitionViewModelUtils.h"
diff --git a/Engine/Plugins/Experimental/Avalanche/Source/AvalancheTransitionEditor/Private/ViewModels/Task/AvaTransitionTaskViewModel.cpp b/Engine/Plugins/Experimental/Avalanche/Source/AvalancheTransitionEditor/Private/ViewModels/Task/AvaTransitionTaskViewModel.cpp
--- a/Engine/Plugins/Experimental/Avalanche/Source/AvalancheTransitionEditor/Private/ViewModels/Task/AvaTransitionTaskViewModel.cpp
+++ b/Engine/Plugins/Experimental/Avalanche/Source/AvalancheTransitionEditor/Private/ViewModels/Task/AvaTransitionTaskViewModel.cpp
@@ -1,11 +1,11 @@
 // Copyright Epic Games, Inc. All Rights Reserved.
 
 #include "AvaTransitionTaskViewModel.h"
-#include "AvaTransitionEditorStyle.h"
 #include "AvaTransitionTreeEditorData.h"
 #include "StateTreeEditorData.h"
 #include "StateTreeEditorStyle.h"
 #include "StateTreeState.h"
+#include "Styling/AvaTransitionEditorStyle.h"
 #include "Tasks/AvaTransitionTask.h"
 #include "Textures/SlateIcon.h"
 #include "ViewModels/AvaTransitionViewModelSharedData.h"
diff --git a/Engine/Plugins/Experimental/Avalanche/Source/AvalancheTransitionEditor/Private/Views/Debugger/SAvaTransitionStateDebugInstance.cpp b/Engine/Plugins/Experimental/Avalanche/Source/AvalancheTransitionEditor/Private/Views/Debugger/SAvaTransitionStateDebugInstance.cpp
--- a/Engine/Plugins/Experimental/Avalanche/Source/AvalancheTransitionEditor/Private/Views/Debugger/SAvaTransitionStateDebugInstance.cpp
+++ b/Engine/Plugins/Experimental/Avalanche/Source/AvalancheTransitionEditor/Private/Views/Debugger/SAvaTransitionStateDebugInstance.cpp
@@ -3,8 +3,8 @@
 #if WITH_STATETREE_DEBUGGER
 
 #include "SAvaTransitionStateDebugInstance.h"
-#include "AvaTransitionEditorStyle.h"
 #include "Debugger/AvaTransitionStateDebugInstance.h"
+#include "Styling/AvaTransitionEditorStyle.h"
 #include "ViewModels/State/AvaTransitionStateViewModel.h"
 #include "Widgets/Images/SImage.h"
 #include "Widgets/Images/SThrobber.h"
diff --git a/Engine/Plugins/Experimental/Avalanche/Source/AvalancheTransitionEditor/Private/Views/Debugger/SAvaTransitionStateDebugInstanceContainer.cpp b/Engine/Plugins/Experimental/Avalanche/Source/AvalancheTransitionEditor/Private/Views/Debugger/SAvaTransitionStateDebugInstanceContainer.cpp
--- a/Engine/Plugins/Experimental/Avalanche/Source/AvalancheTransitionEditor/Private/Views/Debugger/SAvaTransitionStateDebugInstanceContainer.cpp
+++ b/Engine/Plugins/Experimental/Avalanche/Source/AvalancheTransitionEditor/Private/Views/Debugger/SAvaTransitionStateDebugInstanceContainer.cpp
@@ -3,9 +3,9 @@
 #if WITH_STATETREE_DEBUGGER
 
 #include "SAvaTransitionStateDebugInstanceContainer.h"
-#include "AvaTransitionEditorStyle.h"
 #include "Debugger/AvaTransitionDebugger.h"
 #include "SAvaTransitionStateDebugInstance.h"
+#include "Styling/AvaTransitionEditorStyle.h"
 #include "ViewModels/AvaTransitionViewModelSharedData.h"
 #include "ViewModels/State/AvaTransitionStateViewModel.h"
 #include "Widgets/SBoxPanel.h"
diff --git a/Engine/Plugins/Experimental/Avalanche/Source/AvalancheTransitionEditor/Private/Views/SAvaTransitionStateView.cpp b/Engine/Plugins/Experimental/Avalanche/Source/AvalancheTransitionEditor/Private/Views/SAvaTransitionStateView.cpp
--- a/Engine/Plugins/Experimental/Avalanche/Source/AvalancheTransitionEditor/Private/Views/SAvaTransitionStateView.cpp
+++ b/Engine/Plugins/Experimental/Avalanche/Source/AvalancheTransitionEditor/Private/Views/SAvaTransitionStateView.cpp
@@ -1,10 +1,10 @@
 // Copyright Epic Games, Inc. All Rights Reserved.
 
 #include "SAvaTransitionStateView.h"
-#include "AvaTransitionEditorStyle.h"
 #include "SAvaTransitionStateMetadata.h"
 #include "StateTreeEditorStyle.h"
 #include "StateTreeTypes.h"
+#include "Styling/AvaTransitionEditorStyle.h"
 #include "ViewModels/Condition/AvaTransitionConditionContainerViewModel.h"
 #include "ViewModels/State/AvaTransitionStateViewModel.h"
 #include "ViewModels/Task/AvaTransitionTaskContainerViewModel.h"
