Motion Design: Added support for State Tree Prop Functions

diff --git a/Engine/Plugins/Experimental/Avalanche/Source/AvalancheTransition/Private/AvaTransitionTreeSchema.cpp b/Engine/Plugins/Experimental/Avalanche/Source/AvalancheTransition/Private/AvaTransitionTreeSchema.cpp
--- a/Engine/Plugins/Experimental/Avalanche/Source/AvalancheTransition/Private/AvaTransitionTreeSchema.cpp
+++ b/Engine/Plugins/Experimental/Avalanche/Source/AvalancheTransition/Private/AvaTransitionTreeSchema.cpp
@@ -6,6 +6,7 @@
 #include "Conditions/AvaTransitionCondition.h"
 #include "Conditions/AvaTransitionConditionBlueprint.h"
 #include "StateTreeEvaluatorBase.h"
+#include "StateTreePropertyFunctionBase.h"
 #include "StateTreeTaskBase.h"
 #include "Subsystems/WorldSubsystem.h"
 #include "Tasks/AvaTransitionTask.h"
@@ -14,7 +15,8 @@
 bool UAvaTransitionTreeSchema::IsStructAllowed(const UScriptStruct* InScriptStruct) const
 {
 	return InScriptStruct->IsChildOf(FAvaTransitionTask::StaticStruct())
-		|| InScriptStruct->IsChildOf(FAvaTransitionCondition::StaticStruct());
+		|| InScriptStruct->IsChildOf(FAvaTransitionCondition::StaticStruct())
+		|| InScriptStruct->IsChildOf(FStateTreePropertyFunctionCommonBase::StaticStruct());
 }
 
 bool UAvaTransitionTreeSchema::IsClassAllowed(const UClass* InClass) const
