Motion Design: Added Rich Text support for the Tasks and State Views.
Changed the Tasks/Conditions to consider Rich Text.
Operand AND/OR now are represented by their widget analogous to the one drawn in the Details Panel.

diff --git a/Engine/Plugins/Experimental/Avalanche/Source/Avalanche/Private/Transition/Conditions/AvaSceneContainsTagAttributeCondition.cpp b/Engine/Plugins/Experimental/Avalanche/Source/Avalanche/Private/Transition/Conditions/AvaSceneContainsTagAttributeCondition.cpp
--- a/Engine/Plugins/Experimental/Avalanche/Source/Avalanche/Private/Transition/Conditions/AvaSceneContainsTagAttributeCondition.cpp
+++ b/Engine/Plugins/Experimental/Avalanche/Source/Avalanche/Private/Transition/Conditions/AvaSceneContainsTagAttributeCondition.cpp
@@ -24,25 +24,56 @@ FText FAvaSceneContainsTagAttributeConditionBase::GetDescription(const FGuid& In
 	switch (InstanceData.SceneType)
 	{
 	case EAvaTransitionSceneType::This:
-		Arguments.Add(TEXT("IndefinitePronoun"), FText::GetEmpty());
-
-		Arguments.Add(TEXT("Scene"), LOCTEXT("ThisScene", "this scene"));
-
-		Arguments.Add(TEXT("Contains"), bInvertCondition ? LOCTEXT("ThisDoesntContain", "does not contain") : LOCTEXT("ThisContains", "contains"));
+		{
+			if (InFormatting == EStateTreeNodeFormatting::RichText)
+			{
+				Arguments.Add(TEXT("IndefinitePronoun"), FText::GetEmpty());
+				Arguments.Add(TEXT("Scene"), LOCTEXT("ThisSceneRich", "<b>this</> <s>scene</>"));
+				Arguments.Add(TEXT("Contains"), bInvertCondition ? LOCTEXT("ThisDoesntContainRich", "<s>does</> <b>not</> <s>contain</>") : LOCTEXT("ThisContainsRich", "<s>contains</>"));
+			}
+			else
+			{
+				Arguments.Add(TEXT("IndefinitePronoun"), FText::GetEmpty());
+				Arguments.Add(TEXT("Scene"), LOCTEXT("ThisScene", "this scene"));
+				Arguments.Add(TEXT("Contains"), bInvertCondition ? LOCTEXT("ThisDoesntContain", "does not contain") : LOCTEXT("ThisContains", "contains"));
+			}
+		}
 		break;
 
 	case EAvaTransitionSceneType::Other:
-		Arguments.Add(TEXT("IndefinitePronoun"), bInvertCondition ? LOCTEXT("NoScene", "no ") : LOCTEXT("AnyScene", "a "));
-
-		Arguments.Add(TEXT("Scene"), FText::Format(LOCTEXT("OtherScene", "scene in {0}"), FAvaTransitionLayerUtils::GetLayerQueryText(InstanceData.LayerType, *InstanceData.SpecificLayers.ToString())));
-
-		Arguments.Add(TEXT("Contains"), LOCTEXT("OtherSceneContains", "contains"));
+		{
+			FText LayerDesc;
+			{
+				FAvaTransitionLayerUtils::FLayerQueryTextParams Params;
+				Params.LayerType = InstanceData.LayerType;
+				Params.SpecificLayerName = *InstanceData.SpecificLayers.ToString();
+				Params.LayerTypePropertyName = GET_MEMBER_NAME_CHECKED(FInstanceDataType, LayerType);
+				Params.SpecificLayerPropertyName = GET_MEMBER_NAME_CHECKED(FInstanceDataType, SpecificLayers);
+
+				LayerDesc = FAvaTransitionLayerUtils::GetLayerQueryText(MoveTemp(Params), InId, InBindingLookup, InFormatting);
+			}
+
+			if (InFormatting == EStateTreeNodeFormatting::RichText)
+			{
+				Arguments.Add(TEXT("IndefinitePronoun"), bInvertCondition ? LOCTEXT("NoSceneRich", "<b>no</> ") : LOCTEXT("AnySceneRich", "<s>a</> "));
+				Arguments.Add(TEXT("Scene"), FText::Format(LOCTEXT("OtherSceneRich", "<s>scene in</> {0}"), LayerDesc));
+				Arguments.Add(TEXT("Contains"), LOCTEXT("OtherSceneContainsRich", "<s>contains</>"));
+			}
+			else
+			{
+				Arguments.Add(TEXT("IndefinitePronoun"), bInvertCondition ? LOCTEXT("NoScene", "no ") : LOCTEXT("AnyScene", "a "));
+				Arguments.Add(TEXT("Scene"), FText::Format(LOCTEXT("OtherScene", "scene in {0}"), LayerDesc));
+				Arguments.Add(TEXT("Contains"), LOCTEXT("OtherSceneContains", "contains"));
+			}
+		}
 		break;
 	}
 
 	Arguments.Add(TEXT("TagAttribute"), FText::FromName(InstanceData.TagAttribute.ToName()));
 
-	return FText::Format(LOCTEXT("ConditionDescription", "{IndefinitePronoun}{Scene} {Contains} tag attribute '{TagAttribute}'"), Arguments);
+	return InFormatting == EStateTreeNodeFormatting::RichText
+		? FText::Format(LOCTEXT("DescRich", "{IndefinitePronoun}{Scene} {Contains} <s>tag attribute</> <b>'{TagAttribute}'</>"), Arguments)
+		: FText::Format(LOCTEXT("Desc", "{IndefinitePronoun}{Scene} {Contains} tag attribute '{TagAttribute}'"), Arguments);
 }
 #endif
 
diff --git a/Engine/Plugins/Experimental/Avalanche/Source/Avalanche/Private/Transition/Conditions/AvaTransitionRCControllerMatchCondition.cpp b/Engine/Plugins/Experimental/Avalanche/Source/Avalanche/Private/Transition/Conditions/AvaTransitionRCControllerMatchCondition.cpp
--- a/Engine/Plugins/Experimental/Avalanche/Source/Avalanche/Private/Transition/Conditions/AvaTransitionRCControllerMatchCondition.cpp
+++ b/Engine/Plugins/Experimental/Avalanche/Source/Avalanche/Private/Transition/Conditions/AvaTransitionRCControllerMatchCondition.cpp
@@ -16,9 +16,12 @@ FText FAvaTransitionRCControllerMatchCondition::GetDescription(const FGuid& InId
 {
 	const FInstanceDataType& InstanceData = InInstanceDataView.Get<FInstanceDataType>();
 
-	return FText::Format(LOCTEXT("ConditionDescription", "'{0}' is {1}")
-		, InstanceData.ControllerId.ToText()
-		, UEnum::GetDisplayValueAsText(InstanceData.ValueComparisonType).ToLower());
+	const FText ControllerIdText = InstanceData.ControllerId.ToText();
+	const FText ComparisonText = UEnum::GetDisplayValueAsText(InstanceData.ValueComparisonType).ToLower();
+
+	return InFormatting == EStateTreeNodeFormatting::RichText
+		? FText::Format(LOCTEXT("DescRich", "'<b>{0}</>' <s>is</> <b>{1}</>"), ControllerIdText, ComparisonText)
+		: FText::Format(LOCTEXT("Desc", "'{0}' is {1}"), ControllerIdText, ComparisonText);
 }
 #endif
 
diff --git a/Engine/Plugins/Experimental/Avalanche/Source/Avalanche/Private/Transition/Tasks/AvaSceneAddTagAttributeTask.cpp b/Engine/Plugins/Experimental/Avalanche/Source/Avalanche/Private/Transition/Tasks/AvaSceneAddTagAttributeTask.cpp
--- a/Engine/Plugins/Experimental/Avalanche/Source/Avalanche/Private/Transition/Tasks/AvaSceneAddTagAttributeTask.cpp
+++ b/Engine/Plugins/Experimental/Avalanche/Source/Avalanche/Private/Transition/Tasks/AvaSceneAddTagAttributeTask.cpp
@@ -11,7 +11,12 @@
 FText FAvaSceneAddTagAttributeTask::GetDescription(const FGuid& InId, FStateTreeDataView InInstanceDataView, const IStateTreeBindingLookup& InBindingLookup, EStateTreeNodeFormatting InFormatting) const
 {
 	const FInstanceDataType& InstanceData = InInstanceDataView.Get<FInstanceDataType>();
-	return FText::Format(LOCTEXT("TaskDescription", "Add '{0}' tag attribute to this scene"), FText::FromName(InstanceData.TagAttribute.ToName()));
+
+	const FText TagAttributeDesc = FText::FromName(InstanceData.TagAttribute.ToName());
+
+	return InFormatting == EStateTreeNodeFormatting::RichText
+		? FText::Format(LOCTEXT("DescRich", " <b>Add '{0}'</> <s>tag attribute to this scene</>"), TagAttributeDesc)
+		: FText::Format(LOCTEXT("Desc", "Add '{0}' tag attribute to this scene"), TagAttributeDesc);
 }
 #endif
 
diff --git a/Engine/Plugins/Experimental/Avalanche/Source/Avalanche/Private/Transition/Tasks/AvaSceneRemoveTagAttributeTask.cpp b/Engine/Plugins/Experimental/Avalanche/Source/Avalanche/Private/Transition/Tasks/AvaSceneRemoveTagAttributeTask.cpp
--- a/Engine/Plugins/Experimental/Avalanche/Source/Avalanche/Private/Transition/Tasks/AvaSceneRemoveTagAttributeTask.cpp
+++ b/Engine/Plugins/Experimental/Avalanche/Source/Avalanche/Private/Transition/Tasks/AvaSceneRemoveTagAttributeTask.cpp
@@ -11,7 +11,12 @@
 FText FAvaSceneRemoveTagAttributeTask::GetDescription(const FGuid& InId, FStateTreeDataView InInstanceDataView, const IStateTreeBindingLookup& InBindingLookup, EStateTreeNodeFormatting InFormatting) const
 {
 	const FInstanceDataType& InstanceData = InInstanceDataView.Get<FInstanceDataType>();
-	return FText::Format(LOCTEXT("TaskDescription", "Remove '{0}' tag attribute from this scene"), FText::FromName(InstanceData.TagAttribute.ToName()));
+
+	const FText TagAttributeDesc = FText::FromName(InstanceData.TagAttribute.ToName());
+
+	return InFormatting == EStateTreeNodeFormatting::RichText
+		? FText::Format(LOCTEXT("DescRich", " <b>Remove'{0}'</> <s>tag attribute from this scene</>"), TagAttributeDesc)
+		: FText::Format(LOCTEXT("Desc", "Remove '{0}' tag attribute from this scene"), TagAttributeDesc);
 }
 #endif
 
diff --git a/Engine/Plugins/Experimental/Avalanche/Source/AvalancheSequence/Private/Transition/AvaTransitionPlaySequenceTask.cpp b/Engine/Plugins/Experimental/Avalanche/Source/AvalancheSequence/Private/Transition/AvaTransitionPlaySequenceTask.cpp
--- a/Engine/Plugins/Experimental/Avalanche/Source/AvalancheSequence/Private/Transition/AvaTransitionPlaySequenceTask.cpp
+++ b/Engine/Plugins/Experimental/Avalanche/Source/AvalancheSequence/Private/Transition/AvaTransitionPlaySequenceTask.cpp
@@ -33,11 +33,16 @@ FText FAvaTransitionPlaySequenceTask::GetDescription(const FGuid& InId, FStateTr
 	const FInstanceDataType& InstanceData = InInstanceDataView.Get<FInstanceDataType>();
 
 	TArray<FText> AdditionalArgs;
+	AdditionalArgs.Reserve(2);
 	AdditionalArgs.Add(UEnum::GetDisplayValueAsText(InstanceData.PlaySettings.PlayMode));
 	AdditionalArgs.Add(UEnum::GetDisplayValueAsText(InstanceData.WaitType));
 
-	FText AddOnText = FText::Format(INVTEXT("( {0} )"), FText::Join(FText::FromString(TEXT(" | ")), AdditionalArgs));
-	return FText::Format(LOCTEXT("TaskDescription", "Play {0} {1}"), GetSequenceQueryText(InstanceData), AddOnText);
+	const FText SequenceQueryText = GetSequenceQueryText(InstanceData, InFormatting);
+	const FText AddOnText = FText::Format(INVTEXT("( {0} )"), FText::Join(FText::FromString(TEXT(" | ")), AdditionalArgs));
+
+	return InFormatting == EStateTreeNodeFormatting::RichText
+		? FText::Format(LOCTEXT("DescRich", "Play {0} <s>{1}</>"), SequenceQueryText, AddOnText)
+		: FText::Format(LOCTEXT("DescRich", "Play {0} <s>{1}</>"), SequenceQueryText, AddOnText);
 }
 #endif
 
diff --git a/Engine/Plugins/Experimental/Avalanche/Source/AvalancheSequence/Private/Transition/AvaTransitionSequenceTaskBase.cpp b/Engine/Plugins/Experimental/Avalanche/Source/AvalancheSequence/Private/Transition/AvaTransitionSequenceTaskBase.cpp
--- a/Engine/Plugins/Experimental/Avalanche/Source/AvalancheSequence/Private/Transition/AvaTransitionSequenceTaskBase.cpp
+++ b/Engine/Plugins/Experimental/Avalanche/Source/AvalancheSequence/Private/Transition/AvaTransitionSequenceTaskBase.cpp
@@ -101,15 +101,19 @@ void FAvaTransitionSequenceTaskBase::StopActiveSequences(FStateTreeExecutionCont
 	}
 }
 
-FText FAvaTransitionSequenceTaskBase::GetSequenceQueryText(const FInstanceDataType& InInstanceData) const
+FText FAvaTransitionSequenceTaskBase::GetSequenceQueryText(const FInstanceDataType& InInstanceData, EStateTreeNodeFormatting InFormatting) const
 {
 	switch (InInstanceData.QueryType)
 	{
 	case EAvaTransitionSequenceQueryType::Name:
-		return FText::Format(INVTEXT("'{0}'"), FText::FromName(InInstanceData.SequenceName));
+		return InFormatting == EStateTreeNodeFormatting::RichText
+			? FText::Format(INVTEXT("'<b>{0}</>'"), FText::FromName(InInstanceData.SequenceName))
+			: FText::Format(INVTEXT("'{0}'"), FText::FromName(InInstanceData.SequenceName));
 
 	case EAvaTransitionSequenceQueryType::Tag:
-		return FText::Format(LOCTEXT("SequenceQueryTag", "tag '{0}'"), FText::FromName(InInstanceData.SequenceTag.ToName()));
+		return InFormatting == EStateTreeNodeFormatting::RichText
+			? FText::Format(LOCTEXT("SequenceQueryTagRich", "<s>tag</> '<b>{0}</>'"), FText::FromName(InInstanceData.SequenceTag.ToName()))
+			: FText::Format(LOCTEXT("SequenceQueryTag", "tag '{0}'"), FText::FromName(InInstanceData.SequenceTag.ToName()));
 	}
 
 	checkNoEntry();
diff --git a/Engine/Plugins/Experimental/Avalanche/Source/AvalancheSequence/Public/Transition/AvaTransitionSequenceTaskBase.h b/Engine/Plugins/Experimental/Avalanche/Source/AvalancheSequence/Public/Transition/AvaTransitionSequenceTaskBase.h
--- a/Engine/Plugins/Experimental/Avalanche/Source/AvalancheSequence/Public/Transition/AvaTransitionSequenceTaskBase.h
+++ b/Engine/Plugins/Experimental/Avalanche/Source/AvalancheSequence/Public/Transition/AvaTransitionSequenceTaskBase.h
@@ -92,7 +92,7 @@ struct AVALANCHESEQUENCE_API FAvaTransitionSequenceTaskBase : public FAvaTransit
 	 */
 	void StopActiveSequences(FStateTreeExecutionContext& InContext) const;
 
-	FText GetSequenceQueryText(const FInstanceDataType& InInstanceData) const;
+	FText GetSequenceQueryText(const FInstanceDataType& InInstanceData, EStateTreeNodeFormatting InFormatting) const;
 
 	//~ Begin FStateTreeTaskBase
 	virtual EStateTreeRunStatus EnterState(FStateTreeExecutionContext& InContext, const FStateTreeTransitionResult& InTransition) const override;
diff --git a/Engine/Plugins/Experimental/Avalanche/Source/AvalancheTransition/Private/AvaTransitionLayerUtils.cpp b/Engine/Plugins/Experimental/Avalanche/Source/AvalancheTransition/Private/AvaTransitionLayerUtils.cpp
--- a/Engine/Plugins/Experimental/Avalanche/Source/AvalancheTransition/Private/AvaTransitionLayerUtils.cpp
+++ b/Engine/Plugins/Experimental/Avalanche/Source/AvalancheTransition/Private/AvaTransitionLayerUtils.cpp
@@ -6,6 +6,7 @@
 #include "AvaTransitionSubsystem.h"
 #include "Containers/Array.h"
 #include "Execution/IAvaTransitionExecutor.h"
+#include "StateTreePropertyBindings.h"
 
 #define LOCTEXT_NAMESPACE "AvaTransitionLayerUtils"
 
@@ -47,13 +48,35 @@ FAvaTransitionLayerComparator FAvaTransitionLayerUtils::BuildComparator(const FA
 	return LayerComparator;
 }
 
-FText FAvaTransitionLayerUtils::GetLayerQueryText(EAvaTransitionLayerCompareType InLayerType, FName InSpecificLayerName)
+#if WITH_EDITOR
+FText FAvaTransitionLayerUtils::GetLayerQueryText(FLayerQueryTextParams&& InParams, const FGuid& InId, const IStateTreeBindingLookup& InBindingLookup, EStateTreeNodeFormatting InFormatting)
 {
-	if (InLayerType == EAvaTransitionLayerCompareType::MatchingTag)
+	FText LayerType = InBindingLookup.GetBindingSourceDisplayName(FStateTreePropertyPath(InId, InParams.LayerTypePropertyName), InFormatting);
+
+	// If Layer Type is bound assume it might be set to Specific Layer, so always set it
+	if (!LayerType.IsEmpty() || InParams.LayerType == EAvaTransitionLayerCompareType::MatchingTag)
+	{
+		FText SpecificLayer = InBindingLookup.GetBindingSourceDisplayName(FStateTreePropertyPath(InId, InParams.SpecificLayerPropertyName), InFormatting);
+
+		if (SpecificLayer.IsEmpty())
+		{
+			SpecificLayer = FText::FromName(InParams.SpecificLayerName);
+		}
+
+		return InFormatting == EStateTreeNodeFormatting::RichText
+			? FText::Format(LOCTEXT("SpecificLayerQueryTextRich", "<s>layer</> '<b>{0}</>'"), SpecificLayer)
+			: FText::Format(LOCTEXT("SpecificLayerQueryText", "layer '{0}'"), SpecificLayer);
+	}
+
+	if (LayerType.IsEmpty())
 	{
-		return FText::Format(LOCTEXT("LayerTagText", "'{0}' layer"), FText::FromName(InSpecificLayerName));
+		LayerType = UEnum::GetDisplayValueAsText(InParams.LayerType).ToLower();
 	}
-	return FText::Format(LOCTEXT("LayerTypeText", "{0} layer"), UEnum::GetDisplayValueAsText(InLayerType).ToLower());
+
+	return InFormatting == EStateTreeNodeFormatting::RichText
+		? FText::Format(LOCTEXT("LayerQueryTextRich", "<b>{0}</> <s>layer</>"), LayerType)
+		: FText::Format(LOCTEXT("LayerQueryText", "{0} layer"), LayerType);
 }
+#endif
 
 #undef LOCTEXT_NAMESPACE
diff --git a/Engine/Plugins/Experimental/Avalanche/Source/AvalancheTransition/Private/Conditions/AvaTransitionLayerCondition.cpp b/Engine/Plugins/Experimental/Avalanche/Source/AvalancheTransition/Private/Conditions/AvaTransitionLayerCondition.cpp
--- a/Engine/Plugins/Experimental/Avalanche/Source/AvalancheTransition/Private/Conditions/AvaTransitionLayerCondition.cpp
+++ b/Engine/Plugins/Experimental/Avalanche/Source/AvalancheTransition/Private/Conditions/AvaTransitionLayerCondition.cpp
@@ -9,6 +9,21 @@
 
 #define LOCTEXT_NAMESPACE "AvaTransitionLayerCondition"
 
+#if WITH_EDITOR
+FText FAvaTransitionLayerCondition::GetDescription(const FGuid& InId, FStateTreeDataView InInstanceDataView, const IStateTreeBindingLookup& InBindingLookup, EStateTreeNodeFormatting InFormatting) const
+{
+	const FInstanceDataType& InstanceData = InInstanceDataView.Get<FInstanceDataType>();
+
+	FAvaTransitionLayerUtils::FLayerQueryTextParams Params;
+	Params.LayerType = InstanceData.LayerType;
+	Params.SpecificLayerName = InstanceData.SpecificLayer.ToName();
+	Params.LayerTypePropertyName = GET_MEMBER_NAME_CHECKED(FInstanceDataType, LayerType);
+	Params.SpecificLayerPropertyName = GET_MEMBER_NAME_CHECKED(FInstanceDataType, SpecificLayer);
+
+	return FAvaTransitionLayerUtils::GetLayerQueryText(MoveTemp(Params), InId, InBindingLookup, InFormatting);
+}
+#endif
+
 void FAvaTransitionLayerCondition::PostLoad(FStateTreeDataView InInstanceDataView)
 {
 	Super::PostLoad(InInstanceDataView);
@@ -36,9 +51,4 @@ TArray<const FAvaTransitionBehaviorInstance*> FAvaTransitionLayerCondition::Quer
 	return FAvaTransitionLayerUtils::QueryBehaviorInstances(TransitionSubsystem, Comparator);
 }
 
-FText FAvaTransitionLayerCondition::GetLayerQueryText(const FInstanceDataType& InInstanceData) const
-{
-	return FAvaTransitionLayerUtils::GetLayerQueryText(InInstanceData.LayerType, InInstanceData.SpecificLayer.ToName());
-}
-
 #undef LOCTEXT_NAMESPACE
diff --git a/Engine/Plugins/Experimental/Avalanche/Source/AvalancheTransition/Private/Conditions/AvaTransitionLayerMatchCondition.cpp b/Engine/Plugins/Experimental/Avalanche/Source/AvalancheTransition/Private/Conditions/AvaTransitionLayerMatchCondition.cpp
--- a/Engine/Plugins/Experimental/Avalanche/Source/AvalancheTransition/Private/Conditions/AvaTransitionLayerMatchCondition.cpp
+++ b/Engine/Plugins/Experimental/Avalanche/Source/AvalancheTransition/Private/Conditions/AvaTransitionLayerMatchCondition.cpp
@@ -2,16 +2,20 @@
 
 #include "Conditions/AvaTransitionLayerMatchCondition.h"
 #include "AvaTransitionContext.h"
-#include "StateTreeExecutionContext.h"
+#include "AvaTransitionLayerUtils.h"
 #include "Behavior/AvaTransitionBehaviorInstance.h"
+#include "StateTreeExecutionContext.h"
 
 #define LOCTEXT_NAMESPACE "AvaTransitionLayerMatchCondition"
 
 #if WITH_EDITOR
 FText FAvaTransitionLayerMatchCondition::GetDescription(const FGuid& InId, FStateTreeDataView InInstanceDataView, const IStateTreeBindingLookup& InBindingLookup, EStateTreeNodeFormatting InFormatting) const
 {
-	const FInstanceDataType& InstanceData = InInstanceDataView.Get<FInstanceDataType>();
-	return FText::Format(LOCTEXT("ConditionDescription", "scenes transitioning in {0}"), GetLayerQueryText(InstanceData));
+	const FText LayerDesc = Super::GetDescription(InId, InInstanceDataView, InBindingLookup, InFormatting);
+
+	return InFormatting == EStateTreeNodeFormatting::RichText
+		? FText::Format(LOCTEXT("DescRich", "<s>scenes transitioning in</> {0}"), LayerDesc)
+		: FText::Format(LOCTEXT("Desc", "scenes transitioning in {0}"), LayerDesc);
 }
 #endif
 
diff --git a/Engine/Plugins/Experimental/Avalanche/Source/AvalancheTransition/Private/Conditions/AvaTransitionSceneMatchCondition.cpp b/Engine/Plugins/Experimental/Avalanche/Source/AvalancheTransition/Private/Conditions/AvaTransitionSceneMatchCondition.cpp
--- a/Engine/Plugins/Experimental/Avalanche/Source/AvalancheTransition/Private/Conditions/AvaTransitionSceneMatchCondition.cpp
+++ b/Engine/Plugins/Experimental/Avalanche/Source/AvalancheTransition/Private/Conditions/AvaTransitionSceneMatchCondition.cpp
@@ -2,6 +2,7 @@
 
 #include "Conditions/AvaTransitionSceneMatchCondition.h"
 #include "AvaTransitionContext.h"
+#include "AvaTransitionLayerUtils.h"
 #include "AvaTransitionScene.h"
 #include "AvaTransitionUtils.h"
 #include "Behavior/AvaTransitionBehaviorInstance.h"
@@ -14,9 +15,12 @@ FText FAvaTransitionSceneMatchCondition::GetDescription(const FGuid& InId, FStat
 {
 	const FInstanceDataType& InstanceData = InInstanceDataView.Get<FInstanceDataType>();
 
-	return FText::Format(LOCTEXT("ConditionDescription", "{0} scene in {1}")
-		, UEnum::GetDisplayValueAsText(InstanceData.SceneComparisonType).ToLower()
-		, GetLayerQueryText(InstanceData));
+	const FText ComparisonType = UEnum::GetDisplayValueAsText(InstanceData.SceneComparisonType).ToLower();
+	const FText LayerDesc = Super::GetDescription(InId, InInstanceDataView, InBindingLookup, InFormatting);
+
+	return InFormatting == EStateTreeNodeFormatting::RichText
+		? FText::Format(LOCTEXT("DescRich", "<b>{0}</> <s>scene in</> {1}"), ComparisonType, LayerDesc)
+		: FText::Format(LOCTEXT("Desc", "{0} scene in {1}"), ComparisonType, LayerDesc);
 }
 #endif
 
diff --git a/Engine/Plugins/Experimental/Avalanche/Source/AvalancheTransition/Private/Conditions/AvaTransitionStateMatchCondition.cpp b/Engine/Plugins/Experimental/Avalanche/Source/AvalancheTransition/Private/Conditions/AvaTransitionStateMatchCondition.cpp
--- a/Engine/Plugins/Experimental/Avalanche/Source/AvalancheTransition/Private/Conditions/AvaTransitionStateMatchCondition.cpp
+++ b/Engine/Plugins/Experimental/Avalanche/Source/AvalancheTransition/Private/Conditions/AvaTransitionStateMatchCondition.cpp
@@ -1,6 +1,7 @@
 // Copyright Epic Games, Inc. All Rights Reserved.
 
 #include "Conditions/AvaTransitionStateMatchCondition.h"
+#include "AvaTransitionLayerUtils.h"
 #include "AvaTransitionUtils.h"
 #include "Behavior/AvaTransitionBehaviorInstance.h"
 #include "StateTreeExecutionContext.h"
@@ -12,9 +13,12 @@ FText FAvaTransitionStateMatchCondition::GetDescription(const FGuid& InId, FStat
 {
 	const FInstanceDataType& InstanceData = InInstanceDataView.Get<FInstanceDataType>();
 
-	return FText::Format(LOCTEXT("ConditionDescription", "{0} scene in {1}")
-		, UEnum::GetDisplayValueAsText(InstanceData.TransitionState).ToLower()
-		, GetLayerQueryText(InstanceData));
+	const FText TransitionState = UEnum::GetDisplayValueAsText(InstanceData.TransitionState).ToLower();
+	const FText LayerDesc = Super::GetDescription(InId, InInstanceDataView, InBindingLookup, InFormatting);
+
+	return InFormatting == EStateTreeNodeFormatting::RichText
+		? FText::Format(LOCTEXT("DescRich", "<b>{0}</> <s>scene in</> {1}"), TransitionState, LayerDesc)
+		: FText::Format(LOCTEXT("Desc", "{0} scene in {1}"), TransitionState, LayerDesc);
 }
 #endif
 
diff --git a/Engine/Plugins/Experimental/Avalanche/Source/AvalancheTransition/Private/Conditions/AvaTransitionTypeMatchCondition.cpp b/Engine/Plugins/Experimental/Avalanche/Source/AvalancheTransition/Private/Conditions/AvaTransitionTypeMatchCondition.cpp
--- a/Engine/Plugins/Experimental/Avalanche/Source/AvalancheTransition/Private/Conditions/AvaTransitionTypeMatchCondition.cpp
+++ b/Engine/Plugins/Experimental/Avalanche/Source/AvalancheTransition/Private/Conditions/AvaTransitionTypeMatchCondition.cpp
@@ -12,8 +12,11 @@ FText FAvaTransitionTypeMatchCondition::GetDescription(const FGuid& InId, FState
 {
 	const FInstanceDataType& InstanceData = InInstanceDataView.Get<FInstanceDataType>();
 
-	return FText::Format(LOCTEXT("ConditionDescription", "transitioning {0}")
-		, UEnum::GetDisplayValueAsText(InstanceData.TransitionType).ToLower());
+	const FText TransitionTypeDesc = UEnum::GetDisplayValueAsText(InstanceData.TransitionType).ToLower();
+
+	return InFormatting == EStateTreeNodeFormatting::RichText
+		? FText::Format(LOCTEXT("DescRich", "<s>transitioning</> <b>{0}</>"), TransitionTypeDesc)
+		: FText::Format(LOCTEXT("Desc", "transitioning {0}"), TransitionTypeDesc);
 }
 #endif
 
diff --git a/Engine/Plugins/Experimental/Avalanche/Source/AvalancheTransition/Private/Tasks/AvaTransitionDelayTask.cpp b/Engine/Plugins/Experimental/Avalanche/Source/AvalancheTransition/Private/Tasks/AvaTransitionDelayTask.cpp
--- a/Engine/Plugins/Experimental/Avalanche/Source/AvalancheTransition/Private/Tasks/AvaTransitionDelayTask.cpp
+++ b/Engine/Plugins/Experimental/Avalanche/Source/AvalancheTransition/Private/Tasks/AvaTransitionDelayTask.cpp
@@ -10,7 +10,12 @@
 FText FAvaTransitionDelayTask::GetDescription(const FGuid& InId, FStateTreeDataView InInstanceDataView, const IStateTreeBindingLookup& InBindingLookup, EStateTreeNodeFormatting InFormatting) const
 {
 	const FInstanceDataType& InstanceData = InInstanceDataView.Get<FInstanceDataType>();
-	return FText::Format(LOCTEXT("TaskDescription", "Delay {0} seconds"), FText::AsNumber(InstanceData.Duration));
+
+	const FText DurationDesc = FText::AsNumber(InstanceData.Duration);
+
+	return InFormatting == EStateTreeNodeFormatting::RichText
+		? FText::Format(LOCTEXT("DescRich", "Delay <b>{0}</> <s>seconds</>"), DurationDesc)
+		: FText::Format(LOCTEXT("Desc", "Delay {0} seconds"), DurationDesc);
 }
 #endif
 
diff --git a/Engine/Plugins/Experimental/Avalanche/Source/AvalancheTransition/Private/Tasks/AvaTransitionLayerTask.cpp b/Engine/Plugins/Experimental/Avalanche/Source/AvalancheTransition/Private/Tasks/AvaTransitionLayerTask.cpp
--- a/Engine/Plugins/Experimental/Avalanche/Source/AvalancheTransition/Private/Tasks/AvaTransitionLayerTask.cpp
+++ b/Engine/Plugins/Experimental/Avalanche/Source/AvalancheTransition/Private/Tasks/AvaTransitionLayerTask.cpp
@@ -9,6 +9,21 @@
 
 #define LOCTEXT_NAMESPACE "AvaTransitionLayerTask"
 
+#if WITH_EDITOR
+FText FAvaTransitionLayerTask::GetDescription(const FGuid& InId, FStateTreeDataView InInstanceDataView, const IStateTreeBindingLookup& InBindingLookup, EStateTreeNodeFormatting InFormatting) const
+{
+	const FInstanceDataType& InstanceData = InInstanceDataView.Get<FInstanceDataType>();
+
+	FAvaTransitionLayerUtils::FLayerQueryTextParams Params;
+	Params.LayerType = InstanceData.LayerType;
+	Params.SpecificLayerName = InstanceData.SpecificLayer.ToName();
+	Params.LayerTypePropertyName = GET_MEMBER_NAME_CHECKED(FInstanceDataType, LayerType);
+	Params.SpecificLayerPropertyName = GET_MEMBER_NAME_CHECKED(FInstanceDataType, SpecificLayer);
+
+	return FAvaTransitionLayerUtils::GetLayerQueryText(MoveTemp(Params), InId, InBindingLookup, InFormatting);
+}
+#endif
+
 void FAvaTransitionLayerTask::PostLoad(FStateTreeDataView InInstanceDataView)
 {
 	Super::PostLoad(InInstanceDataView);
@@ -25,11 +40,6 @@ void FAvaTransitionLayerTask::PostLoad(FStateTreeDataView InInstanceDataView)
 	PRAGMA_ENABLE_DEPRECATION_WARNINGS
 }
 
-FText FAvaTransitionLayerTask::GetLayerQueryText(const FInstanceDataType& InInstanceData) const
-{
-	return FAvaTransitionLayerUtils::GetLayerQueryText(InInstanceData.LayerType, InInstanceData.SpecificLayer.ToName());
-}
-
 TArray<const FAvaTransitionBehaviorInstance*> FAvaTransitionLayerTask::QueryBehaviorInstances(const FStateTreeExecutionContext& InContext) const
 {
 	const FInstanceDataType& InstanceData = InContext.GetInstanceData(*this);
diff --git a/Engine/Plugins/Experimental/Avalanche/Source/AvalancheTransition/Private/Tasks/AvaTransitionWaitForLayerTask.cpp b/Engine/Plugins/Experimental/Avalanche/Source/AvalancheTransition/Private/Tasks/AvaTransitionWaitForLayerTask.cpp
--- a/Engine/Plugins/Experimental/Avalanche/Source/AvalancheTransition/Private/Tasks/AvaTransitionWaitForLayerTask.cpp
+++ b/Engine/Plugins/Experimental/Avalanche/Source/AvalancheTransition/Private/Tasks/AvaTransitionWaitForLayerTask.cpp
@@ -10,8 +10,11 @@
 #if WITH_EDITOR
 FText FAvaTransitionWaitForLayerTask::GetDescription(const FGuid& InId, FStateTreeDataView InInstanceDataView, const IStateTreeBindingLookup& InBindingLookup, EStateTreeNodeFormatting InFormatting) const
 {
-	const FInstanceDataType& InstanceData = InInstanceDataView.Get<FInstanceDataType>();
-	return FText::Format(LOCTEXT("TaskDescription", "Wait for others in {0} to finish"), GetLayerQueryText(InstanceData));
+	const FText LayerDesc = Super::GetDescription(InId, InInstanceDataView, InBindingLookup, InFormatting);
+
+	return InFormatting == EStateTreeNodeFormatting::RichText
+		? FText::Format(LOCTEXT("DescRich", "Wait <s>for others in</> {0} <s>to finish</>"), LayerDesc)
+		: FText::Format(LOCTEXT("Desc", "Wait for others in {0} to finish"), LayerDesc);
 }
 #endif
 
diff --git a/Engine/Plugins/Experimental/Avalanche/Source/AvalancheTransition/Public/AvaTransitionLayerUtils.h b/Engine/Plugins/Experimental/Avalanche/Source/AvalancheTransition/Public/AvaTransitionLayerUtils.h
--- a/Engine/Plugins/Experimental/Avalanche/Source/AvalancheTransition/Public/AvaTransitionLayerUtils.h
+++ b/Engine/Plugins/Experimental/Avalanche/Source/AvalancheTransition/Public/AvaTransitionLayerUtils.h
@@ -2,27 +2,39 @@
 
 #pragma once
 
+#include "StateTreeTypes.h"
 #include "Containers/ContainersFwd.h"
 #include "UObject/NameTypes.h"
 
 class UAvaTransitionSubsystem;
+enum EStateTreeNodeFormatting : uint8;
 enum class EAvaTransitionLayerCompareType : uint8;
 struct FAvaTagHandle;
 struct FAvaTagHandleContainer;
 struct FAvaTransitionBehaviorInstance;
 struct FAvaTransitionContext;
 struct FAvaTransitionLayerComparator;
+struct IStateTreeBindingLookup;
 
-struct AVALANCHETRANSITION_API FAvaTransitionLayerUtils
+struct FAvaTransitionLayerUtils
 {
 	/** Gets all the Behavior Instances that match the Comparator */
-	static TArray<const FAvaTransitionBehaviorInstance*> QueryBehaviorInstances(UAvaTransitionSubsystem& InTransitionSubsystem, const FAvaTransitionLayerComparator& InComparator);
+	AVALANCHETRANSITION_API static TArray<const FAvaTransitionBehaviorInstance*> QueryBehaviorInstances(UAvaTransitionSubsystem& InTransitionSubsystem, const FAvaTransitionLayerComparator& InComparator);
 
 	/** Builds a Comparator for the given Context (and optionally Layer), excluding the provided Transition Context (assumes it's the transition context calling this) */
-	static FAvaTransitionLayerComparator BuildComparator(const FAvaTransitionContext& InTransitionContext, EAvaTransitionLayerCompareType InCompareType, const FAvaTagHandle& InSpecificLayer);
+	AVALANCHETRANSITION_API static FAvaTransitionLayerComparator BuildComparator(const FAvaTransitionContext& InTransitionContext, EAvaTransitionLayerCompareType InCompareType, const FAvaTagHandle& InSpecificLayer);
 
 	/** Builds a Comparator for the given Context and specific layers, excluding the provided Transition Context (assumes it's the transition context calling this) */
-	static FAvaTransitionLayerComparator BuildComparator(const FAvaTransitionContext& InTransitionContext, EAvaTransitionLayerCompareType InCompareType, const FAvaTagHandleContainer& InSpecificLayers);
+	AVALANCHETRANSITION_API static FAvaTransitionLayerComparator BuildComparator(const FAvaTransitionContext& InTransitionContext, EAvaTransitionLayerCompareType InCompareType, const FAvaTagHandleContainer& InSpecificLayers);
 
-	static FText GetLayerQueryText(EAvaTransitionLayerCompareType InLayerType, FName InSpecificLayerName);
+#if WITH_EDITOR
+	struct FLayerQueryTextParams
+	{
+		EAvaTransitionLayerCompareType LayerType;
+		FName SpecificLayerName;
+		FName LayerTypePropertyName;
+		FName SpecificLayerPropertyName;
+	};
+	AVALANCHETRANSITION_API static FText GetLayerQueryText(FLayerQueryTextParams&& InParams, const FGuid& InId, const IStateTreeBindingLookup& InBindingLookup, EStateTreeNodeFormatting InFormatting);
+#endif
 };
diff --git a/Engine/Plugins/Experimental/Avalanche/Source/AvalancheTransition/Public/Conditions/AvaTransitionLayerCondition.h b/Engine/Plugins/Experimental/Avalanche/Source/AvalancheTransition/Public/Conditions/AvaTransitionLayerCondition.h
--- a/Engine/Plugins/Experimental/Avalanche/Source/AvalancheTransition/Public/Conditions/AvaTransitionLayerCondition.h
+++ b/Engine/Plugins/Experimental/Avalanche/Source/AvalancheTransition/Public/Conditions/AvaTransitionLayerCondition.h
@@ -38,6 +38,9 @@ struct AVALANCHETRANSITION_API FAvaTransitionLayerCondition : public FAvaTransit
 	PRAGMA_ENABLE_DEPRECATION_WARNINGS
 
 	//~ Begin FStateTreeNodeBase
+#if WITH_EDITOR
+	virtual FText GetDescription(const FGuid& InId, FStateTreeDataView InInstanceDataView, const IStateTreeBindingLookup& InBindingLookup, EStateTreeNodeFormatting InFormatting) const override;
+#endif
 	virtual const UStruct* GetInstanceDataType() const override { return FInstanceDataType::StaticStruct(); }
 	virtual void PostLoad(FStateTreeDataView InInstanceDataView) override;
 	//~ End FStateTreeNodeBase
@@ -45,8 +48,6 @@ struct AVALANCHETRANSITION_API FAvaTransitionLayerCondition : public FAvaTransit
 	/** Gets all the Behavior Instances that match the Layer Query. Always excludes the Instance belonging to this Transition */
 	TArray<const FAvaTransitionBehaviorInstance*> QueryBehaviorInstances(const FStateTreeExecutionContext& InContext) const;
 
-	FText GetLayerQueryText(const FInstanceDataType& InInstanceData) const;
-
 	UE_DEPRECATED(5.5, "LayerType has been moved to Instance Data")
 	UPROPERTY(meta = (DeprecatedProperty, DeprecationMessage = "Use the Instance Data LayerType instead"))
 	EAvaTransitionLayerCompareType LayerType_DEPRECATED = EAvaTransitionLayerCompareType::None;
diff --git a/Engine/Plugins/Experimental/Avalanche/Source/AvalancheTransition/Public/Tasks/AvaTransitionLayerTask.h b/Engine/Plugins/Experimental/Avalanche/Source/AvalancheTransition/Public/Tasks/AvaTransitionLayerTask.h
--- a/Engine/Plugins/Experimental/Avalanche/Source/AvalancheTransition/Public/Tasks/AvaTransitionLayerTask.h
+++ b/Engine/Plugins/Experimental/Avalanche/Source/AvalancheTransition/Public/Tasks/AvaTransitionLayerTask.h
@@ -36,12 +36,13 @@ struct AVALANCHETRANSITION_API FAvaTransitionLayerTask : public FAvaTransitionTa
 	PRAGMA_ENABLE_DEPRECATION_WARNINGS
 
 	//~ Begin FStateTreeNodeBase
+#if WITH_EDITOR
+	virtual FText GetDescription(const FGuid& InId, FStateTreeDataView InInstanceDataView, const IStateTreeBindingLookup& InBindingLookup, EStateTreeNodeFormatting InFormatting) const override;
+#endif
 	virtual const UStruct* GetInstanceDataType() const override { return FInstanceDataType::StaticStruct(); }
 	virtual void PostLoad(FStateTreeDataView InInstanceDataView) override;
 	//~ End FStateTreeNodeBase
 
-	FText GetLayerQueryText(const FInstanceDataType& InInstanceData) const;
-
 	/** Gets all the Behavior Instances that match the Layer Query. Always excludes the Instance belonging to this Transition */
 	TArray<const FAvaTransitionBehaviorInstance*> QueryBehaviorInstances(const FStateTreeExecutionContext& InContext) const;
 
diff --git a/Engine/Plugins/Experimental/Avalanche/Source/AvalancheTransitionEditor/Private/AvaTransitionTreeEditorData.cpp b/Engine/Plugins/Experimental/Avalanche/Source/AvalancheTransitionEditor/Private/AvaTransitionTreeEditorData.cpp
--- a/Engine/Plugins/Experimental/Avalanche/Source/AvalancheTransitionEditor/Private/AvaTransitionTreeEditorData.cpp
+++ b/Engine/Plugins/Experimental/Avalanche/Source/AvalancheTransitionEditor/Private/AvaTransitionTreeEditorData.cpp
@@ -2,6 +2,29 @@
 
 #include "AvaTransitionTreeEditorData.h"
 
+UAvaTransitionTreeEditorData::UAvaTransitionTreeEditorData()
+{
+	FStateTreeEditorColor DefaultColor;
+	DefaultColor.ColorRef = FStateTreeEditorColorRef(UE::AvaTransitionEditor::ColorId_Default);
+	DefaultColor.Color = FLinearColor(0, 0.15, 0.2);
+	DefaultColor.DisplayName = TEXT("Default Color");
+
+	FStateTreeEditorColor InColor;
+	InColor.ColorRef = FStateTreeEditorColorRef(UE::AvaTransitionEditor::ColorId_In);
+	InColor.Color = FLinearColor(0, 0.2, 0);
+	InColor.DisplayName = TEXT("In Color");
+
+	FStateTreeEditorColor OutColor;
+	OutColor.ColorRef = FStateTreeEditorColorRef(UE::AvaTransitionEditor::ColorId_Out);
+	OutColor.Color = FLinearColor(0.2, 0, 0.15);
+	OutColor.DisplayName = TEXT("Out Color");
+
+	Colors.Empty(3);
+	Colors.Add(MoveTemp(DefaultColor));
+	Colors.Add(MoveTemp(InColor));
+	Colors.Add(MoveTemp(OutColor));
+}
+
 UStateTreeState& UAvaTransitionTreeEditorData::CreateState(const UStateTreeState& InSiblingState, bool bInAfter)
 {
 	UObject* Outer = this;
diff --git a/Engine/Plugins/Experimental/Avalanche/Source/AvalancheTransitionEditor/Private/Customizations/AvaTransitionTreeEditorDataCustomization.cpp b/Engine/Plugins/Experimental/Avalanche/Source/AvalancheTransitionEditor/Private/Customizations/AvaTransitionTreeEditorDataCustomization.cpp
--- a/Engine/Plugins/Experimental/Avalanche/Source/AvalancheTransitionEditor/Private/Customizations/AvaTransitionTreeEditorDataCustomization.cpp
+++ b/Engine/Plugins/Experimental/Avalanche/Source/AvalancheTransitionEditor/Private/Customizations/AvaTransitionTreeEditorDataCustomization.cpp
@@ -8,7 +8,6 @@
 #include "Modules/ModuleManager.h"
 #include "PropertyBagDetails.h"
 #include "PropertyEditorModule.h"
-#include "StateTreeEditorStyle.h"
 #include "ViewModels/AvaTransitionViewModelSharedData.h"
 
 TSharedRef<IDetailCustomization> FAvaTransitionTreeEditorDataCustomization::MakeInstance(TWeakPtr<FAvaTransitionViewModelSharedData> InSharedDataWeak)
diff --git a/Engine/Plugins/Experimental/Avalanche/Source/AvalancheTransitionEditor/Private/Styling/AvaTransitionEditorStyle.cpp b/Engine/Plugins/Experimental/Avalanche/Source/AvalancheTransitionEditor/Private/Styling/AvaTransitionEditorStyle.cpp
--- a/Engine/Plugins/Experimental/Avalanche/Source/AvalancheTransitionEditor/Private/Styling/AvaTransitionEditorStyle.cpp
+++ b/Engine/Plugins/Experimental/Avalanche/Source/AvalancheTransitionEditor/Private/Styling/AvaTransitionEditorStyle.cpp
@@ -1,18 +1,25 @@
 // Copyright Epic Games, Inc. All Rights Reserved.
 
 #include "AvaTransitionEditorStyle.h"
-#include "Brushes/SlateImageBrush.h"
 #include "AvaTransitionTree.h"
+#include "Brushes/SlateImageBrush.h"
 #include "Interfaces/IPluginManager.h"
 #include "Layout/Margin.h"
 #include "Misc/Paths.h"
+#include "StateTreeEditorStyle.h"
+#include "Styling/AppStyle.h"
+#include "Styling/CoreStyle.h"
 #include "Styling/SlateStyleMacros.h"
 #include "Styling/SlateStyleRegistry.h"
+#include "Styling/SlateTypes.h"
+#include "Styling/StyleColors.h"
 #include "Textures/SlateIcon.h"
 
 FAvaTransitionEditorStyle::FAvaTransitionEditorStyle()
 	: FSlateStyleSet(TEXT("AvaTransitionEditor"))
 {
+	ParentStyleName = FStateTreeEditorStyle::Get().GetStyleSetName();
+
 	const FVector2f Icon16(16.f);
 	const FVector2f Icon20(20.f);
 	const FVector2f Icon64(64.f);
@@ -20,8 +27,8 @@ FAvaTransitionEditorStyle::FAvaTransitionEditorStyle()
 	const TSharedPtr<IPlugin> Plugin = IPluginManager::Get().FindPlugin(UE_PLUGIN_NAME);
 	check(Plugin.IsValid());
 
-	SetContentRoot(FPaths::EngineContentDir() / TEXT("Editor/Slate"));
-	SetCoreContentRoot(FPaths::EngineContentDir() / TEXT("Slate"));
+	ContentRootDir = FPaths::EngineContentDir() / TEXT("Editor/Slate");
+	CoreContentRootDir = FPaths::EngineContentDir() / TEXT("Slate");
 
 	// Re-use the Behavior Tree Icon
 	Set("ClassIcon.AvaTransitionTree"     , new IMAGE_BRUSH("Icons/AssetIcons/BehaviorTree_16x", Icon16));
@@ -39,6 +46,12 @@ FAvaTransitionEditorStyle::FAvaTransitionEditorStyle()
 
 	Set("DebugIndicatorBorder", new BOX_BRUSH("Images/NamespaceBorder", FMargin(0.25f)));
 
+	FLinearColor OperandBoxColor = FStyleColors::AccentGreen.GetSpecifiedColor().Desaturate(0.3f);
+	OperandBoxColor.A = 1.f;
+	Set("OperandBox", new FSlateRoundedBoxBrush(OperandBoxColor, 4.0f));
+
+	SetupFonts();
+
 	FSlateStyleRegistry::RegisterSlateStyle(*this);
 }
 
@@ -91,3 +104,24 @@ FSlateIcon FAvaTransitionEditorStyle::ParseIcon(FName InIconName)
 
 	return FSlateIcon(IconPathNames[0], IconPathNames[1], IconPathNames[2], IconPathNames[3]);
 }
+
+void FAvaTransitionEditorStyle::SetupFonts()
+{
+	const FTextBlockStyle& NormalText = FAppStyle::Get().GetWidgetStyle<FTextBlockStyle>("NormalText");
+
+	const FTextBlockStyle TitleText = FTextBlockStyle(NormalText)
+		.SetFont(DEFAULT_FONT("Bold", 12));
+
+	Set("StateTree.State.Title", FTextBlockStyle(TitleText)
+		.SetColorAndOpacity(FLinearColor(FColor(230, 230, 230, 205))));
+
+	Set("StateTree.State.Title.Bold", FTextBlockStyle(TitleText)
+		.SetColorAndOpacity(FLinearColor(FColor(230, 230, 230, 230))));
+
+	Set("StateTree.State.Title.Subdued", FTextBlockStyle(TitleText)
+		.SetColorAndOpacity(FLinearColor(FColor(230, 230, 230, 180))));
+
+	Set("StateTree.State.Operand", FTextBlockStyle(NormalText)
+		.SetFont(FAppStyle::GetFontStyle(TEXT("PropertyWindow.BoldFont")))
+		.SetFontSize(10));
+}
diff --git a/Engine/Plugins/Experimental/Avalanche/Source/AvalancheTransitionEditor/Private/Styling/AvaTransitionEditorStyle.h b/Engine/Plugins/Experimental/Avalanche/Source/AvalancheTransitionEditor/Private/Styling/AvaTransitionEditorStyle.h
--- a/Engine/Plugins/Experimental/Avalanche/Source/AvalancheTransitionEditor/Private/Styling/AvaTransitionEditorStyle.h
+++ b/Engine/Plugins/Experimental/Avalanche/Source/AvalancheTransitionEditor/Private/Styling/AvaTransitionEditorStyle.h
@@ -21,4 +21,7 @@ public:
 	static FLinearColor LerpColorSRGB(const FLinearColor& InA, const FLinearColor& InB, float InAlpha);
 
 	static FSlateIcon ParseIcon(FName InIconName);
+
+private:
+	void SetupFonts();
 };
diff --git a/Engine/Plugins/Experimental/Avalanche/Source/AvalancheTransitionEditor/Private/Styling/AvaTransitionTextStyling.cpp b/Engine/Plugins/Experimental/Avalanche/Source/AvalancheTransitionEditor/Private/Styling/AvaTransitionTextStyling.cpp
new file mode 100644
--- /dev/null
+++ b/Engine/Plugins/Experimental/Avalanche/Source/AvalancheTransitionEditor/Private/Styling/AvaTransitionTextStyling.cpp
@@ -0,0 +1,63 @@
+// Copyright Epic Games, Inc. All Rights Reserved.
+
+#include "AvaTransitionTextStyling.h"
+#include "Framework/Text/SlateTextRun.h"
+#include "Styling/SlateTypes.h"
+
+TSharedRef<FAvaTransitionTextStyleDecorator> FAvaTransitionTextStyleDecorator::Create(FString InName, const FTextBlockStyle& InTextStyle)
+{
+	TSharedRef<FAvaTransitionTextStyleDecorator> TextStyleDecorator = MakeShared<FAvaTransitionTextStyleDecorator>();
+	TextStyleDecorator->DecoratorName = MoveTemp(InName);
+	TextStyleDecorator->TextStyle = InTextStyle;
+	return TextStyleDecorator;
+}
+
+bool FAvaTransitionTextStyleDecorator::Supports(const FTextRunParseResults& InRunInfo, const FString& InText) const
+{
+	return InRunInfo.Name == DecoratorName;
+}
+
+TSharedRef<ISlateRun> FAvaTransitionTextStyleDecorator::Create(const TSharedRef<class FTextLayout>& InTextLayout, const FTextRunParseResults& InRunParseResult, const FString& InOriginalText, const TSharedRef<FString>& InModelText, const ISlateStyle* InStyle)
+{
+	FRunInfo RunInfo(InRunParseResult.Name);
+	for (const TPair<FString, FTextRange>& Pair : InRunParseResult.MetaData)
+	{
+		RunInfo.MetaData.Add(Pair.Key, InOriginalText.Mid(Pair.Value.BeginIndex, Pair.Value.EndIndex - Pair.Value.BeginIndex));
+	}
+
+	FString Run;
+	if (!InRunParseResult.ContentRange.IsEmpty())
+	{
+		Run = InOriginalText.Mid(InRunParseResult.ContentRange.BeginIndex, InRunParseResult.ContentRange.Len());
+	}
+	else
+	{
+		// Handles the case when the decorator name is empty (matches the runs without any tags).
+		Run = InOriginalText.Mid(InRunParseResult.OriginalRange.BeginIndex, InRunParseResult.OriginalRange.Len());
+	}
+
+	// Remove any formatting in between (e.g. when using nested tags).
+	int32 TagStartIndex = INDEX_NONE;
+	for (int32 Index = 0; Index < Run.Len(); Index++)
+	{
+		if (TagStartIndex == INDEX_NONE)
+		{
+			if (Run[Index] == TEXT('<'))
+			{
+				TagStartIndex = Index;
+			}
+		}
+		else
+		{
+			if (Run[Index] == TEXT('>'))
+			{
+				Run.RemoveAt(TagStartIndex, Index - TagStartIndex + 1);
+				Index = TagStartIndex;
+				TagStartIndex = INDEX_NONE;
+			}
+		}
+	}
+
+	InModelText->Append(Run);
+	return FSlateTextRun::Create(RunInfo, InModelText, TextStyle);
+}
diff --git a/Engine/Plugins/Experimental/Avalanche/Source/AvalancheTransitionEditor/Private/Styling/AvaTransitionTextStyling.h b/Engine/Plugins/Experimental/Avalanche/Source/AvalancheTransitionEditor/Private/Styling/AvaTransitionTextStyling.h
new file mode 100644
--- /dev/null
+++ b/Engine/Plugins/Experimental/Avalanche/Source/AvalancheTransitionEditor/Private/Styling/AvaTransitionTextStyling.h
@@ -0,0 +1,30 @@
+// Copyright Epic Games, Inc. All Rights Reserved.
+
+#pragma once
+
+#include "Framework/Text/ITextDecorator.h"
+#include "Styling/SlateTypes.h"
+
+/**
+ * Simple decorator for SRichTextBlock, which allows to map simple/shorthand tags to specific style names.
+ * The decorator also strips away any tags, e.g. in case the initial text came with nested tags.
+ *
+ * Sourced from FTextStyleDecorator
+ */
+class FAvaTransitionTextStyleDecorator : public ITextDecorator
+{
+public:
+	static TSharedRef<FAvaTransitionTextStyleDecorator> Create(FString InName, const FTextBlockStyle& InTextStyle);
+
+	//~ Begin ITextDecorator
+	virtual bool Supports(const FTextRunParseResults& InRunInfo, const FString& InText) const override;
+	virtual TSharedRef<ISlateRun> Create(const TSharedRef<FTextLayout>& InTextLayout,const FTextRunParseResults& InRunParseResult,const FString& InOriginalText,const TSharedRef<FString>& InModelText,const ISlateStyle* InStyle) override;
+	//~ End ITextDecorator
+
+private:
+	/** Name of this decorator */
+	FString DecoratorName;
+
+	/** Color of this decorator */
+	FTextBlockStyle TextStyle;
+};
diff --git a/Engine/Plugins/Experimental/Avalanche/Source/AvalancheTransitionEditor/Private/Styling/AvaTransitionWidgetStyling.cpp b/Engine/Plugins/Experimental/Avalanche/Source/AvalancheTransitionEditor/Private/Styling/AvaTransitionWidgetStyling.cpp
new file mode 100644
--- /dev/null
+++ b/Engine/Plugins/Experimental/Avalanche/Source/AvalancheTransitionEditor/Private/Styling/AvaTransitionWidgetStyling.cpp
@@ -0,0 +1,42 @@
+// Copyright Epic Games, Inc. All Rights Reserved.
+
+#include "AvaTransitionWidgetStyling.h"
+#include "AvaTransitionEditorStyle.h"
+#include "Fonts/FontMeasure.h"
+#include "Framework/Application/SlateApplication.h"
+#include "Framework/Text/ITextDecorator.h"
+#include "Styling/StyleColors.h"
+#include "Widgets/Layout/SBorder.h"
+#include "Widgets/Text/STextBlock.h"
+
+FSlateWidgetRun::FWidgetRunInfo FAvaTransitionWidgetStyling::CreateOperandWidget(const FTextRunInfo& InRunInfo, const ISlateStyle* InStyle)
+{
+	FSlateColor BackgroundColor(FStyleColors::AccentGreen);
+	if (const FString* ColorString = InRunInfo.MetaData.Find(TEXT("color")))
+	{
+		BackgroundColor = FAppStyle::GetSlateColor(**ColorString);
+	}
+
+	const FTextBlockStyle& OperandStyle = FAvaTransitionEditorStyle::Get().GetWidgetStyle<FTextBlockStyle>("StateTree.State.Operand");
+	const FTextBlockStyle& TitleStyle = FAvaTransitionEditorStyle::Get().GetWidgetStyle<FTextBlockStyle>("StateTree.State.Title");
+
+	constexpr int32 VPadding = 2;
+
+	TSharedRef<SWidget> Widget = SNew(SBorder)
+		.BorderImage(FAvaTransitionEditorStyle::Get().GetBrush("OperandBox"))
+		.BorderBackgroundColor(BackgroundColor)
+		.Padding(3, VPadding)
+		[
+			SNew(STextBlock)
+			.Text(InRunInfo.Content)
+			.TextStyle(&OperandStyle)
+			.ColorAndOpacity(FLinearColor::White)
+		];
+
+	TSharedRef<FSlateFontMeasure> FontMeasure = FSlateApplication::Get().GetRenderer()->GetFontMeasureService();
+	int16 Baseline = FontMeasure->GetBaseline(OperandStyle.Font);
+	int16 TitleBaseline = FontMeasure->GetBaseline(TitleStyle.Font);
+
+	// ensure the resulting widget remains at the center of the widget
+	return FSlateWidgetRun::FWidgetRunInfo(Widget, (TitleBaseline - Baseline) - VPadding);
+}
diff --git a/Engine/Plugins/Experimental/Avalanche/Source/AvalancheTransitionEditor/Private/Styling/AvaTransitionWidgetStyling.h b/Engine/Plugins/Experimental/Avalanche/Source/AvalancheTransitionEditor/Private/Styling/AvaTransitionWidgetStyling.h
new file mode 100644
--- /dev/null
+++ b/Engine/Plugins/Experimental/Avalanche/Source/AvalancheTransitionEditor/Private/Styling/AvaTransitionWidgetStyling.h
@@ -0,0 +1,13 @@
+// Copyright Epic Games, Inc. All Rights Reserved.
+
+#pragma once
+
+#include "Framework/Text/SlateWidgetRun.h"
+
+class ISlateStyle;
+struct FTextRunInfo;
+
+struct FAvaTransitionWidgetStyling
+{
+	static FSlateWidgetRun::FWidgetRunInfo CreateOperandWidget(const FTextRunInfo& InRunInfo, const ISlateStyle* InStyle);
+};
diff --git a/Engine/Plugins/Experimental/Avalanche/Source/AvalancheTransitionEditor/Private/ViewModels/Condition/AvaTransitionConditionContainerViewModel.cpp b/Engine/Plugins/Experimental/Avalanche/Source/AvalancheTransitionEditor/Private/ViewModels/Condition/AvaTransitionConditionContainerViewModel.cpp
--- a/Engine/Plugins/Experimental/Avalanche/Source/AvalancheTransitionEditor/Private/ViewModels/Condition/AvaTransitionConditionContainerViewModel.cpp
+++ b/Engine/Plugins/Experimental/Avalanche/Source/AvalancheTransitionEditor/Private/ViewModels/Condition/AvaTransitionConditionContainerViewModel.cpp
@@ -4,8 +4,8 @@
 #include "AvaTransitionConditionViewModel.h"
 #include "AvaTransitionTreeEditorData.h"
 #include "Conditions/AvaTransitionCondition.h"
-#include "StateTreeEditorStyle.h"
 #include "StateTreeState.h"
+#include "Styling/AvaTransitionEditorStyle.h"
 #include "ViewModels/AvaTransitionViewModelUtils.h"
 #include "Widgets/Images/SImage.h"
 #include "Widgets/Layout/SBox.h"
@@ -17,9 +17,19 @@ namespace UE::AvaTransitionEditor::Private
 	FText GetOperandText(int32 InConditionIndex, const FStateTreeEditorNode& InEditorNode)
 	{
 		// First Conditions or Copy operands should not have any operand displayed
-		if (InConditionIndex > 0 && InEditorNode.ExpressionOperand != EStateTreeExpressionOperand::Copy)
+		if (InConditionIndex > 0)
 		{
-			return FText::Format(INVTEXT("{0} "), UEnum::GetDisplayValueAsText(InEditorNode.ExpressionOperand).ToLower());
+			switch (InEditorNode.ExpressionOperand)
+			{
+			case EStateTreeExpressionOperand::Copy:
+				break;
+
+			case EStateTreeExpressionOperand::And:
+				return INVTEXT("<op color=\"Colors.AccentPink\">AND</> ");
+
+			case EStateTreeExpressionOperand::Or:
+				return INVTEXT("<op color=\"Colors.AccentBlue\">OR</> ");
+			}
 		}
 		return FText::GetEmpty();
 	}
@@ -31,7 +41,7 @@ namespace UE::AvaTransitionEditor::Private
 			TArray<FText> Parenthesis;
 			Parenthesis.Reserve(IndentAmount);
 
-			FText ParenthesisType = InDeltaIndent > 0 ? INVTEXT("(") : INVTEXT(")");
+			FText ParenthesisType = InDeltaIndent > 0 ? INVTEXT("( ") : INVTEXT(" )");
 
 			for (int8 IndentCount = 0; IndentCount < IndentAmount; ++IndentCount)
 			{
@@ -102,9 +112,13 @@ FText FAvaTransitionConditionContainerViewModel::UpdateStateDescription() const
 			? LOCTEXT("ConditionFormatOpening", "{Operand}{Parenthesis}{Description}")
 			: LOCTEXT("ConditionFormatClosing", "{Operand}{Description}{Parenthesis}");
 
+		// Formatting set to Text as Titles due to the Background color do not visualize Bold/Subdued/Normal very well
+		// However, the State Text is still using Rich text for things like Operand
+		constexpr EStateTreeNodeFormatting Formatting = EStateTreeNodeFormatting::Text;
+
 		FFormatNamedArguments TextArguments;
 		TextArguments.Add(TEXT("Operand"), Private::GetOperandText(ConditionIndex, *EditorNode));
-		TextArguments.Add(TEXT("Description"), EditorData->GetNodeDescription(*EditorNode, EStateTreeNodeFormatting::RichText));
+		TextArguments.Add(TEXT("Description"), EditorData->GetNodeDescription(*EditorNode, Formatting));
 		TextArguments.Add(TEXT("Parenthesis"), Private::GetParenthesisText(DeltaIndent));
 
 		ConditionDescriptions.Add(FText::Format(TextFormat, TextArguments));
@@ -136,7 +150,7 @@ TSharedRef<SWidget> FAvaTransitionConditionContainerViewModel::CreateWidget()
 		[
 			SNew(SImage)
 			.ColorAndOpacity(FLinearColor(1, 1, 1, 0.5f))
-			.Image(FStateTreeEditorStyle::Get().GetBrush("StateTreeEditor.StateConditions"))
+			.Image(FAvaTransitionEditorStyle::Get().GetBrush("StateTreeEditor.StateConditions"))
 			.ToolTipText(LOCTEXT("StateHasEnterConditions", "State selection is guarded with enter conditions."))
 		];
 }
diff --git a/Engine/Plugins/Experimental/Avalanche/Source/AvalancheTransitionEditor/Private/ViewModels/Task/AvaTransitionTaskViewModel.cpp b/Engine/Plugins/Experimental/Avalanche/Source/AvalancheTransitionEditor/Private/ViewModels/Task/AvaTransitionTaskViewModel.cpp
--- a/Engine/Plugins/Experimental/Avalanche/Source/AvalancheTransitionEditor/Private/ViewModels/Task/AvaTransitionTaskViewModel.cpp
+++ b/Engine/Plugins/Experimental/Avalanche/Source/AvalancheTransitionEditor/Private/ViewModels/Task/AvaTransitionTaskViewModel.cpp
@@ -3,9 +3,9 @@
 #include "AvaTransitionTaskViewModel.h"
 #include "AvaTransitionTreeEditorData.h"
 #include "StateTreeEditorData.h"
-#include "StateTreeEditorStyle.h"
 #include "StateTreeState.h"
 #include "Styling/AvaTransitionEditorStyle.h"
+#include "Styling/AvaTransitionTextStyling.h"
 #include "Tasks/AvaTransitionTask.h"
 #include "Textures/SlateIcon.h"
 #include "ViewModels/AvaTransitionViewModelSharedData.h"
@@ -15,7 +15,7 @@
 #include "Widgets/Layout/SBox.h"
 #include "Widgets/SBoxPanel.h"
 #include "Widgets/SOverlay.h"
-#include "Widgets/Text/STextBlock.h"
+#include "Widgets/Text/SRichTextBlock.h"
 
 #define LOCTEXT_NAMESPACE "AvaTransitionTaskViewModel"
 
@@ -169,12 +169,15 @@ TSharedRef<SWidget> FAvaTransitionTaskViewModel::CreateWidget()
 				.VAlign(VAlign_Center)
 				.HAlign(HAlign_Left)
 				[
-					SNew(STextBlock)
+					SNew(SRichTextBlock)
 					.Margin(FMargin(4.f, 0.f))
 					.Text(this, &FAvaTransitionTaskViewModel::GetTaskDescription)
 					.ToolTipText(this, &FAvaTransitionTaskViewModel::GetTaskDescription)
-					.TextStyle(FStateTreeEditorStyle::Get(), "StateTree.Task.Title")
+					.TextStyle(&FAvaTransitionEditorStyle::Get().GetWidgetStyle<FTextBlockStyle>("StateTree.Task.Title"))
 					.OverflowPolicy(ETextOverflowPolicy::Ellipsis)
+					+SRichTextBlock::Decorator(FAvaTransitionTextStyleDecorator::Create(TEXT(""), FAvaTransitionEditorStyle::Get().GetWidgetStyle<FTextBlockStyle>("StateTree.Task.Title")))
+					+SRichTextBlock::Decorator(FAvaTransitionTextStyleDecorator::Create(TEXT("b"), FAvaTransitionEditorStyle::Get().GetWidgetStyle<FTextBlockStyle>("StateTree.Task.Title.Bold")))
+					+SRichTextBlock::Decorator(FAvaTransitionTextStyleDecorator::Create(TEXT("s"), FAvaTransitionEditorStyle::Get().GetWidgetStyle<FTextBlockStyle>("StateTree.Task.Title.Subdued")))
 				]
 			]
 			// Task Breakpoint
@@ -191,7 +194,7 @@ TSharedRef<SWidget> FAvaTransitionTaskViewModel::CreateWidget()
 					[
 						SNew(SImage)
 						.DesiredSizeOverride(FVector2D(10.f, 10.f))
-						.Image(FStateTreeEditorStyle::Get().GetBrush("StateTreeEditor.Debugger.Breakpoint.EnabledAndValid"))
+						.Image(FAvaTransitionEditorStyle::Get().GetBrush("StateTreeEditor.Debugger.Breakpoint.EnabledAndValid"))
 						.Visibility(this, &FAvaTransitionTaskViewModel::GetBreakpointVisibility)
 						.ToolTipText(this, &FAvaTransitionTaskViewModel::GetBreakpointTooltip)
 					]
diff --git a/Engine/Plugins/Experimental/Avalanche/Source/AvalancheTransitionEditor/Private/ViewModels/Transition/AvaTransitionTransitionViewModel.cpp b/Engine/Plugins/Experimental/Avalanche/Source/AvalancheTransitionEditor/Private/ViewModels/Transition/AvaTransitionTransitionViewModel.cpp
--- a/Engine/Plugins/Experimental/Avalanche/Source/AvalancheTransitionEditor/Private/ViewModels/Transition/AvaTransitionTransitionViewModel.cpp
+++ b/Engine/Plugins/Experimental/Avalanche/Source/AvalancheTransitionEditor/Private/ViewModels/Transition/AvaTransitionTransitionViewModel.cpp
@@ -3,8 +3,8 @@
 #include "AvaTransitionTransitionViewModel.h"
 #include "AvaTransitionTreeEditorData.h"
 #include "EditorFontGlyphs.h"
-#include "StateTreeEditorStyle.h"
 #include "StateTreeState.h"
+#include "Styling/AvaTransitionEditorStyle.h"
 #include "ViewModels/AvaTransitionEditorViewModel.h"
 #include "ViewModels/AvaTransitionViewModelSharedData.h"
 #include "ViewModels/AvaTransitionViewModelUtils.h"
@@ -68,7 +68,7 @@ const FSlateBrush* FAvaTransitionTransitionViewModel::GetIcon() const
 			const UStateTreeState* State = GetState();
 			if (State && State->Children.IsEmpty() && State->Type == EStateTreeStateType::State && EnumHasAnyFlags(Transition->Trigger, EStateTreeTransitionTrigger::OnStateCompleted))
 			{
-				return FStateTreeEditorStyle::Get().GetBrush("StateTreeEditor.Transition.Parent");
+				return FAvaTransitionEditorStyle::Get().GetBrush("StateTreeEditor.Transition.Parent");
 			}
 		}
 		// falls through
@@ -76,11 +76,11 @@ const FSlateBrush* FAvaTransitionTransitionViewModel::GetIcon() const
 	case EStateTreeTransitionType::Succeeded:
 	case EStateTreeTransitionType::Failed:
 	case EStateTreeTransitionType::GotoState:
-		return FStateTreeEditorStyle::Get().GetBrush("StateTreeEditor.Transition.Goto");
+		return FAvaTransitionEditorStyle::Get().GetBrush("StateTreeEditor.Transition.Goto");
 
 	case EStateTreeTransitionType::NextState:
 	case EStateTreeTransitionType::NextSelectableState:
-		return FStateTreeEditorStyle::Get().GetBrush("StateTreeEditor.Transition.Next");
+		return FAvaTransitionEditorStyle::Get().GetBrush("StateTreeEditor.Transition.Next");
 
 	default:
 		ensureMsgf(false, TEXT("Unhandled transition type."));
@@ -173,7 +173,7 @@ TSharedRef<SWidget> FAvaTransitionTransitionViewModel::CreateWidget()
 			[
 				SNew(SImage)
 				.DesiredSizeOverride(FVector2D(10, 10))
-				.Image(FStateTreeEditorStyle::Get().GetBrush(TEXT("StateTreeEditor.Debugger.Breakpoint.EnabledAndValid")))
+				.Image(FAvaTransitionEditorStyle::Get().GetBrush(TEXT("StateTreeEditor.Debugger.Breakpoint.EnabledAndValid")))
 				.Visibility(this, &FAvaTransitionTransitionViewModel::GetBreakpointVisibility)
 			]
 		]
@@ -184,7 +184,7 @@ TSharedRef<SWidget> FAvaTransitionTransitionViewModel::CreateWidget()
 		[
 			SNew(STextBlock)
 			.Text(this, &FAvaTransitionTransitionViewModel::GetDescription)
-			.TextStyle(FStateTreeEditorStyle::Get(), "StateTree.Details")
+			.TextStyle(FAvaTransitionEditorStyle::Get(), "StateTree.Details")
 		];
 }
 
diff --git a/Engine/Plugins/Experimental/Avalanche/Source/AvalancheTransitionEditor/Private/Views/SAvaTransitionStateRow.cpp b/Engine/Plugins/Experimental/Avalanche/Source/AvalancheTransitionEditor/Private/Views/SAvaTransitionStateRow.cpp
--- a/Engine/Plugins/Experimental/Avalanche/Source/AvalancheTransitionEditor/Private/Views/SAvaTransitionStateRow.cpp
+++ b/Engine/Plugins/Experimental/Avalanche/Source/AvalancheTransitionEditor/Private/Views/SAvaTransitionStateRow.cpp
@@ -3,7 +3,7 @@
 #include "SAvaTransitionStateRow.h"
 #include "Extensions/IAvaTransitionDragDropExtension.h"
 #include "SAvaTransitionStateView.h"
-#include "StateTreeEditorStyle.h"
+#include "Styling/AvaTransitionEditorStyle.h"
 #include "ViewModels/AvaTransitionViewModel.h"
 #include "ViewModels/AvaTransitionViewModelSharedData.h"
 #include "ViewModels/State/AvaTransitionStateViewModel.h"
@@ -13,7 +13,7 @@
 void SAvaTransitionStateRow::Construct(const FArguments& InArgs, const TSharedRef<STableViewBase>& InOwnerTableView, const TSharedRef<FAvaTransitionStateViewModel>& InStateViewModel)
 {
 	STableRow::FArguments TableRowArgs = STableRow::FArguments()
-		.Style(&FStateTreeEditorStyle::Get().GetWidgetStyle<FTableRowStyle>("StateTree.Selection"))
+		.Style(&FAvaTransitionEditorStyle::Get().GetWidgetStyle<FTableRowStyle>("StateTree.Selection"))
 		.Content()
 		[
 			SNew(SAvaTransitionStateView, InStateViewModel)
diff --git a/Engine/Plugins/Experimental/Avalanche/Source/AvalancheTransitionEditor/Private/Views/SAvaTransitionStateView.cpp b/Engine/Plugins/Experimental/Avalanche/Source/AvalancheTransitionEditor/Private/Views/SAvaTransitionStateView.cpp
--- a/Engine/Plugins/Experimental/Avalanche/Source/AvalancheTransitionEditor/Private/Views/SAvaTransitionStateView.cpp
+++ b/Engine/Plugins/Experimental/Avalanche/Source/AvalancheTransitionEditor/Private/Views/SAvaTransitionStateView.cpp
@@ -2,9 +2,9 @@
 
 #include "SAvaTransitionStateView.h"
 #include "SAvaTransitionStateMetadata.h"
-#include "StateTreeEditorStyle.h"
 #include "StateTreeTypes.h"
 #include "Styling/AvaTransitionEditorStyle.h"
+#include "Styling/AvaTransitionWidgetStyling.h"
 #include "ViewModels/Condition/AvaTransitionConditionContainerViewModel.h"
 #include "ViewModels/State/AvaTransitionStateViewModel.h"
 #include "ViewModels/Task/AvaTransitionTaskContainerViewModel.h"
@@ -14,6 +14,7 @@
 #include "Widgets/Layout/SBox.h"
 #include "Widgets/SBoxPanel.h"
 #include "Widgets/SOverlay.h"
+#include "Widgets/Text/SRichTextBlock.h"
 #include "Widgets/Text/STextBlock.h"
 
 #define LOCTEXT_NAMESPACE "SAvaTransitionStateView"
@@ -72,7 +73,7 @@ void SAvaTransitionStateView::Construct(const FArguments& InArgs, const TSharedR
 		.Padding(0.f, 4.f)
 		[
 			SNew(SBorder)
-			.BorderImage(FStateTreeEditorStyle::Get().GetBrush("StateTree.State.Border"))
+			.BorderImage(FAvaTransitionEditorStyle::Get().GetBrush("StateTree.State.Border"))
 			.BorderBackgroundColor(this, &SAvaTransitionStateView::GetActiveStateColor)
 			[
 				OuterStateBox
@@ -113,7 +114,7 @@ void SAvaTransitionStateView::Construct(const FArguments& InArgs, const TSharedR
 			SNew(STextBlock)
 			.Text(InStateViewModel, &FAvaTransitionStateViewModel::GetAutogeneratedLabel)
 			.Visibility(this, &SAvaTransitionStateView::GetAutogeneratedLabelVisibility)
-			.TextStyle(FStateTreeEditorStyle::Get(), "StateTree.Details")
+			.TextStyle(FAvaTransitionEditorStyle::Get(), "StateTree.Details")
 		];
 
 	ChildSlot
@@ -141,16 +142,19 @@ TSharedRef<SWidget> SAvaTransitionStateView::CreateStateSlotWidget(const TShared
 			]
 		];
 
+	TSharedRef<SAvaTransitionStateView> This = SharedThis(this);
+
 	// State Description
 	InStateBox->AddSlot()
 		.VAlign(VAlign_Center)
 		.FillWidth(1.f)
 		[
-			SNew(STextBlock)
+			SNew(SRichTextBlock)
 			.Text(InStateViewModel, &FAvaTransitionStateViewModel::GetStateDescription)
 			.ToolTipText(InStateViewModel, &FAvaTransitionStateViewModel::GetStateTooltip)
-			.TextStyle(FStateTreeEditorStyle::Get(),"StateTree.State.Title")
+			.TextStyle(&FAvaTransitionEditorStyle::Get().GetWidgetStyle<FTextBlockStyle>("StateTree.State.Title"))
 			.OverflowPolicy(ETextOverflowPolicy::Ellipsis)
+			+SRichTextBlock::WidgetDecorator(TEXT("op"), FWidgetDecorator::FCreateWidget::CreateStatic(&FAvaTransitionWidgetStyling::CreateOperandWidget))
 		];
 
 	return SNew(SBox)
@@ -176,7 +180,7 @@ TSharedRef<SWidget> SAvaTransitionStateView::CreateStateSlotWidget(const TShared
 				[
 					SNew(SImage)
 					.DesiredSizeOverride(FVector2D(12.f, 12.f))
-					.Image(FStateTreeEditorStyle::Get().GetBrush(TEXT("StateTreeEditor.Debugger.Breakpoint.EnabledAndValid")))
+					.Image(FAvaTransitionEditorStyle::Get().GetBrush(TEXT("StateTreeEditor.Debugger.Breakpoint.EnabledAndValid")))
 					.Visibility(this, &SAvaTransitionStateView::GetStateBreakpointVisibility)
 					.ToolTipText(InStateViewModel, &FAvaTransitionStateViewModel::GetBreakpointTooltip)
 				]
@@ -197,16 +201,16 @@ const FSlateBrush* SAvaTransitionStateView::GetSelectorIcon() const
 	switch(RuntimeBehavior)
 	{
 	case EStateTreeStateSelectionBehavior::None:
-		return FStateTreeEditorStyle::Get().GetBrush("StateTreeEditor.SelectNone");
+		return FAvaTransitionEditorStyle::Get().GetBrush("StateTreeEditor.SelectNone");
 
 	case EStateTreeStateSelectionBehavior::TryEnterState:
-		return FStateTreeEditorStyle::Get().GetBrush("StateTreeEditor.TryEnterState");
+		return FAvaTransitionEditorStyle::Get().GetBrush("StateTreeEditor.TryEnterState");
 
 	case EStateTreeStateSelectionBehavior::TrySelectChildrenInOrder:
-		return FStateTreeEditorStyle::Get().GetBrush("StateTreeEditor.TrySelectChildrenInOrder");
+		return FAvaTransitionEditorStyle::Get().GetBrush("StateTreeEditor.TrySelectChildrenInOrder");
 
 	case EStateTreeStateSelectionBehavior::TryFollowTransitions:
-		return FStateTreeEditorStyle::Get().GetBrush("StateTreeEditor.TryFollowTransitions");
+		return FAvaTransitionEditorStyle::Get().GetBrush("StateTreeEditor.TryFollowTransitions");
 	}
 
 	checkNoEntry();
diff --git a/Engine/Plugins/Experimental/Avalanche/Source/AvalancheTransitionEditor/Public/AvaTransitionTreeEditorData.h b/Engine/Plugins/Experimental/Avalanche/Source/AvalancheTransitionEditor/Public/AvaTransitionTreeEditorData.h
--- a/Engine/Plugins/Experimental/Avalanche/Source/AvalancheTransitionEditor/Public/AvaTransitionTreeEditorData.h
+++ b/Engine/Plugins/Experimental/Avalanche/Source/AvalancheTransitionEditor/Public/AvaTransitionTreeEditorData.h
@@ -14,12 +14,21 @@
 class UAvaTransitionTreeEditorData;
 class UStateTreeState;
 
+namespace UE::AvaTransitionEditor
+{
+	constexpr FGuid ColorId_Default;
+	constexpr FGuid ColorId_In(0x1DDBC788, 0xD5EB400E, 0xBB71E5DA, 0xB27A784D);
+	constexpr FGuid ColorId_Out(0xE549EFA0, 0xDEFF45A7, 0xA8D907AF, 0xDB8F7643);
+}
+
 UCLASS(MinimalAPI, HideCategories=(Common))
 class UAvaTransitionTreeEditorData : public UStateTreeEditorData
 {
 	GENERATED_BODY()
 
 public:
+	UAvaTransitionTreeEditorData();
+
 	AVALANCHETRANSITIONEDITOR_API UStateTreeState& CreateState(const UStateTreeState& InSiblingState, bool bInAfter);
 
 	FAvaTagHandle GetTransitionLayer() const
