MotionDesign : AvaShapes
- Fixed property displayed order to avoid materials options affecting shape options when new material slot becomes available

diff --git a/Engine/Plugins/Experimental/Avalanche/Source/AvalancheEditor/Private/DetailView/Customizations/AvaMeshesDetailCustomization.cpp b/Engine/Plugins/Experimental/Avalanche/Source/AvalancheEditor/Private/DetailView/Customizations/AvaMeshesDetailCustomization.cpp
--- a/Engine/Plugins/Experimental/Avalanche/Source/AvalancheEditor/Private/DetailView/Customizations/AvaMeshesDetailCustomization.cpp
+++ b/Engine/Plugins/Experimental/Avalanche/Source/AvalancheEditor/Private/DetailView/Customizations/AvaMeshesDetailCustomization.cpp
@@ -4,102 +4,55 @@
 
 #include "AssetRegistry/AssetRegistryModule.h"
 #include "AssetToolsModule.h"
-#include "DetailView/Widgets/SAvaDynamicMaterialWidget.h"
 #include "Components/DynamicMeshComponent.h"
 #include "DetailCategoryBuilder.h"
 #include "DetailLayoutBuilder.h"
+#include "DetailView/Widgets/SAvaDynamicMaterialWidget.h"
 #include "DetailWidgetRow.h"
 #include "Dialogs/DlgPickAssetPath.h"
+#include "DynamicMeshes/AvaShapeDynMeshBase.h"
 #include "DynamicMeshToMeshDescription.h"
 #include "Engine/StaticMesh.h"
 #include "IAssetTools.h"
-#include "IDetailGroup.h"
-#include "Material/DynamicMaterialInstance.h"
-#include "StaticMeshAttributes.h"
-#include "AvaShapeActor.h"
-#include "DynamicMeshes/AvaShapeDynMeshBase.h"
 #include "UObject/Object.h"
 #include "Widgets/Input/SButton.h"
 #include "Widgets/Text/STextBlock.h"
 
-#define LOCTEXT_NAMESPACE "MeshesDetailsCustomization"
+#define LOCTEXT_NAMESPACE "AvaMeshesDetailCustomization"
 
-void FAvaMeshesDetailCustomization::CustomizeDetails(IDetailLayoutBuilder& DetailBuilder)
+void FAvaMeshesDetailCustomization::CustomizeDetails(IDetailLayoutBuilder& InDetailBuilder)
 {
-	TSharedPtr<IPropertyHandle> MeshDatasHandle = DetailBuilder.GetProperty(
+	TSharedRef<IPropertyHandle> MeshDatasHandle = InDetailBuilder.GetProperty(
 		GET_MEMBER_NAME_CHECKED(UAvaShapeDynamicMeshBase, MeshDatas),
 		UAvaShapeDynamicMeshBase::StaticClass()
 	);
 
-	if (!MeshDatasHandle.IsValid())
-	{
-		return;
-	}
-
-	SelectedDynamicMeshes.Empty();
-	TArray<FName> MeshNames;
-	UAvaShapeDynamicMeshBase* DynMesh = nullptr;
-
-	for (const TWeakObjectPtr<UObject>& Object : DetailBuilder.GetSelectedObjects())
-	{
-		if (!Object.IsValid())
-		{
-			continue;
-		}
-
-		DynMesh = Cast<UAvaShapeDynamicMeshBase>(Object.Get());
-
-		if (!DynMesh)
-		{
-			AAvaShapeActor* ShapeActor = Cast<AAvaShapeActor>(Object.Get());
-
-			if (ShapeActor)
-			{
-				if (ShapeActor->GetDynamicMesh())
-				{
-					DynMesh = Cast<UAvaShapeDynamicMeshBase>(ShapeActor->GetDynamicMesh());
-				}
-			}
-		}
-
-		if (DynMesh)
-		{
-			MeshNames.Append(DynMesh->GetMeshDataNames());
-			SelectedDynamicMeshes.Add(DynMesh);
-		}
-	}
-
-	IDetailPropertyRow* MeshDatasRow = DetailBuilder.EditDefaultProperty(MeshDatasHandle);
-	DetailBuilder.HideProperty(MeshDatasHandle);
-
-	IDetailCategoryBuilder& CategoryBuilder = DetailBuilder.EditCategory(FName("Material"));
+	InDetailBuilder.HideProperty(MeshDatasHandle);
 
-	// use primary material everywhere
-	TSharedPtr<IPropertyHandle> UsePrimaryMaterialEverywhereHandle = DetailBuilder.GetProperty(
+	TSharedRef<IPropertyHandle> UsePrimaryMaterialEverywhereHandle = InDetailBuilder.GetProperty(
 		GET_MEMBER_NAME_CHECKED(UAvaShapeDynamicMeshBase, bUsePrimaryMaterialEverywhere),
 		UAvaShapeDynamicMeshBase::StaticClass()
 	);
 
-	// fix bug where toggle would appear in other group with same name
-	if (UsePrimaryMaterialEverywhereHandle.IsValid())
-	{
-		DetailBuilder.HideProperty(UsePrimaryMaterialEverywhereHandle);
-	}
+	InDetailBuilder.HideProperty(UsePrimaryMaterialEverywhereHandle);
+
+	MeshGeneratorsWeak = InDetailBuilder.GetObjectsOfTypeBeingCustomized<UAvaShapeDynamicMeshBase>();
+
+	// Set material category after shape category to avoid jump when new materials slot becomes available
+	IDetailCategoryBuilder& ShapeCategoryBuilder = InDetailBuilder.EditCategory(FName("Shape"));
+	IDetailCategoryBuilder& MaterialCategoryBuilder = InDetailBuilder.EditCategory(FName("Material"));
+	MaterialCategoryBuilder.SetSortOrder(ShapeCategoryBuilder.GetSortOrder() + 1);
 
-	if (MeshDatasRow)
+	if (MeshGeneratorsWeak.Num() == 1 && MeshGeneratorsWeak[0].IsValid())
 	{
-		MeshDatasRow->Visibility(EVisibility::Hidden);
+		UAvaShapeDynamicMeshBase* DynMesh = MeshGeneratorsWeak[0].Get();
 
 		TSharedPtr<IPropertyHandleMap> MapHandle = MeshDatasHandle->AsMap();
+
 		uint32 Count = 0;
 		MapHandle->GetNumElements(Count);
 
-		uint32 CountNames = MeshNames.Num();
-
-		if (Count < CountNames)
-		{
-			return;
-		}
+		const TArray<FName> MeshNames = DynMesh->GetMeshDataNames();
 
 		for (uint32 Index = 0; Index < Count; ++Index)
 		{
@@ -118,7 +71,7 @@ void FAvaMeshesDetailCustomization::CustomizeDetails(IDetailLayoutBuilder& Detai
 			check(MaterialTypeHandle.IsValid());
 
 			FString MaterialTypeName = MeshName + TEXT("Material Type");
-			FDetailWidgetRow& MaterialTypeRow = CategoryBuilder.AddCustomRow(FText::FromString(MaterialTypeName));
+			FDetailWidgetRow& MaterialTypeRow = MaterialCategoryBuilder.AddCustomRow(FText::FromString(MaterialTypeName));
 
 			MaterialTypeRow.NameContent()[MaterialTypeHandle->CreatePropertyNameWidget(FText::FromString(MaterialTypeName))];
 			MaterialTypeRow.ValueContent()[MaterialTypeHandle->CreatePropertyValueWidget()];
@@ -131,7 +84,7 @@ void FAvaMeshesDetailCustomization::CustomizeDetails(IDetailLayoutBuilder& Detai
 			check(MaterialHandle.IsValid());
 
 			FString MaterialName = MeshName + TEXT("Material Asset");
-			FDetailWidgetRow& MaterialRow = CategoryBuilder.AddCustomRow(FText::FromString(MaterialName));
+			FDetailWidgetRow& MaterialRow = MaterialCategoryBuilder.AddCustomRow(FText::FromString(MaterialName));
 
 			MaterialRow.NameContent()
 			[
@@ -160,7 +113,7 @@ void FAvaMeshesDetailCustomization::CustomizeDetails(IDetailLayoutBuilder& Detai
 				TSharedPtr<IPropertyHandle> ParametricChildHandle = ParametricMaterialHandle->GetChildHandle(ChildIdx);
 				check(ParametricChildHandle.IsValid());
 
-				IDetailPropertyRow& NewParametricRow = CategoryBuilder.AddProperty(ParametricChildHandle.ToSharedRef());
+				IDetailPropertyRow& NewParametricRow = MaterialCategoryBuilder.AddProperty(ParametricChildHandle.ToSharedRef());
 				FString ParametricRowName = MeshName + ParametricChildHandle->GetPropertyDisplayName().ToString();
 				NewParametricRow.DisplayName(FText::FromString(ParametricRowName));
 				NewParametricRow.Visibility(MakeAttributeLambda([=]() {
@@ -175,7 +128,7 @@ void FAvaMeshesDetailCustomization::CustomizeDetails(IDetailLayoutBuilder& Detai
 			check(UsePrimaryUVParamsHandle.IsValid());
 
 			FString UsePrimaryParamsName = MeshName + TEXT("Override UV");
-			FDetailWidgetRow& UsePrimaryUVParamsRow = CategoryBuilder.AddCustomRow(FText::FromString(UsePrimaryParamsName));
+			FDetailWidgetRow& UsePrimaryUVParamsRow = MaterialCategoryBuilder.AddCustomRow(FText::FromString(UsePrimaryParamsName));
 
 			UsePrimaryUVParamsRow.NameContent()[UsePrimaryUVParamsHandle->CreatePropertyNameWidget(FText::FromString(UsePrimaryParamsName))];
 			UsePrimaryUVParamsRow.ValueContent()[UsePrimaryUVParamsHandle->CreatePropertyValueWidget()];
@@ -184,62 +137,63 @@ void FAvaMeshesDetailCustomization::CustomizeDetails(IDetailLayoutBuilder& Detai
 			}));
 
 			// Only Add it the first time at this specific point
-			if (UsePrimaryMaterialEverywhereHandle.IsValid() && Index == 0 && MeshNames.Num() > 1)
+			if (Index == 0 && MeshNames.Num() > 1)
 			{
-				CategoryBuilder.AddProperty(UsePrimaryMaterialEverywhereHandle);
+				MaterialCategoryBuilder.AddProperty(UsePrimaryMaterialEverywhereHandle);
 			}
 
 			// uv params
 			TSharedPtr<IPropertyHandle> MaterialUVHandle = MeshPropertyHandle->GetChildHandle(GET_MEMBER_NAME_CHECKED(FAvaShapeMeshData, MaterialUVParams));
 			check(MaterialUVHandle.IsValid());
 
-			IDetailPropertyRow& MaterialUVRow = CategoryBuilder.AddProperty(MaterialUVHandle.ToSharedRef());
+			IDetailPropertyRow& MaterialUVRow = MaterialCategoryBuilder.AddProperty(MaterialUVHandle.ToSharedRef());
 			FString MaterialUVName = MeshName + TEXT("Material UV");
 			MaterialUVRow.DisplayName(FText::FromString(MaterialUVName));
 			MaterialUVRow.Visibility(MakeAttributeLambda([=]() {
 				return MaterialUVHandle->IsEditable() ? EVisibility::Visible : EVisibility::Hidden;
 			}));
 
-			// Separator row
-			FDetailWidgetRow& SeparatorRow = CategoryBuilder.AddCustomRow(FText::GetEmpty());
-			SeparatorRow.WholeRowContent()[
-				SNullWidget::NullWidget
-			];
-
-			// visibility for the separator row
-			SeparatorRow.Visibility(MakeAttributeLambda([=]() {
-				return UsePrimaryUVParamsHandle->IsEditable() || MaterialUVHandle->IsEditable() ? EVisibility::Visible : EVisibility::Hidden;
-			}));
+			if (Index < Count - 1)
+			{
+				// Separator row
+				FDetailWidgetRow& SeparatorRow = MaterialCategoryBuilder.AddCustomRow(FText::GetEmpty());
+				SeparatorRow.WholeRowContent()[
+					SNullWidget::NullWidget
+				];
+
+				// visibility for the separator row
+				SeparatorRow.Visibility(MakeAttributeLambda([=]() {
+					return UsePrimaryUVParamsHandle->IsEditable() || MaterialUVHandle->IsEditable() ? EVisibility::Visible : EVisibility::Hidden;
+				}));
+			}
 		}
-	}
 
-	// StaticMesh Export button
-
-	static const FText ConvertToStaticMeshText = LOCTEXT("ConvertToStaticMesh", "Create StaticMesh");
-
-	IDetailGroup& GroupBuilderExport = CategoryBuilder.AddGroup(FName("Export"), FText::FromString(TEXT("Export")), false, true);
-	GroupBuilderExport.ToggleExpansion(true);
-
-	GroupBuilderExport.AddWidgetRow()
-	.NameContent()
-	[
-		SNullWidget::NullWidget
-	]
-	.ValueContent()
-	.VAlign(VAlign_Center)
-	.MaxDesiredWidth(250)
-	[
-		SNew(SButton)
-		.VAlign(VAlign_Center)
-		.ToolTipText(LOCTEXT("ConvertToStaticMeshTooltip", "Create a new StaticMesh asset using current geometry from this DynamicMeshComponent. Does not modify instance."))
-		.OnClicked(this, &FAvaMeshesDetailCustomization::OnConvertToStaticMeshClicked)
-		.IsEnabled(this, &FAvaMeshesDetailCustomization::CanConvertToStaticMesh)
-		.Content()
-		[
-			SNew(STextBlock)
-			.Text(ConvertToStaticMeshText)
-		]
-	];
+		const FText ExportRowText = LOCTEXT("ExportMesh", "Export Mesh");
+		FDetailWidgetRow& ExportRow = ShapeCategoryBuilder.AddCustomRow(ExportRowText, /** Advanced */true);
+
+		ExportRow
+			.NameContent()
+			[
+				SNew(STextBlock)
+				.Text(ExportRowText)
+				.Font(InDetailBuilder.GetDetailFont())
+			]
+			.ValueContent()
+			.VAlign(VAlign_Center)
+			.MaxDesiredWidth(250)
+			[
+				SNew(SButton)
+				.VAlign(VAlign_Center)
+				.ToolTipText(LOCTEXT("ConvertToStaticMeshTooltip", "Create a new StaticMesh asset using current geometry from this DynamicMeshComponent. Does not modify instance."))
+				.OnClicked(this, &FAvaMeshesDetailCustomization::OnConvertToStaticMeshClicked)
+				.IsEnabled(this, &FAvaMeshesDetailCustomization::CanConvertToStaticMesh)
+				.Content()
+				[
+					SNew(STextBlock)
+					.Text(LOCTEXT("ConvertToStaticMesh", "Create Static Mesh"))
+				]
+			];
+	}
 }
 
 FReply FAvaMeshesDetailCustomization::OnConvertToStaticMeshClicked()
@@ -249,15 +203,15 @@ FReply FAvaMeshesDetailCustomization::OnConvertToStaticMeshClicked()
 		return FReply::Handled();
 	}
 
-	TWeakObjectPtr<UAvaShapeDynamicMeshBase> DynMeshBase = SelectedDynamicMeshes[0];
+	UAvaShapeDynamicMeshBase* DynMesh = MeshGeneratorsWeak[0].Get();
 
-	if (!DynMeshBase.IsValid())
+	if (!DynMesh)
 	{
 		return FReply::Handled();
 	}
 
 	// generate name for asset
-	FString NewNameSuggestion = TEXT("SM_MotionDesign") + DynMeshBase->GetMeshName();
+	FString NewNameSuggestion = TEXT("SM_MotionDesign") + DynMesh->GetMeshName();
 	FString PackageName = FString(TEXT("/Game/Meshes/")) + NewNameSuggestion;
 	FString AssetName;
 
@@ -286,7 +240,7 @@ FReply FAvaMeshesDetailCustomization::OnConvertToStaticMeshClicked()
 		MeshName = *AssetName;
 	}
 
-	UE::Geometry::FDynamicMesh3* MeshIn = DynMeshBase->GetShapeMeshComponent()->GetMesh();
+	const UE::Geometry::FDynamicMesh3* MeshIn = DynMesh->GetShapeMeshComponent()->GetMesh();
 
 	// empty mesh do not export
 	if (!MeshIn || MeshIn->TriangleCount() == 0)
@@ -301,7 +255,7 @@ FReply FAvaMeshesDetailCustomization::OnConvertToStaticMeshClicked()
 	// Create StaticMesh object
 	UStaticMesh* StaticMesh = NewObject<UStaticMesh>(Package, MeshName, RF_Public | RF_Standalone);
 
-	if (DynMeshBase->ExportToStaticMesh(StaticMesh))
+	if (DynMesh->ExportToStaticMesh(StaticMesh))
 	{
 		// Notify asset registry of new asset
 		FAssetRegistryModule::AssetCreated(StaticMesh);
@@ -312,7 +266,7 @@ FReply FAvaMeshesDetailCustomization::OnConvertToStaticMeshClicked()
 
 bool FAvaMeshesDetailCustomization::CanConvertToStaticMesh() const
 {
-	return SelectedDynamicMeshes.Num() == 1;
+	return MeshGeneratorsWeak.Num() == 1 && MeshGeneratorsWeak[0].IsValid();
 }
 
 #undef LOCTEXT_NAMESPACE
\ No newline at end of file
diff --git a/Engine/Plugins/Experimental/Avalanche/Source/AvalancheEditor/Private/DetailView/Customizations/AvaMeshesDetailCustomization.h b/Engine/Plugins/Experimental/Avalanche/Source/AvalancheEditor/Private/DetailView/Customizations/AvaMeshesDetailCustomization.h
--- a/Engine/Plugins/Experimental/Avalanche/Source/AvalancheEditor/Private/DetailView/Customizations/AvaMeshesDetailCustomization.h
+++ b/Engine/Plugins/Experimental/Avalanche/Source/AvalancheEditor/Private/DetailView/Customizations/AvaMeshesDetailCustomization.h
@@ -11,7 +11,7 @@ class FReply;
 class UAvaShapeDynamicMeshBase;
 class UDynamicMaterialModel;
 
-/* Used to create the details materials meshes widget and export to StaticMesh */
+/** Used to create the details materials meshes widget and export to StaticMesh */
 class FAvaMeshesDetailCustomization : public IDetailCustomization
 {
 public:
@@ -20,18 +20,16 @@ public:
 		return MakeShared<FAvaMeshesDetailCustomization>();
 	}
 
-	FAvaMeshesDetailCustomization()
-	{
-	}
-
-	virtual void CustomizeDetails(IDetailLayoutBuilder& DetailBuilder) override;
+	//~ Begin IDetailCustomization
+	virtual void CustomizeDetails(IDetailLayoutBuilder& InDetailBuilder) override;
+	//~ End IDetailCustomization
 
 protected:
-	// Handler when the convert button is clicked
+	/** Handler when the convert button is clicked */
 	FReply OnConvertToStaticMeshClicked();
 
-	// Enable or disable the button when selected object is not compatible
+	/** Enable or disable the button when selected object is not compatible */
 	bool CanConvertToStaticMesh() const;
 
-	TArray<TWeakObjectPtr<UAvaShapeDynamicMeshBase>> SelectedDynamicMeshes;
+	TArray<TWeakObjectPtr<UAvaShapeDynamicMeshBase>> MeshGeneratorsWeak;
 };
\ No newline at end of file
diff --git a/Engine/Plugins/Experimental/Avalanche/Source/AvalancheShapes/Public/DynamicMeshes/AvaShapeDynMeshBase.h b/Engine/Plugins/Experimental/Avalanche/Source/AvalancheShapes/Public/DynamicMeshes/AvaShapeDynMeshBase.h
--- a/Engine/Plugins/Experimental/Avalanche/Source/AvalancheShapes/Public/DynamicMeshes/AvaShapeDynMeshBase.h
+++ b/Engine/Plugins/Experimental/Avalanche/Source/AvalancheShapes/Public/DynamicMeshes/AvaShapeDynMeshBase.h
@@ -22,7 +22,7 @@ DECLARE_MULTICAST_DELEGATE_OneParam(FMaskEnabledDelegate, AActor* /** NewMaskAct
 DECLARE_MULTICAST_DELEGATE_OneParam(FMaskDisabledDelegate, AActor* /** OldMaskActor */);
 DECLARE_MULTICAST_DELEGATE_TwoParams(FMaskVisibilityDelegate, const UWorld* /** CurrentWorld */, bool /** bMaskActorVisible */);
 
-UCLASS(MinimalAPI, ClassGroup="Shape", Abstract, BlueprintType, CustomConstructor, EditInlineNew, DefaultToInstanced)
+UCLASS(MinimalAPI, ClassGroup="Shape", Abstract, BlueprintType, CustomConstructor, EditInlineNew, DefaultToInstanced, HideCategories=(Tags, Activation, Cooking, AssetUserData, Navigation))
 class UAvaShapeDynamicMeshBase
 	: public UActorComponent
 	, public IAvaInteractiveToolsModeDetailsObject
@@ -341,35 +341,29 @@ protected:
 	void ScaleVertices(const FVector& InScale);
 	void ScaleVertices(const FVector2D& InScale);
 
-	// Meshes used for the current shape sections
-	UPROPERTY(EditInstanceOnly, BlueprintReadOnly, Category="Material", meta=(EditFixedOrder, EditFixedSize, DisplayAfter="bUsePrimaryMaterialEverywhere", AllowPrivateAccess = "true"))
-	TMap<int32, FAvaShapeMeshData> MeshDatas;
-
-#if WITH_EDITORONLY_DATA
-	/** Stores select settings to allow restoration when Gizmo mode turned off */
-	UPROPERTY()
-	TMap<int32, FAvaShapeMeshData> NonGizmoMeshData;
-#endif
-
 	// enable mesh size property editing
-	UPROPERTY(EditDefaultsOnly, BlueprintReadOnly, Category="Shape", meta=(DisplayPriority=1))
+	UPROPERTY(EditDefaultsOnly, BlueprintReadOnly, Category="Shape")
 	bool bAllowEditSize;
 
 	// the type of size you want to handle
-	UPROPERTY(EditInstanceOnly, BlueprintReadWrite, Category="Shape", Transient, meta=(DisplayName="Size Type", DisplayPriority=1, EditCondition="bAllowEditSize", EditConditionHides, AllowPrivateAccess="true"))
+	UPROPERTY(EditInstanceOnly, BlueprintReadWrite, Category="Shape", Transient, meta=(DisplayName="Size Type", EditCondition="bAllowEditSize", EditConditionHides, AllowPrivateAccess="true"))
 	ESizeType SizeType = ESizeType::UnrealUnit;
 
 	// Uniform scaled size of the mesh
-	UPROPERTY(EditInstanceOnly, BlueprintReadWrite, Transient, Setter, Getter, Category="Shape", meta=(DisplayPriority=1, ClampMin="0.0", DisplayName="Uniform Scaled Size", Units="times", AllowPrivateAccess="true"))
+	UPROPERTY(EditInstanceOnly, BlueprintReadWrite, Transient, Setter, Getter, Category="Shape", meta=(ClampMin="0.0", DisplayName="Uniform Scaled Size", Units="times", AllowPrivateAccess="true"))
 	float UniformScaledSize = 1.f;
 
-	UPROPERTY(BlueprintReadWrite, AdvancedDisplay, Category="Shape", meta=(DisplayPriority=1, DisplayName="Material Vertex Color", AllowPrivateAccess="true"))
+	UPROPERTY(BlueprintReadWrite, AdvancedDisplay, Category="Shape", meta=(DisplayName="Material Vertex Color", AllowPrivateAccess="true"))
 	FLinearColor VertexColor;
 
 	// use primary material for every slot available
-	UPROPERTY(EditInstanceOnly, BlueprintReadWrite, Setter="SetUsePrimaryMaterialEverywhere", Getter="GetUsePrimaryMaterialEverywhere", Category="Material", meta=(DisplayPriority=1, DisplayName="Use Single Material", AllowPrivateAccess="true"))
+	UPROPERTY(EditInstanceOnly, BlueprintReadWrite, Setter="SetUsePrimaryMaterialEverywhere", Getter="GetUsePrimaryMaterialEverywhere", Category="Material", meta=(DisplayName="Use Single Material", AllowPrivateAccess="true"))
 	bool bUsePrimaryMaterialEverywhere;
 
+	// Meshes used for the current shape sections
+	UPROPERTY(EditInstanceOnly, BlueprintReadOnly, Category="Material", meta=(EditFixedOrder, EditFixedSize, AllowPrivateAccess="true"))
+	TMap<int32, FAvaShapeMeshData> MeshDatas;
+
 	TArray<FAvaSnapPoint> LocalSnapPoints;
 
 private:
@@ -417,6 +411,12 @@ private:
 	UPROPERTY()
 	FVector MeshRegenWorldLocation;
 
+#if WITH_EDITORONLY_DATA
+	/** Stores select settings to allow restoration when Gizmo mode turned off */
+	UPROPERTY()
+	TMap<int32, FAvaShapeMeshData> NonGizmoMeshData;
+#endif
+
 	bool bHasNewMeshRegenWorldLocation;
 
 	// flag to register meshes only once
