[Motion Design] Rundown fix: "Take Next" action must not be allowed during transitions. (Trivial)

This is currently a design constraint of the rundown transition logic system: page transitions can't occur concurrently within the same world/channel. This is enforced in the rundown backend for now.

diff --git a/Engine/Plugins/Experimental/Avalanche/Source/AvalancheMedia/Private/Rundown/AvaRundownPlaybackUtils.cpp b/Engine/Plugins/Experimental/Avalanche/Source/AvalancheMedia/Private/Rundown/AvaRundownPlaybackUtils.cpp
--- a/Engine/Plugins/Experimental/Avalanche/Source/AvalancheMedia/Private/Rundown/AvaRundownPlaybackUtils.cpp
+++ b/Engine/Plugins/Experimental/Avalanche/Source/AvalancheMedia/Private/Rundown/AvaRundownPlaybackUtils.cpp
@@ -38,7 +38,7 @@ int32 FAvaRundownPlaybackUtils::GetPageIdToPlayNext(const UAvaRundown* InRundown
 			if (bIsHeadPlaying)
 			{
 				const int32 NextPageId = InRundown->GetNextPage(PlaybackContext->PlayHeadPageId, InPageListReference).GetPageId();
-				if (IsPageIdValid(NextPageId) && InRundown->CanPlayPage(NextPageId, true))
+				if (IsPageIdValid(NextPageId) && InRundown->CanPlayPage(NextPageId, bInPreview))
 				{
 					return NextPageId;
 				}
