Motion Design: Added description for remainder of tasks that didn't have a custom description yet (and needed it).

diff --git a/Engine/Plugins/Experimental/Avalanche/Source/AvalancheSequence/Private/Transition/AvaTransitionContinueSequenceTask.cpp b/Engine/Plugins/Experimental/Avalanche/Source/AvalancheSequence/Private/Transition/AvaTransitionContinueSequenceTask.cpp
--- a/Engine/Plugins/Experimental/Avalanche/Source/AvalancheSequence/Private/Transition/AvaTransitionContinueSequenceTask.cpp
+++ b/Engine/Plugins/Experimental/Avalanche/Source/AvalancheSequence/Private/Transition/AvaTransitionContinueSequenceTask.cpp
@@ -4,6 +4,26 @@
 #include "AvaSequencePlaybackObject.h"
 #include "StateTreeExecutionContext.h"
 
+#define LOCTEXT_NAMESPACE "AvaTransitionContinueSequenceTask"
+
+#if WITH_EDITOR
+FText FAvaTransitionContinueSequenceTask::GetDescription(const FGuid& InId, FStateTreeDataView InInstanceDataView, const IStateTreeBindingLookup& InBindingLookup, EStateTreeNodeFormatting InFormatting) const
+{
+	const FInstanceDataType& InstanceData = InInstanceDataView.Get<FInstanceDataType>();
+
+	TArray<FText> AdditionalArgs;
+	AdditionalArgs.Reserve(1);
+	AdditionalArgs.Add(UEnum::GetDisplayValueAsText(InstanceData.WaitType));
+
+	const FText SequenceQueryText = GetSequenceQueryText(InstanceData, InFormatting);
+	const FText AddOnText = FText::Format(INVTEXT("( {0} )"), FText::Join(FText::FromString(TEXT(" | ")), AdditionalArgs));
+
+	return InFormatting == EStateTreeNodeFormatting::RichText
+		? FText::Format(LOCTEXT("DescRich", "Continue {0} <s>{1}</>"), SequenceQueryText, AddOnText)
+		: FText::Format(LOCTEXT("Desc", "Continue {0} {1}"), SequenceQueryText, AddOnText);
+}
+#endif
+
 TArray<UAvaSequencePlayer*> FAvaTransitionContinueSequenceTask::ExecuteSequenceTask(FStateTreeExecutionContext& InContext) const
 {
 	IAvaSequencePlaybackObject* PlaybackObject = GetPlaybackObject(InContext);
@@ -26,3 +46,5 @@ TArray<UAvaSequencePlayer*> FAvaTransitionContinueSequenceTask::ExecuteSequenceT
 	checkNoEntry();
 	return TArray<UAvaSequencePlayer*>();
 }
+
+#undef LOCTEXT_NAMESPACE
diff --git a/Engine/Plugins/Experimental/Avalanche/Source/AvalancheSequence/Private/Transition/AvaTransitionPlaySequenceTask.cpp b/Engine/Plugins/Experimental/Avalanche/Source/AvalancheSequence/Private/Transition/AvaTransitionPlaySequenceTask.cpp
--- a/Engine/Plugins/Experimental/Avalanche/Source/AvalancheSequence/Private/Transition/AvaTransitionPlaySequenceTask.cpp
+++ b/Engine/Plugins/Experimental/Avalanche/Source/AvalancheSequence/Private/Transition/AvaTransitionPlaySequenceTask.cpp
@@ -42,7 +42,7 @@ FText FAvaTransitionPlaySequenceTask::GetDescription(const FGuid& InId, FStateTr
 
 	return InFormatting == EStateTreeNodeFormatting::RichText
 		? FText::Format(LOCTEXT("DescRich", "Play {0} <s>{1}</>"), SequenceQueryText, AddOnText)
-		: FText::Format(LOCTEXT("DescRich", "Play {0} <s>{1}</>"), SequenceQueryText, AddOnText);
+		: FText::Format(LOCTEXT("Desc", "Play {0} {1}"), SequenceQueryText, AddOnText);
 }
 #endif
 
diff --git a/Engine/Plugins/Experimental/Avalanche/Source/AvalancheSequence/Private/Transition/AvaTransitionWaitForSequenceTask.cpp b/Engine/Plugins/Experimental/Avalanche/Source/AvalancheSequence/Private/Transition/AvaTransitionWaitForSequenceTask.cpp
--- a/Engine/Plugins/Experimental/Avalanche/Source/AvalancheSequence/Private/Transition/AvaTransitionWaitForSequenceTask.cpp
+++ b/Engine/Plugins/Experimental/Avalanche/Source/AvalancheSequence/Private/Transition/AvaTransitionWaitForSequenceTask.cpp
@@ -5,6 +5,21 @@
 #include "AvaSequencePlayer.h"
 #include "StateTreeExecutionContext.h"
 
+#define LOCTEXT_NAMESPACE "AvaTransitionWaitForSequenceTask"
+
+#if WITH_EDITOR
+FText FAvaTransitionWaitForSequenceTask::GetDescription(const FGuid& InId, FStateTreeDataView InInstanceDataView, const IStateTreeBindingLookup& InBindingLookup, EStateTreeNodeFormatting InFormatting) const
+{
+	const FInstanceDataType& InstanceData = InInstanceDataView.Get<FInstanceDataType>();
+
+	const FText SequenceQueryText = GetSequenceQueryText(InstanceData, InFormatting);
+
+	return InFormatting == EStateTreeNodeFormatting::RichText
+		? FText::Format(LOCTEXT("DescRich", "Wait <s>for</> {0}"), SequenceQueryText)
+		: FText::Format(LOCTEXT("Desc", "Wait for {0}"), SequenceQueryText);
+}
+#endif
+
 TArray<UAvaSequencePlayer*> FAvaTransitionWaitForSequenceTask::ExecuteSequenceTask(FStateTreeExecutionContext& InContext) const
 {
 	IAvaSequencePlaybackObject* PlaybackObject = GetPlaybackObject(InContext);
@@ -27,3 +42,5 @@ TArray<UAvaSequencePlayer*> FAvaTransitionWaitForSequenceTask::ExecuteSequenceTa
 	checkNoEntry();
 	return TArray<UAvaSequencePlayer*>();
 }
+
+#undef LOCTEXT_NAMESPACE
diff --git a/Engine/Plugins/Experimental/Avalanche/Source/AvalancheSequence/Public/Transition/AvaTransitionContinueSequenceTask.h b/Engine/Plugins/Experimental/Avalanche/Source/AvalancheSequence/Public/Transition/AvaTransitionContinueSequenceTask.h
--- a/Engine/Plugins/Experimental/Avalanche/Source/AvalancheSequence/Public/Transition/AvaTransitionContinueSequenceTask.h
+++ b/Engine/Plugins/Experimental/Avalanche/Source/AvalancheSequence/Public/Transition/AvaTransitionContinueSequenceTask.h
@@ -10,6 +10,12 @@ struct AVALANCHESEQUENCE_API FAvaTransitionContinueSequenceTask : public FAvaTra
 {
 	GENERATED_BODY()
 
+	//~ Begin FStateTreeNodeBase
+#if WITH_EDITOR
+	virtual FText GetDescription(const FGuid& InId, FStateTreeDataView InInstanceDataView, const IStateTreeBindingLookup& InBindingLookup, EStateTreeNodeFormatting InFormatting) const override;
+#endif
+	//~ End FStateTreeNodeBase
+
 	//~ Begin FAvaTransitionSequenceTaskBase
 	virtual TArray<UAvaSequencePlayer*> ExecuteSequenceTask(FStateTreeExecutionContext& InContext) const override;
 	//~ End FAvaTransitionSequenceTaskBase
diff --git a/Engine/Plugins/Experimental/Avalanche/Source/AvalancheSequence/Public/Transition/AvaTransitionWaitForSequenceTask.h b/Engine/Plugins/Experimental/Avalanche/Source/AvalancheSequence/Public/Transition/AvaTransitionWaitForSequenceTask.h
--- a/Engine/Plugins/Experimental/Avalanche/Source/AvalancheSequence/Public/Transition/AvaTransitionWaitForSequenceTask.h
+++ b/Engine/Plugins/Experimental/Avalanche/Source/AvalancheSequence/Public/Transition/AvaTransitionWaitForSequenceTask.h
@@ -10,6 +10,12 @@ struct AVALANCHESEQUENCE_API FAvaTransitionWaitForSequenceTask : public FAvaTran
 {
 	GENERATED_BODY()
 
+	//~ Begin FStateTreeNodeBase
+#if WITH_EDITOR
+	virtual FText GetDescription(const FGuid& InId, FStateTreeDataView InInstanceDataView, const IStateTreeBindingLookup& InBindingLookup, EStateTreeNodeFormatting InFormatting) const override;
+#endif
+	//~ End FStateTreeNodeBase
+
 	//~ Begin FAvaTransitionSequenceTaskBase
 	virtual TArray<UAvaSequencePlayer*> ExecuteSequenceTask(FStateTreeExecutionContext& InContext) const override;
 	//~ End FAvaTransitionSequenceTaskBase
