Motion Design: Added way for the MD Sequence api to be able to retrieve the Sequence provider for a given level.
+ fixed issue with scene interfaces not being properly registered at editor time.

diff --git a/Engine/Plugins/Experimental/Avalanche/Source/Avalanche/Private/AvaScene.cpp b/Engine/Plugins/Experimental/Avalanche/Source/Avalanche/Private/AvaScene.cpp
--- a/Engine/Plugins/Experimental/Avalanche/Source/Avalanche/Private/AvaScene.cpp
+++ b/Engine/Plugins/Experimental/Avalanche/Source/Avalanche/Private/AvaScene.cpp
@@ -264,33 +264,13 @@ void AAvaScene::RebuildSequenceTree()
 void AAvaScene::PostActorCreated()
 {
 	Super::PostActorCreated();
-
-	if (!HasAnyFlags(EObjectFlags::RF_ClassDefaultObject))
-	{
-		FAvaRemoteControlUtils::RegisterRemoteControlPreset(RemoteControlPreset, /*bInEnsureUniqueId*/ false);
-	}
-
-	// Register AAvaScenes created past Subsystem Initialization
-	if (UAvaSceneSubsystem* SceneSubsystem = FAvaWorldSubsystemUtils::GetWorldSubsystem<UAvaSceneSubsystem>(this))
-	{
-		SceneSubsystem->RegisterSceneInterface(GetLevel(), this);
-	}
+	RegisterObjects();
 }
 
 void AAvaScene::PostLoad()
 {
 	Super::PostLoad();
-
-	if (!HasAnyFlags(EObjectFlags::RF_ClassDefaultObject))
-	{
-		FAvaRemoteControlUtils::RegisterRemoteControlPreset(RemoteControlPreset, /*bInEnsureUniqueId*/ true);
-	}
-
-	// Register AAvaScenes created past Subsystem Initialization
-	if (UAvaSceneSubsystem* SceneSubsystem = FAvaWorldSubsystemUtils::GetWorldSubsystem<UAvaSceneSubsystem>(this))
-	{
-		SceneSubsystem->RegisterSceneInterface(GetLevel(), this);
-	}
+	RegisterObjects();
 
 	if (SceneState)
 	{
@@ -310,29 +290,20 @@ void AAvaScene::PostInitializeComponents()
 void AAvaScene::PostDuplicate(EDuplicateMode::Type InDuplicateMode)
 {
 	Super::PostDuplicate(InDuplicateMode);
-	if (!HasAnyFlags(EObjectFlags::RF_ClassDefaultObject))
-	{
-		FAvaRemoteControlUtils::RegisterRemoteControlPreset(RemoteControlPreset, /*bInEnsureUniqueId*/ true);
-	}
+	RegisterObjects();
 }
 
 void AAvaScene::PostEditImport()
 {
 	Super::PostEditImport();
-	if (!HasAnyFlags(EObjectFlags::RF_ClassDefaultObject))
-	{
-		FAvaRemoteControlUtils::RegisterRemoteControlPreset(RemoteControlPreset, /*bInEnsureUniqueId*/ true);
-	}
+	RegisterObjects();
 }
 
 void AAvaScene::BeginDestroy()
 {
 	Super::BeginDestroy();
 
-	if (!HasAnyFlags(EObjectFlags::RF_ClassDefaultObject))
-	{
-		FAvaRemoteControlUtils::UnregisterRemoteControlPreset(RemoteControlPreset);
-	}
+	UnregisterObjects();
 
 #if WITH_EDITOR
 	FWorldDelegates::OnPreWorldRename.Remove(PreWorldRenameDelegate);
@@ -350,4 +321,36 @@ void AAvaScene::SetStartupCameraName(FName InName)
 }
 #endif
 
+void AAvaScene::RegisterObjects()
+{
+	if (!HasAnyFlags(EObjectFlags::RF_ClassDefaultObject))
+	{
+		FAvaRemoteControlUtils::RegisterRemoteControlPreset(RemoteControlPreset, /*bInEnsureUniqueId*/ true);
+	}
+
+	// Register AAvaScenes created past Subsystem Initialization
+	if (UAvaSceneSubsystem* SceneSubsystem = FAvaWorldSubsystemUtils::GetWorldSubsystem<UAvaSceneSubsystem>(this))
+	{
+		SceneSubsystem->RegisterSceneInterface(GetLevel(), this);
+	}
+
+	if (UAvaSequenceSubsystem* SequenceSubsystem = FAvaWorldSubsystemUtils::GetWorldSubsystem<UAvaSequenceSubsystem>(this))
+	{
+		SequenceSubsystem->RegisterSequenceProvider(GetLevel(), this);
+	}
+}
+
+void AAvaScene::UnregisterObjects()
+{
+	if (!HasAnyFlags(EObjectFlags::RF_ClassDefaultObject))
+	{
+		FAvaRemoteControlUtils::UnregisterRemoteControlPreset(RemoteControlPreset);
+	}
+
+	if (UAvaSequenceSubsystem* SequenceSubsystem = FAvaWorldSubsystemUtils::GetWorldSubsystem<UAvaSequenceSubsystem>(this))
+	{
+		SequenceSubsystem->UnregisterSequenceProvider(GetLevel(), this);
+	}
+}
+
 #undef LOCTEXT_NAMESPACE
diff --git a/Engine/Plugins/Experimental/Avalanche/Source/Avalanche/Private/AvaSceneSubsystem.cpp b/Engine/Plugins/Experimental/Avalanche/Source/Avalanche/Private/AvaSceneSubsystem.cpp
--- a/Engine/Plugins/Experimental/Avalanche/Source/Avalanche/Private/AvaSceneSubsystem.cpp
+++ b/Engine/Plugins/Experimental/Avalanche/Source/Avalanche/Private/AvaSceneSubsystem.cpp
@@ -65,7 +65,7 @@ void UAvaSceneSubsystem::PostInitialize()
 		AAvaScene* SceneActor = nullptr;
 		if (Level && Level->Actors.FindItemByClass(&SceneActor))
 		{
-			RegisterSceneInterface(Level, SceneActor);
+			SceneActor->RegisterObjects();
 		}
 	}
 }
diff --git a/Engine/Plugins/Experimental/Avalanche/Source/Avalanche/Public/AvaScene.h b/Engine/Plugins/Experimental/Avalanche/Source/Avalanche/Public/AvaScene.h
--- a/Engine/Plugins/Experimental/Avalanche/Source/Avalanche/Public/AvaScene.h
+++ b/Engine/Plugins/Experimental/Avalanche/Source/Avalanche/Public/AvaScene.h
@@ -101,6 +101,9 @@ public:
 #endif
 	//~ End IAvaViewportDataProvider
 
+	void RegisterObjects();
+	void UnregisterObjects();
+
 protected:
 	UPROPERTY()
 	FAvaSceneTree SceneTree;
diff --git a/Engine/Plugins/Experimental/Avalanche/Source/AvalancheSequence/Private/AvaSequenceSubsystem.cpp b/Engine/Plugins/Experimental/Avalanche/Source/AvalancheSequence/Private/AvaSequenceSubsystem.cpp
--- a/Engine/Plugins/Experimental/Avalanche/Source/AvalancheSequence/Private/AvaSequenceSubsystem.cpp
+++ b/Engine/Plugins/Experimental/Avalanche/Source/AvalancheSequence/Private/AvaSequenceSubsystem.cpp
@@ -110,6 +110,33 @@ void UAvaSequenceSubsystem::RemovePlaybackObject(IAvaSequencePlaybackObject* InP
 	PlaybackObjects.Remove(InPlaybackObject);
 }
 
+IAvaSequenceProvider* UAvaSequenceSubsystem::FindSequenceProvider(const ULevel* InLevel) const
+{
+	if (const TWeakInterfacePtr<IAvaSequenceProvider>* SequenceProvider = SequenceProviders.Find(InLevel))
+	{
+		return SequenceProvider->Get();
+	}
+	return nullptr;
+}
+
+void UAvaSequenceSubsystem::RegisterSequenceProvider(const ULevel* InLevel, IAvaSequenceProvider* InSequenceProvider)
+{
+	SequenceProviders.FindOrAdd(InLevel) = InSequenceProvider;
+}
+
+bool UAvaSequenceSubsystem::UnregisterSequenceProvider(const ULevel* InLevel, IAvaSequenceProvider* InSequenceProvider)
+{
+	if (TWeakInterfacePtr<IAvaSequenceProvider>* SequenceProvider = SequenceProviders.Find(InLevel))
+	{
+		if (*SequenceProvider == InSequenceProvider)
+		{
+			SequenceProviders.Remove(InLevel);
+			return true;
+		}
+	}
+	return false;
+}
+
 bool UAvaSequenceSubsystem::DoesSupportWorldType(const EWorldType::Type InWorldType) const
 {
 	const bool bDisallowedWorld = UE::AvaSequence::Private::GUnsupportedWorlds.Contains(InWorldType);
diff --git a/Engine/Plugins/Experimental/Avalanche/Source/AvalancheSequence/Public/AvaSequenceSubsystem.h b/Engine/Plugins/Experimental/Avalanche/Source/AvalancheSequence/Public/AvaSequenceSubsystem.h
--- a/Engine/Plugins/Experimental/Avalanche/Source/AvalancheSequence/Public/AvaSequenceSubsystem.h
+++ b/Engine/Plugins/Experimental/Avalanche/Source/AvalancheSequence/Public/AvaSequenceSubsystem.h
@@ -2,9 +2,11 @@
 
 #pragma once
 
-#include "Containers/Set.h"
+#include "Containers/Array.h"
+#include "Containers/Map.h"
 #include "Delegates/IDelegateInstance.h"
 #include "Subsystems/WorldSubsystem.h"
+#include "UObject/ObjectKey.h"
 #include "UObject/WeakInterfacePtr.h"
 #include "AvaSequenceSubsystem.generated.h"
 
@@ -14,33 +16,39 @@ class IAvaSequenceProvider;
 class UAvaSequence;
 class ULevel;
 
-UCLASS()
-class AVALANCHESEQUENCE_API UAvaSequenceSubsystem : public UWorldSubsystem
+UCLASS(MinimalAPI)
+class UAvaSequenceSubsystem : public UWorldSubsystem
 {
 	GENERATED_BODY()
 
 public:
 	UAvaSequenceSubsystem() = default;
 
-	static UAvaSequenceSubsystem* Get(UObject* InPlaybackContext);
+	AVALANCHESEQUENCE_API static UAvaSequenceSubsystem* Get(UObject* InPlaybackContext);
 
-	static TSharedRef<IAvaSequenceController> CreateSequenceController(UAvaSequence& InSequence, IAvaSequencePlaybackObject* InPlaybackObject);
+	AVALANCHESEQUENCE_API static TSharedRef<IAvaSequenceController> CreateSequenceController(UAvaSequence& InSequence, IAvaSequencePlaybackObject* InPlaybackObject);
 
-	IAvaSequencePlaybackObject* FindOrCreatePlaybackObject(ULevel* InLevel, IAvaSequenceProvider& InSequenceProvider);
+	AVALANCHESEQUENCE_API IAvaSequencePlaybackObject* FindOrCreatePlaybackObject(ULevel* InLevel, IAvaSequenceProvider& InSequenceProvider);
 
-	IAvaSequencePlaybackObject* FindPlaybackObject(ULevel* InLevel) const;
+	AVALANCHESEQUENCE_API IAvaSequencePlaybackObject* FindPlaybackObject(ULevel* InLevel) const;
 
-	void AddPlaybackObject(IAvaSequencePlaybackObject* InPlaybackObject);
+	AVALANCHESEQUENCE_API void AddPlaybackObject(IAvaSequencePlaybackObject* InPlaybackObject);
+	AVALANCHESEQUENCE_API void RemovePlaybackObject(IAvaSequencePlaybackObject* InPlaybackObject);
 
-	void RemovePlaybackObject(IAvaSequencePlaybackObject* InPlaybackObject);
+	AVALANCHESEQUENCE_API IAvaSequenceProvider* FindSequenceProvider(const ULevel* InLevel) const;
+
+	AVALANCHESEQUENCE_API void RegisterSequenceProvider(const ULevel* InLevel, IAvaSequenceProvider* InSequenceProvider);
+	AVALANCHESEQUENCE_API bool UnregisterSequenceProvider(const ULevel* InLevel, IAvaSequenceProvider* InSequenceProvider);
 
 protected:
 	//~ Begin UWorldSubsystem
-	virtual bool DoesSupportWorldType(const EWorldType::Type InWorldType) const override;
+	AVALANCHESEQUENCE_API virtual bool DoesSupportWorldType(const EWorldType::Type InWorldType) const override;
 	//~ End UWorldSubsystem
 
 private:
 	bool EnsureLevelIsAppropriate(ULevel*& InLevel) const;
 
 	TArray<TWeakInterfacePtr<IAvaSequencePlaybackObject>> PlaybackObjects;
+
+	TMap<TObjectKey<ULevel>, TWeakInterfacePtr<IAvaSequenceProvider>> SequenceProviders;
 };
