Motion Design Sequence: Changed Sequence Name to be a struct for TL + made type customization out of it

diff --git a/Engine/Plugins/Experimental/Avalanche/Source/AvalancheSequence/Private/AvaSequenceName.cpp b/Engine/Plugins/Experimental/Avalanche/Source/AvalancheSequence/Private/AvaSequenceName.cpp
new file mode 100644
--- /dev/null
+++ b/Engine/Plugins/Experimental/Avalanche/Source/AvalancheSequence/Private/AvaSequenceName.cpp
@@ -0,0 +1,13 @@
+// Copyright Epic Games, Inc. All Rights Reserved.
+
+#include "AvaSequenceName.h"
+
+bool FAvaSequenceName::SerializeFromMismatchedTag(const FPropertyTag& InPropertyTag, FStructuredArchive::FSlot InSlot)
+{
+	if (InPropertyTag.GetType().GetName() == NAME_NameProperty)
+	{
+		InSlot << Name;
+		return true;
+	}
+	return false;
+}
diff --git a/Engine/Plugins/Experimental/Avalanche/Source/AvalancheSequence/Public/AvaSequenceName.h b/Engine/Plugins/Experimental/Avalanche/Source/AvalancheSequence/Public/AvaSequenceName.h
new file mode 100644
--- /dev/null
+++ b/Engine/Plugins/Experimental/Avalanche/Source/AvalancheSequence/Public/AvaSequenceName.h
@@ -0,0 +1,43 @@
+// Copyright Epic Games, Inc. All Rights Reserved.
+
+#pragma once
+
+#include "UObject/NameTypes.h"
+#include "AvaSequenceName.generated.h"
+
+USTRUCT(BlueprintType, DisplayName="Motion Design Sequence Name")
+struct FAvaSequenceName
+{
+	GENERATED_BODY()
+
+	FAvaSequenceName() = default;
+
+	FAvaSequenceName(FName InSequenceName)
+		: Name(InSequenceName)
+	{
+	}
+
+	AVALANCHESEQUENCE_API bool SerializeFromMismatchedTag(const FPropertyTag& InPropertyTag, FStructuredArchive::FSlot InSlot);
+
+	bool IsNone() const
+	{
+		return Name.IsNone();
+	}
+
+	operator FName() const
+	{
+		return Name;
+	}
+
+	UPROPERTY(EditAnywhere, BlueprintReadWrite, Category="Motion Design Sequence")
+	FName Name;
+};
+
+template<>
+struct TStructOpsTypeTraits<FAvaSequenceName> : public TStructOpsTypeTraitsBase2<FAvaSequenceName>
+{
+	enum
+	{
+		WithStructuredSerializeFromMismatchedTag = true,
+	};
+};
diff --git a/Engine/Plugins/Experimental/Avalanche/Source/AvalancheSequence/Public/Transition/AvaTransitionSequenceTaskBase.h b/Engine/Plugins/Experimental/Avalanche/Source/AvalancheSequence/Public/Transition/AvaTransitionSequenceTaskBase.h
--- a/Engine/Plugins/Experimental/Avalanche/Source/AvalancheSequence/Public/Transition/AvaTransitionSequenceTaskBase.h
+++ b/Engine/Plugins/Experimental/Avalanche/Source/AvalancheSequence/Public/Transition/AvaTransitionSequenceTaskBase.h
@@ -2,6 +2,7 @@
 
 #pragma once
 
+#include "AvaSequenceName.h"
 #include "AvaSequenceShared.h"
 #include "AvaTagHandle.h"
 #include "AvaTransitionSequenceEnums.h"
@@ -32,7 +33,7 @@ struct FAvaTransitionSequenceTaskBaseInstanceData : public FAvaTransitionSequenc
 	EAvaTransitionSequenceQueryType QueryType = EAvaTransitionSequenceQueryType::Name;
 
 	UPROPERTY(EditAnywhere, Category="Parameter", meta=(DisplayPriority=1, EditCondition="QueryType==EAvaTransitionSequenceQueryType::Name", EditConditionHides))
-	FName SequenceName;
+	FAvaSequenceName SequenceName;
 
 	UPROPERTY(EditAnywhere, Category="Parameter", meta=(DisplayPriority=1, EditCondition="QueryType==EAvaTransitionSequenceQueryType::Tag", EditConditionHides))
 	FAvaTagHandle SequenceTag;
diff --git a/Engine/Plugins/Experimental/Avalanche/Source/AvalancheSequencer/Private/AvaSequencerModule.cpp b/Engine/Plugins/Experimental/Avalanche/Source/AvalancheSequencer/Private/AvaSequencerModule.cpp
--- a/Engine/Plugins/Experimental/Avalanche/Source/AvalancheSequencer/Private/AvaSequencerModule.cpp
+++ b/Engine/Plugins/Experimental/Avalanche/Source/AvalancheSequencer/Private/AvaSequencerModule.cpp
@@ -3,12 +3,14 @@
 #include "AvaSequencerModule.h"
 #include "AvaSequence.h"
 #include "AvaSequenceEditor.h"
+#include "AvaSequenceName.h"
 #include "AvaSequencer.h"
 #include "AvaSequencerUtils.h"
 #include "Commands/AvaSequencerCommands.h"
 #include "Customization/AvaDisplayRateCustomization.h"
 #include "Customization/AvaMarkSettingCustomization.h"
 #include "Customization/AvaSequenceCustomization.h"
+#include "Customization/AvaSequenceNameCustomization.h"
 #include "Customization/AvaSequenceTimeCustomization.h"
 #include "Director/AvaSequenceDirectorBlueprint.h"
 #include "Director/AvaSequenceDirectorCompiler.h"
@@ -179,6 +181,7 @@ void FAvaSequencerModule::RegisterCustomLayouts()
 	RegisterCustomPropertyTypeLayout<FAvaSequencerDisplayRate, FAvaDisplayRateCustomization>(PropertyModule);
 	RegisterCustomPropertyTypeLayout<FAvaMarkSetting, FAvaMarkSettingCustomization>(PropertyModule);
 	RegisterCustomPropertyTypeLayout<FAvaSequenceTime, FAvaSequenceTimeCustomization>(PropertyModule);
+	RegisterCustomPropertyTypeLayout<FAvaSequenceName, FAvaSequenceNameCustomization>(PropertyModule);
 }
 
 void FAvaSequencerModule::UnregisterCustomLayouts()
diff --git a/Engine/Plugins/Experimental/Avalanche/Source/AvalancheSequencer/Private/Customization/AvaSequenceNameCustomization.cpp b/Engine/Plugins/Experimental/Avalanche/Source/AvalancheSequencer/Private/Customization/AvaSequenceNameCustomization.cpp
new file mode 100644
--- /dev/null
+++ b/Engine/Plugins/Experimental/Avalanche/Source/AvalancheSequencer/Private/Customization/AvaSequenceNameCustomization.cpp
@@ -0,0 +1,32 @@
+// Copyright Epic Games, Inc. All Rights Reserved.
+
+#include "AvaSequenceNameCustomization.h"
+#include "DetailWidgetRow.h"
+#include "Widgets/SAvaSequencePicker.h"
+
+void FAvaSequenceNameCustomization::CustomizeHeader(TSharedRef<IPropertyHandle> InPropertyHandle, FDetailWidgetRow& InHeaderRow, IPropertyTypeCustomizationUtils& InCustomizationUtils)
+{
+	InHeaderRow
+		.NameContent()
+		[
+			InPropertyHandle->CreatePropertyNameWidget()
+		]
+		.ValueContent()
+		[
+			SNew(SHorizontalBox)
+			+ SHorizontalBox::Slot()
+			.VAlign(VAlign_Center)
+			.HAlign(HAlign_Fill)
+			.FillWidth(1.f)
+			[
+				SNew(SAvaSequencePicker, InPropertyHandle)
+			]
+			+ SHorizontalBox::Slot()
+			.VAlign(VAlign_Center)
+			.HAlign(HAlign_Right)
+			.AutoWidth()
+			[
+				InPropertyHandle->CreateDefaultPropertyButtonWidgets()
+			]
+		];
+}
diff --git a/Engine/Plugins/Experimental/Avalanche/Source/AvalancheSequencer/Private/Customization/AvaSequenceNameCustomization.h b/Engine/Plugins/Experimental/Avalanche/Source/AvalancheSequencer/Private/Customization/AvaSequenceNameCustomization.h
new file mode 100644
--- /dev/null
+++ b/Engine/Plugins/Experimental/Avalanche/Source/AvalancheSequencer/Private/Customization/AvaSequenceNameCustomization.h
@@ -0,0 +1,14 @@
+// Copyright Epic Games, Inc. All Rights Reserved.
+
+#pragma once
+
+#include "IPropertyTypeCustomization.h"
+
+class FAvaSequenceNameCustomization : public IPropertyTypeCustomization
+{
+public:
+	//~ Begin IPropertyTypeCustomization
+	virtual void CustomizeHeader(TSharedRef<IPropertyHandle> InPropertyHandle, FDetailWidgetRow& InHeaderRow, IPropertyTypeCustomizationUtils& InCustomizationUtils) override;
+	virtual void CustomizeChildren(TSharedRef<IPropertyHandle> InPropertyHandle, IDetailChildrenBuilder& InChildBuilder, IPropertyTypeCustomizationUtils& InCustomizationUtils) override {}
+	//~ End IPropertyTypeCustomization
+};
diff --git a/Engine/Plugins/Experimental/Avalanche/Source/AvalancheSequencer/Private/Customization/Widgets/SAvaSequencePicker.cpp b/Engine/Plugins/Experimental/Avalanche/Source/AvalancheSequencer/Private/Customization/Widgets/SAvaSequencePicker.cpp
new file mode 100644
--- /dev/null
+++ b/Engine/Plugins/Experimental/Avalanche/Source/AvalancheSequencer/Private/Customization/Widgets/SAvaSequencePicker.cpp
@@ -0,0 +1,191 @@
+// Copyright Epic Games, Inc. All Rights Reserved.
+
+#include "SAvaSequencePicker.h"
+#include "AvaSequenceName.h"
+#include "AvaSequenceSubsystem.h"
+#include "DetailLayoutBuilder.h"
+#include "Engine/Level.h"
+#include "IAvaSequenceProvider.h"
+#include "Widgets/Input/SComboBox.h"
+#include "Widgets/Input/SEditableTextBox.h"
+
+#define LOCTEXT_NAMESPACE "SAvaSequencePicker"
+
+void SAvaSequencePicker::Construct(const FArguments& InArgs, const TSharedRef<IPropertyHandle>& InSequenceNameHandle)
+{
+	NameHandle = InSequenceNameHandle->GetChildHandle(GET_MEMBER_NAME_CHECKED(FAvaSequenceName, Name));
+
+	// Figure out the Level to use
+	{
+		TArray<UObject*> OuterObjects;
+		InSequenceNameHandle->GetOuterObjects(OuterObjects);
+
+		for (UObject* OuterObject : OuterObjects)
+		{
+			if (ULevel* Level = OuterObject->GetTypedOuter<ULevel>())
+			{
+				LevelWeak = Level;
+				break;
+			}
+		}
+	}
+
+	// Editable Text box to deal with the Name Handle
+	TSharedRef<SWidget> Widget = SNew(SEditableTextBox)
+		.Text(this, &SAvaSequencePicker::GetSequenceName)
+		.OnTextCommitted(this, &SAvaSequencePicker::OnSequenceNameCommitted)
+		.Font(IDetailLayoutBuilder::GetDetailFont());
+
+	// if the found level is valid, then use a combo box to display the options found for that level instead
+	if (LevelWeak.IsValid())
+	{
+		RefreshOptions();
+
+		Widget = SAssignNew(ComboBox, SComboBox<TSharedPtr<FAvaSequencePickerOption>>)
+			.OptionsSource(&Options)
+			.InitiallySelectedItem(SelectedOption)
+			.OnGenerateWidget(this, &SAvaSequencePicker::GenerateOptionWidget)
+			.OnComboBoxOpening(this, &SAvaSequencePicker::RefreshOptions)
+			.OnSelectionChanged(this, &SAvaSequencePicker::OnOptionSelectionChanged)
+			[
+				Widget
+			];
+	}
+
+	ChildSlot
+	[
+		Widget
+	];
+}
+
+FText SAvaSequencePicker::GetSequenceName() const
+{
+	FName CurrentSequenceName;
+	const FPropertyAccess::Result Result = NameHandle->GetValue(CurrentSequenceName);
+
+	if (Result == FPropertyAccess::Fail)
+	{
+		return FText::GetEmpty();
+	}
+
+	if (Result == FPropertyAccess::MultipleValues)
+	{
+		return LOCTEXT("MultipleValues", "Multiple Values");
+	}
+
+	return FText::FromName(CurrentSequenceName);
+}
+
+void SAvaSequencePicker::SetSequenceName(FName InSequenceName)
+{
+	if (!NameHandle.IsValid())
+	{
+		return;
+	}
+
+	FName CurrentSequenceName;
+	bool bSucceeded = NameHandle->GetValue(CurrentSequenceName) == FPropertyAccess::Success;
+
+	// Only skip if successfully retrieved the current value and it matches the new value.
+	// Else, attempt to set the new value in the handle
+	if (bSucceeded && CurrentSequenceName == InSequenceName)
+	{
+		return;
+	}
+
+	NameHandle->SetValue(InSequenceName, EPropertyValueSetFlags::DefaultFlags);
+}
+
+void SAvaSequencePicker::OnSequenceNameCommitted(const FText& InText, ETextCommit::Type InCommitType)
+{
+	SetSequenceName(*InText.ToString());
+}
+
+void SAvaSequencePicker::RefreshOptions()
+{
+	Options.Reset();
+	SelectedOption.Reset();
+
+	if (!NameHandle.IsValid())
+	{
+		return;	
+	}
+
+	ULevel* Level = LevelWeak.Get();
+	if (!Level)
+	{
+		return;
+	}
+
+	UAvaSequenceSubsystem* SequenceSubsystem = UAvaSequenceSubsystem::Get(Level);
+	if (!SequenceSubsystem)
+	{
+		return;
+	}
+
+	IAvaSequenceProvider* SequenceProvider = SequenceSubsystem->FindSequenceProvider(Level);
+	if (!SequenceProvider)
+	{
+		return;
+	}
+
+	FName CurrentSequenceName;
+	NameHandle->GetValue(CurrentSequenceName);
+
+	TConstArrayView<UAvaSequence*> Sequences = SequenceProvider->GetSequences();
+	Options.Reserve(Sequences.Num());
+
+	TSet<FName> ProcessedNames;
+	ProcessedNames.Reserve(Sequences.Num());
+
+	for (UAvaSequence* Sequence : Sequences)
+	{
+		if (Sequence)
+		{
+			const FName SequenceName = Sequence->GetLabel();
+
+			// Since FAvaSequencePickerOption only handles names directly, avoid adding the same name entry twice
+			bool bAlreadyInSet;
+			ProcessedNames.Add(SequenceName, &bAlreadyInSet);
+			if (bAlreadyInSet)
+			{
+				continue;
+			}
+
+			TSharedRef<FAvaSequencePickerOption> Item = MakeShared<FAvaSequencePickerOption>();
+			Item->SequenceName = Sequence->GetLabel();
+
+			// Set the first Item that matches the Sequence Name as the selected item
+			if (!SelectedOption.IsValid() && Item->SequenceName == CurrentSequenceName)
+			{
+				SelectedOption = Item;
+			}
+
+			Options.Emplace(MoveTemp(Item));
+		}
+	}
+
+	// Update the Selected Item to the selected option (can be null if there's no sequence found matching the current name stored)
+	if (ComboBox.IsValid())
+	{
+		ComboBox->SetSelectedItem(SelectedOption);
+	}
+}
+
+void SAvaSequencePicker::OnOptionSelectionChanged(TSharedPtr<FAvaSequencePickerOption> InSelectedOption, ESelectInfo::Type InSelectInfo)
+{
+	if (InSelectedOption.IsValid())
+	{
+		SelectedOption = InSelectedOption;
+		SetSequenceName(SelectedOption->SequenceName);
+	}
+}
+
+TSharedRef<SWidget> SAvaSequencePicker::GenerateOptionWidget(TSharedPtr<FAvaSequencePickerOption> InOption)
+{
+	return SNew(STextBlock)
+		.Text(FText::FromName(InOption->SequenceName))
+		.Font(IDetailLayoutBuilder::GetDetailFont());
+}
+
+#undef LOCTEXT_NAMESPACE
diff --git a/Engine/Plugins/Experimental/Avalanche/Source/AvalancheSequencer/Private/Customization/Widgets/SAvaSequencePicker.h b/Engine/Plugins/Experimental/Avalanche/Source/AvalancheSequencer/Private/Customization/Widgets/SAvaSequencePicker.h
new file mode 100644
--- /dev/null
+++ b/Engine/Plugins/Experimental/Avalanche/Source/AvalancheSequencer/Private/Customization/Widgets/SAvaSequencePicker.h
@@ -0,0 +1,49 @@
+// Copyright Epic Games, Inc. All Rights Reserved.
+
+#pragma once
+
+#include "Containers/Array.h"
+#include "Widgets/SCompoundWidget.h"
+
+class IPropertyHandle;
+class UAvaSequence;
+class ULevel;
+template<typename OptionType> class SComboBox;
+
+struct FAvaSequencePickerOption
+{
+	/** Sequence name converted to Text */
+	FName SequenceName;
+};
+
+class SAvaSequencePicker : public SCompoundWidget
+{
+public:
+	SLATE_BEGIN_ARGS(SAvaSequencePicker) {}
+	SLATE_END_ARGS()
+
+	void Construct(const FArguments& InArgs, const TSharedRef<IPropertyHandle>& InSequenceNameHandle);
+
+private:
+	FText GetSequenceName() const;
+
+	void SetSequenceName(FName InSequenceName);
+
+	void OnSequenceNameCommitted(const FText& InText, ETextCommit::Type InCommitType);
+
+	void RefreshOptions();
+
+	void OnOptionSelectionChanged(TSharedPtr<FAvaSequencePickerOption> InSelectedOption, ESelectInfo::Type InSelectInfo);
+
+	TSharedRef<SWidget> GenerateOptionWidget(TSharedPtr<FAvaSequencePickerOption> InOption);
+
+	TSharedPtr<SComboBox<TSharedPtr<FAvaSequencePickerOption>>> ComboBox;
+
+	TWeakObjectPtr<ULevel> LevelWeak;
+
+	TSharedPtr<IPropertyHandle> NameHandle;
+
+	TArray<TSharedPtr<FAvaSequencePickerOption>> Options;
+
+	TSharedPtr<FAvaSequencePickerOption> SelectedOption;
+};
