[Motion Design] Display Cluster synchronization support for Level Streaming Playable.

* New Modular Feature for Synchronized Events. Simple interface to push events.

Synchronization for either clustering or forking becomes a first class citizen in the playable framework.
The current implementation uses the Display Cluster Events and is in a new plugin AvaDisplayCluster. More extension points will likely be added and implemented in that plugin.

There is also a "no sync" implementation in AvalancheMedia Module when running in non-clustered. I could eventually have another implementation for "forked" channels not based on nDisplay.

The implementation is selectable in Avalanche settings.

* Testing methodology:

The playback server a new test feature that will randomly delay the playback commands to cause desynchronization. It is enable with a cvar so it is possible to configure it from the command line or type it in the console.

Tracking the issues and desyncs required a lot of logging. I formatted the logs to be used with a log highlighter for notepad++ to make it easier to read. Combined with log replication, I could get a all on one log of all events on both client and servers in a good chronological order so I could track down when a desync occurs and what caused it.

* Playback Client Issues with desynchronized clustered/forked channels

As a result of these tests, I found a lot of weaknesses in the playback client side of things when dealing with "forked" channels. It required rewrites in remote transition and remote playable proxies to add tracking and reconcile of different states that can be desynchronized. The principle is, if it works when desynchronized, it will work when synchronized. But it has to work even if desynchronized. The client also has to be able to recover it's state, from the server(s), if it is shutdown and restarted and be able to continue operating the rundown without issues.

* Scene View Extension to Hide "Waiting" visible levels

Playable Groups have a scene view extension allowing playables to setup the view.
For the Level Streaming Playable, it is used to hide the level once it has become visible by the level streaming update, but should not yet be shown because it is waiting for other nodes in the cluster that are not yet ready.

diff --git a/Engine/Plugins/Experimental/Avalanche/Source/Avalanche/Public/Framework/AvaInstanceSettings.h b/Engine/Plugins/Experimental/Avalanche/Source/Avalanche/Public/Framework/AvaInstanceSettings.h
--- a/Engine/Plugins/Experimental/Avalanche/Source/Avalanche/Public/Framework/AvaInstanceSettings.h
+++ b/Engine/Plugins/Experimental/Avalanche/Source/Avalanche/Public/Framework/AvaInstanceSettings.h
@@ -4,6 +4,19 @@
 
 #include "AvaInstanceSettings.generated.h"
 
+USTRUCT()
+struct FAvaSynchronizedEventsFeatureSelection
+{
+	GENERATED_BODY()
+	
+	/**
+	 * Select the implementation for synchronizing events.
+	 * "Default" will select the most appropriate implementation available.
+	 */
+	UPROPERTY(Config, EditAnywhere, Category = Settings)
+	FString Implementation;
+};
+
 /**
  * Settings applied when instancing a Motion Design Asset for playback.
  */
@@ -22,5 +35,10 @@ struct FAvaInstanceSettings
 	 */
 	UPROPERTY(config, EditAnywhere, Category = Settings)
 	bool bDefaultPlayableTransitionWaitForSequences = false;
-	
+
+	/**
+	 * Select the implementation for synchronizing events.
+	 */
+	UPROPERTY(Config, EditAnywhere, Category = Settings)
+	FAvaSynchronizedEventsFeatureSelection SynchronizedEventsFeature;
 };
diff --git a/Engine/Plugins/Experimental/Avalanche/Source/AvalancheMedia/Internal/ModularFeature/AvaMediaSynchronizedEventsFeature.h b/Engine/Plugins/Experimental/Avalanche/Source/AvalancheMedia/Internal/ModularFeature/AvaMediaSynchronizedEventsFeature.h
new file mode 100644
--- /dev/null
+++ b/Engine/Plugins/Experimental/Avalanche/Source/AvalancheMedia/Internal/ModularFeature/AvaMediaSynchronizedEventsFeature.h
@@ -0,0 +1,37 @@
+// Copyright Epic Games, Inc. All Rights Reserved.
+
+#pragma once
+
+#include "ModularFeature/IAvaMediaSynchronizedEventsFeature.h"
+
+class FAvaMediaSynchronizedEventsFeature
+{
+public:
+	static void Startup();
+	
+	static void Shutdown();
+	
+	/**
+	 * Returns currently selected implementation.
+	 * This function always returns a valid implementation.
+	 */
+	static AVALANCHEMEDIA_API IAvaMediaSynchronizedEventsFeature* Get();
+
+	/**
+	 * Creates a new event dispatcher from the currently selected implementation. 
+	 */ 
+	static AVALANCHEMEDIA_API TSharedPtr<IAvaMediaSynchronizedEventDispatcher> CreateDispatcher(const FString& InSignature);
+
+	static AVALANCHEMEDIA_API void EnumerateImplementations(TFunctionRef<void(const IAvaMediaSynchronizedEventsFeature*)> InCallback);
+
+	/**
+	 * Find the implementation with the given name.
+	 * The name can be "default", in which case it will automatically select the highest priority implementation.
+	 * If no implementation match the given name, it will return an internal "no sync" implementation.
+	 * This function always returns a valid implementation.
+	 */
+	static AVALANCHEMEDIA_API IAvaMediaSynchronizedEventsFeature* FindImplementation(FName InImplementation);
+
+private:
+	static IAvaMediaSynchronizedEventsFeature* GetInternalImplementation();
+};
\ No newline at end of file
diff --git a/Engine/Plugins/Experimental/Avalanche/Source/AvalancheMedia/Internal/Playback/AvaPlaybackUtils.h b/Engine/Plugins/Experimental/Avalanche/Source/AvalancheMedia/Internal/Playback/AvaPlaybackUtils.h
--- a/Engine/Plugins/Experimental/Avalanche/Source/AvalancheMedia/Internal/Playback/AvaPlaybackUtils.h
+++ b/Engine/Plugins/Experimental/Avalanche/Source/AvalancheMedia/Internal/Playback/AvaPlaybackUtils.h
@@ -3,6 +3,7 @@
 #pragma once
 
 #include "Containers/Array.h"
+#include "Containers/UnrealString.h"
 
 class UPackage;
 struct FAssetData;
@@ -57,4 +58,10 @@ namespace UE::AvaPlayback::Utils
 	 *	This is to used for logging and tracing.
 	 */
 	AVALANCHEMEDIA_API FString GetBriefFrameInfo();
+
+	template<typename InEnumType>
+	FString StaticEnumToString(InEnumType InValue)
+	{
+		return StaticEnum<InEnumType>()->GetNameStringByValue(static_cast<int64>(InValue));
+	}
 }
\ No newline at end of file
diff --git a/Engine/Plugins/Experimental/Avalanche/Source/AvalancheMedia/Private/AvaMediaModule.cpp b/Engine/Plugins/Experimental/Avalanche/Source/AvalancheMedia/Private/AvaMediaModule.cpp
--- a/Engine/Plugins/Experimental/Avalanche/Source/AvalancheMedia/Private/AvaMediaModule.cpp
+++ b/Engine/Plugins/Experimental/Avalanche/Source/AvalancheMedia/Private/AvaMediaModule.cpp
@@ -11,8 +11,11 @@
 #include "Interfaces/IPluginManager.h"
 #include "Misc/CommandLine.h"
 #include "Misc/CoreDelegates.h"
+#include "ModularFeature/AvaMediaSynchronizedEventsFeature.h"
+#include "Playable/AvaPlayableGroupSceneViewExtension.h"
 #include "Playback/AvaPlaybackClientDelegates.h"
 #include "Playback/AvaPlaybackClientDummy.h"
+#include "SceneViewExtension.h"
 #include "ShaderCore.h"
 
 #if WITH_EDITOR
@@ -57,6 +60,8 @@ void FAvaMediaModule::StartupModule()
 	const FString PluginShaderDir = FPaths::Combine(Plugin->GetBaseDir(), TEXT("Shaders"));
 	AddShaderSourceDirectoryMapping(UE::AvaBroadcastRenderTargetMediaUtils::VirtualShaderMountPoint, PluginShaderDir);
 
+	FAvaMediaSynchronizedEventsFeature::Startup();
+
 	IMediaIOCoreModule::Get().RegisterDeviceProvider(&AvaDisplayDeviceProvider);
 
 	ConsoleCmds.Add(IConsoleManager::Get().RegisterConsoleCommand(
@@ -205,6 +210,7 @@ void FAvaMediaModule::ShutdownModule()
 	ConsoleCmds.Empty();
 
 	AvaMediaSync.Reset();
+	FAvaMediaSynchronizedEventsFeature::Shutdown();
 }
 
 void FAvaMediaModule::StartPlaybackClient()
@@ -436,6 +442,8 @@ void FAvaMediaModule::PostEngineInit()
 {
 	using namespace UE::AvaMediaModule::Private;
 
+	PlayableGroupSceneViewExtension = FSceneViewExtensions::NewExtension<FAvaPlayableGroupSceneViewExtension>();
+	
 	ConditionalCreateLocalPlaybackManager();
 	ConditionalCreateManagedInstanceCache();
 
@@ -522,6 +530,8 @@ void FAvaMediaModule::PostEngineInit()
 void FAvaMediaModule::EnginePreExit()
 {
 	StopAllServices();
+
+	PlayableGroupSceneViewExtension.Reset();
 }
 
 void FAvaMediaModule::PrePIEEnded(const bool)
diff --git a/Engine/Plugins/Experimental/Avalanche/Source/AvalancheMedia/Private/AvaMediaModule.h b/Engine/Plugins/Experimental/Avalanche/Source/AvalancheMedia/Private/AvaMediaModule.h
--- a/Engine/Plugins/Experimental/Avalanche/Source/AvalancheMedia/Private/AvaMediaModule.h
+++ b/Engine/Plugins/Experimental/Avalanche/Source/AvalancheMedia/Private/AvaMediaModule.h
@@ -16,6 +16,8 @@
 #include "Rundown/AvaRundownManagedInstanceCache.h"
 #include "Rundown/AvaRundownServer.h"
 
+class FAvaPlayableGroupSceneViewExtension;
+
 class FAvaMediaModule : public IAvaMediaModule
 {
 public:
@@ -104,6 +106,8 @@ private:
 	TUniquePtr<FAvaMediaSync> AvaMediaSync;
 
 	TArray<IConsoleObject*> ConsoleCmds;
+
+	TSharedPtr<FAvaPlayableGroupSceneViewExtension> PlayableGroupSceneViewExtension;
 	
 	/**
 	 *	Wraps the local default UAvaMediaSettings.
diff --git a/Engine/Plugins/Experimental/Avalanche/Source/AvalancheMedia/Private/AvaMediaSettings.cpp b/Engine/Plugins/Experimental/Avalanche/Source/AvalancheMedia/Private/AvaMediaSettings.cpp
--- a/Engine/Plugins/Experimental/Avalanche/Source/AvalancheMedia/Private/AvaMediaSettings.cpp
+++ b/Engine/Plugins/Experimental/Avalanche/Source/AvalancheMedia/Private/AvaMediaSettings.cpp
@@ -6,10 +6,13 @@
 
 #define LOCTEXT_NAMESPACE "AvaMediaSettings"
 
+const FName UAvaMediaSettings::SynchronizedEventsFeatureSelection_Default(TEXT("Default"));
+
 UAvaMediaSettings::UAvaMediaSettings()
 {
 	CategoryName = TEXT("Motion Design");
 	SectionName  = TEXT("Playback & Broadcast");
+	AvaInstanceSettings.SynchronizedEventsFeature.Implementation = SynchronizedEventsFeatureSelection_Default.ToString();
 }
 
 UAvaMediaSettings* UAvaMediaSettings::GetSingletonInstance()
diff --git a/Engine/Plugins/Experimental/Avalanche/Source/AvalancheMedia/Private/ModularFeature/AvaMediaSynchronizedEventsFeature.cpp b/Engine/Plugins/Experimental/Avalanche/Source/AvalancheMedia/Private/ModularFeature/AvaMediaSynchronizedEventsFeature.cpp
new file mode 100644
--- /dev/null
+++ b/Engine/Plugins/Experimental/Avalanche/Source/AvalancheMedia/Private/ModularFeature/AvaMediaSynchronizedEventsFeature.cpp
@@ -0,0 +1,203 @@
+// Copyright Epic Games, Inc. All Rights Reserved.
+
+#include "ModularFeature/AvaMediaSynchronizedEventsFeature.h"
+
+#include "AvaMediaSettings.h"
+#include "Features/IModularFeatures.h"
+#include "Framework/AvaInstanceSettings.h"
+#include "HAL/IConsoleManager.h"
+#include "IAvaMediaModule.h"
+#include "ModularFeature/AvaMediaSynchronizedEvent.h"
+#include "ModularFeature/IAvaMediaSynchronizedEventDispatcher.h"
+
+#define LOCTEXT_NAMESPACE "AvaMediaSynchronizedEventsFeature"
+
+namespace UE::AvaMedia::SynchronizedEvents::Private
+{
+	// Allow events to be dispatched as early as possible (at most 1 frame earlier).
+	TAutoConsoleVariable<bool> CVarSyncEarlyDispatch(
+		TEXT("AvaMediaSynchronizedEvent.NoSync.EarlyDispatch")
+		, false
+		, TEXT("If true, will dispatch events as soon as ready. if false, ready events are all batched on the next tick."), ECVF_Cheat);
+}
+
+/**
+ * Provide a default no-sync implementation that fires the queued events on the next tick (*).
+ * 
+ * (*) Note: could fire the event either immediately, i.e. when it is made, or on the next dispatch update.
+ * Sync events should tolerate to not be executed immediately (not all events are like that). So deferring the
+ * invoke to next Tick helps identify such issues. Only events that the invoke can be deferred can be used with
+ * this synchronisation mechanism.
+ */
+class FAvaMediaSynchronizedEventDispatcher : public IAvaMediaSynchronizedEventDispatcher
+{
+public:
+	virtual bool PushEvent(FString&& InEventSignature, TUniqueFunction<void()> InFunction) override
+	{
+		check(IsInGameThread());
+		
+		if (!InFunction)
+		{
+			return false;
+		}
+		
+		if (UE::AvaMedia::SynchronizedEvents::Private::CVarSyncEarlyDispatch.GetValueOnAnyThread())
+		{
+			InFunction();
+		}
+		else
+		{
+			Events.Add(MakeUnique<FAvaMediaSynchronizedEvent>(MoveTemp(InEventSignature), MoveTemp(InFunction)));	
+		}
+		return true;
+	}
+
+	virtual EAvaMediaSynchronizedEventState GetEventState(const FString& InEventSignature) const override
+	{
+		return EAvaMediaSynchronizedEventState::NotFound;
+	}
+
+	virtual void DispatchEvents() override
+	{
+		check(IsInGameThread());
+		
+		for (const TUniquePtr<FAvaMediaSynchronizedEvent>& Event : Events)
+		{
+			Event->Function();
+		}
+		Events.Reset();
+	}
+
+	TArray<TUniquePtr<FAvaMediaSynchronizedEvent>> Events;
+};
+
+class FAvaMediaSynchronizedEventsNoSync : public IAvaMediaSynchronizedEventsFeature
+{
+	virtual FName GetName() const override
+	{
+		static const FName ImplementationName(TEXT("NoSync"));
+		return ImplementationName;
+	}
+
+	virtual FText GetDisplayName() const override
+	{
+		static const FText DisplayName = LOCTEXT("NoSyncDisplayName", "No Sync");
+		return DisplayName;
+	}
+
+	virtual FText GetDisplayDescription() const override
+	{
+		static const FText DisplayDescription = LOCTEXT("NoSyncDisplayDescription",
+			 "This implementation does not perform synchronisation. "
+			 "Events are either executed immedially or on the next frame depending on the configuration.");
+		return DisplayDescription;
+	}
+
+	virtual int32 GetPriority() const override
+	{
+		return 0; // So any other implementation will have priority.
+	}
+	
+	virtual TSharedPtr<IAvaMediaSynchronizedEventDispatcher> CreateDispatcher(const FString& InSignature) override
+	{
+		return MakeShared<FAvaMediaSynchronizedEventDispatcher>();
+	}
+};
+
+void FAvaMediaSynchronizedEventsFeature::Startup()
+{
+	IModularFeatures::Get().RegisterModularFeature(IAvaMediaSynchronizedEventsFeature::GetModularFeatureName(), GetInternalImplementation());
+}
+
+void FAvaMediaSynchronizedEventsFeature::Shutdown()
+{
+	IModularFeatures::Get().UnregisterModularFeature(IAvaMediaSynchronizedEventsFeature::GetModularFeatureName(), GetInternalImplementation());	
+}
+
+IAvaMediaSynchronizedEventsFeature* FAvaMediaSynchronizedEventsFeature::Get()
+{
+	const FName SelectedImplementation(IAvaMediaModule::Get().GetAvaInstanceSettings().SynchronizedEventsFeature.Implementation);
+	
+	IModularFeatures::FScopedLockModularFeatureList ScopedLockModularFeatureList;
+	return FindImplementation(SelectedImplementation);
+}
+
+TSharedPtr<IAvaMediaSynchronizedEventDispatcher> FAvaMediaSynchronizedEventsFeature::CreateDispatcher(const FString& InSignature)
+{
+	if (IAvaMediaSynchronizedEventsFeature* Feature = Get())
+	{
+		return Feature->CreateDispatcher(InSignature);
+	}
+
+	// Fallback returns the no-sync dispatcher implementation directly.
+	return MakeShared<FAvaMediaSynchronizedEventDispatcher>();
+}
+
+void FAvaMediaSynchronizedEventsFeature::EnumerateImplementations(TFunctionRef<void(const IAvaMediaSynchronizedEventsFeature*)> InCallback)
+{
+	IModularFeatures& ModularFeatures = IModularFeatures::Get();
+	static const FName FeatureName = IAvaMediaSynchronizedEventsFeature::GetModularFeatureName();
+
+	const int32 NumImplementations = ModularFeatures.GetModularFeatureImplementationCount(FeatureName);
+	for (int32 ImplementationIndex = 0; ImplementationIndex < NumImplementations; ++ImplementationIndex)
+	{
+		const IAvaMediaSynchronizedEventsFeature* Implementation
+			= static_cast<IAvaMediaSynchronizedEventsFeature*>(ModularFeatures.GetModularFeatureImplementation(FeatureName, ImplementationIndex));
+		InCallback(Implementation);
+	}
+}
+
+IAvaMediaSynchronizedEventsFeature* FAvaMediaSynchronizedEventsFeature::FindImplementation(FName InImplementation)
+{
+	IModularFeatures& ModularFeatures = IModularFeatures::Get();
+	static const FName FeatureName = IAvaMediaSynchronizedEventsFeature::GetModularFeatureName();
+
+	const int32 NumImplementations = ModularFeatures.GetModularFeatureImplementationCount(FeatureName);
+	
+	if (InImplementation == UAvaMediaSettings::SynchronizedEventsFeatureSelection_Default)
+	{
+		// Automatic selection of highest priority implementation.
+		IAvaMediaSynchronizedEventsFeature* SelectedImplementation = nullptr;
+		int SelectedPriority = 0;
+		
+		for (int32 ImplementationIndex = 0; ImplementationIndex < NumImplementations; ++ImplementationIndex)
+		{
+			IAvaMediaSynchronizedEventsFeature* Implementation
+				= static_cast<IAvaMediaSynchronizedEventsFeature*>(ModularFeatures.GetModularFeatureImplementation(FeatureName, ImplementationIndex));
+			if (Implementation && (!SelectedImplementation || Implementation->GetPriority() > SelectedPriority))
+			{
+				SelectedPriority = Implementation->GetPriority();
+				SelectedImplementation = Implementation;
+			}
+		}
+
+		if (SelectedImplementation)
+		{
+			return SelectedImplementation;
+		}
+	}
+	else
+	{
+		// Selection by explicit name
+		for (int32 ImplementationIndex = 0; ImplementationIndex < NumImplementations; ++ImplementationIndex)
+		{
+			IAvaMediaSynchronizedEventsFeature* Implementation
+				= static_cast<IAvaMediaSynchronizedEventsFeature*>(ModularFeatures.GetModularFeatureImplementation(FeatureName, ImplementationIndex));
+			if (Implementation && Implementation->GetName() == InImplementation)
+			{
+				return Implementation;
+			}
+		}
+	}
+
+	// Fallback to the no-sync implementation.
+	return GetInternalImplementation();
+}
+
+IAvaMediaSynchronizedEventsFeature* FAvaMediaSynchronizedEventsFeature::GetInternalImplementation()
+{
+	static FAvaMediaSynchronizedEventsNoSync SynchronizedEventsNoSync;
+	return &SynchronizedEventsNoSync;
+}
+
+#undef LOCTEXT_NAMESPACE
\ No newline at end of file
diff --git a/Engine/Plugins/Experimental/Avalanche/Source/AvalancheMedia/Private/Playable/AvaPlayable.cpp b/Engine/Plugins/Experimental/Avalanche/Source/AvalancheMedia/Private/Playable/AvaPlayable.cpp
--- a/Engine/Plugins/Experimental/Avalanche/Source/AvalancheMedia/Private/Playable/AvaPlayable.cpp
+++ b/Engine/Plugins/Experimental/Avalanche/Source/AvalancheMedia/Private/Playable/AvaPlayable.cpp
@@ -54,8 +54,8 @@ namespace UE::AvaPlayable::Private
 		{
 			return FString::Printf(TEXT("Id:%s, Asset:%s, Status:%s"),
 				*InPlayable->GetInstanceId().ToString(),
-				*InPlayable->GetSourceAssetPath().ToString(),
-				*StaticEnum<EAvaPlayableStatus>()->GetNameByValue(static_cast<int32>(InPlayable->GetPlayableStatus())).ToString());
+				*InPlayable->GetSourceAssetPath().GetAssetName(),
+				*AvaPlayback::Utils::StaticEnumToString(InPlayable->GetPlayableStatus()));
 		}
 		return TEXT("(nullptr)");
 	}
@@ -79,7 +79,7 @@ namespace UE::AvaPlayable::Private
 	FString GetPrettySequenceCommandInfo(EAvaPlaybackAnimAction InAnimAction, const FAvaPlaybackAnimPlaySettings& InAnimPlaySettings)
 	{
 		return FString::Printf(TEXT("Action:%s, Name:%s"),
-			*StaticEnum<EAvaPlaybackAnimAction>()->GetNameByValue(static_cast<int32>(InAnimAction)).ToString(),
+			*AvaPlayback::Utils::StaticEnumToString(InAnimAction),
 			*InAnimPlaySettings.AnimationName.ToString());
 	}
 }
diff --git a/Engine/Plugins/Experimental/Avalanche/Source/AvalancheMedia/Private/Playable/AvaPlayableGroup.cpp b/Engine/Plugins/Experimental/Avalanche/Source/AvalancheMedia/Private/Playable/AvaPlayableGroup.cpp
--- a/Engine/Plugins/Experimental/Avalanche/Source/AvalancheMedia/Private/Playable/AvaPlayableGroup.cpp
+++ b/Engine/Plugins/Experimental/Avalanche/Source/AvalancheMedia/Private/Playable/AvaPlayableGroup.cpp
@@ -205,6 +205,29 @@ bool UAvaPlayableGroup::HasTransitions() const
 	return !PlayableTransitions.IsEmpty();
 }
 
+void UAvaPlayableGroup::PushSynchronizedEvent(FString&& InEventSignature, TUniqueFunction<void()> InFunction)
+{
+	if (UAvaPlayableGroupManager* Manager = GetPlayableGroupManager())
+	{
+		// Using one dispatcher for now, if the event signature needs to be scoped per playable group
+		// we could have a dispatcher for each playable group with a unique signature.
+		Manager->PushSynchronizedEvent(MoveTemp(InEventSignature), MoveTemp(InFunction));
+	}
+	else if (InFunction)
+	{
+		InFunction();
+	}
+}
+
+bool UAvaPlayableGroup::IsSynchronizedEventPushed(const FString& InEventSignature) const
+{
+	if (const UAvaPlayableGroupManager* Manager = GetPlayableGroupManager())
+	{
+		return Manager->IsSynchronizedEventPushed(InEventSignature);
+	}
+	return false;	
+}
+
 void UAvaPlayableGroup::SetLastAppliedCameraPlayable(UAvaPlayable* InPlayable)
 {
 	LastAppliedCameraPlayableWeak = InPlayable;
@@ -320,6 +343,19 @@ void UAvaPlayableGroup::RequestSetVisibility(UAvaPlayable* InPlayable, bool bInS
 	}
 }
 
+void UAvaPlayableGroup::SetupView(FSceneViewFamily& InViewFamily, FSceneView& InView)
+{
+	for (const TObjectKey<UAvaPlayable>& PlayableKey : Playables)
+	{
+		UAvaPlayable* Playable = PlayableKey.ResolveObjectPtr();
+		
+		if (Playable && Playable->IsPlaying())
+		{
+			Playable->SetupView(InViewFamily, InView);
+		}
+	}
+}
+
 bool UAvaPlayableGroup::IsVisibilityConstrained(const UAvaPlayable* InPlayable) const
 {
 	for (const TWeakInterfacePtr<IAvaPlayableVisibilityConstraint>& ConstraintWeak : VisibilityConstraints)
diff --git a/Engine/Plugins/Experimental/Avalanche/Source/AvalancheMedia/Private/Playable/AvaPlayableGroupManager.cpp b/Engine/Plugins/Experimental/Avalanche/Source/AvalancheMedia/Private/Playable/AvaPlayableGroupManager.cpp
--- a/Engine/Plugins/Experimental/Avalanche/Source/AvalancheMedia/Private/Playable/AvaPlayableGroupManager.cpp
+++ b/Engine/Plugins/Experimental/Avalanche/Source/AvalancheMedia/Private/Playable/AvaPlayableGroupManager.cpp
@@ -9,6 +9,8 @@
 #include "Framework/AvaGameInstance.h"
 #include "Misc/CoreDelegates.h"
 #include "Misc/TimeGuard.h"
+#include "ModularFeature/AvaMediaSynchronizedEventsFeature.h"
+#include "ModularFeature/IAvaMediaSynchronizedEventDispatcher.h"
 #include "Playable/AvaPlayableGroup.h"
 #include "Playable/PlayableGroups/AvaRemoteProxyPlayableGroup.h"
 #include "UObject/Package.h"
@@ -102,6 +104,20 @@ void UAvaPlayableGroupChannelManager::GetPlayableGroups(TArray<TWeakObjectPtr<UA
 	OutGroups.Append(PlayableGroupsWeak);
 }
 
+UAvaPlayableGroup* UAvaPlayableGroupChannelManager::FindPlayableGroupForWorld(const UWorld* InWorld) const
+{
+	for (const TWeakObjectPtr<UAvaPlayableGroup>& PlayableGroupWeak : PlayableGroupsWeak)
+	{
+		if (UAvaPlayableGroup* PlayableGroup = PlayableGroupWeak.Get())
+		{
+			if (PlayableGroup->GetPlayWorld() == InWorld)
+			{
+				return PlayableGroup;	
+			}
+		}
+	}
+	return nullptr;
+}
 
 void UAvaPlayableGroupChannelManager::Shutdown()
 {
@@ -116,6 +132,12 @@ void UAvaPlayableGroupChannelManager::BeginDestroy()
 
 void UAvaPlayableGroupManager::Init()
 {
+	if (!SynchronizedEventDispatcher)
+	{
+		// There is only a primary/default dispatcher for now.
+		SynchronizedEventDispatcher = FAvaMediaSynchronizedEventsFeature::CreateDispatcher(TEXT("DefaultGroup"));
+	}
+
 	if (!UAvaGameInstance::GetOnEndPlay().IsBoundToObject(this))
 	{
 		UAvaGameInstance::GetOnEndPlay().AddUObject(this, &UAvaPlayableGroupManager::OnGameInstanceEndPlay);
@@ -139,6 +161,33 @@ void UAvaPlayableGroupManager::Tick(double InDeltaSeconds)
 	SCOPE_TIME_GUARD(TEXT("UAvaPlayableGroupManager::Tick"));
 	UpdateLevelStreaming();
 	TickTransitions(InDeltaSeconds);
+
+	if (SynchronizedEventDispatcher)
+	{
+		SynchronizedEventDispatcher->DispatchEvents();
+	}
+}
+
+void UAvaPlayableGroupManager::PushSynchronizedEvent(FString&& InEventSignature, TUniqueFunction<void()> InFunction)
+{
+	if (SynchronizedEventDispatcher)
+	{
+		SynchronizedEventDispatcher->PushEvent(MoveTemp(InEventSignature), MoveTemp(InFunction));
+	}
+	else if (InFunction)
+	{
+		InFunction();
+	}
+}
+
+bool UAvaPlayableGroupManager::IsSynchronizedEventPushed(const FString& InEventSignature) const
+{
+	if (SynchronizedEventDispatcher)
+	{
+		const EAvaMediaSynchronizedEventState EventState = SynchronizedEventDispatcher->GetEventState(InEventSignature);
+		return EventState == EAvaMediaSynchronizedEventState::Pending || EventState == EAvaMediaSynchronizedEventState::Ready;
+	}
+	return false;	
 }
 
 UAvaPlayableGroupChannelManager* UAvaPlayableGroupManager::FindOrAddChannelManager(const FName& InChannelName)
@@ -227,6 +276,21 @@ TArray<TWeakObjectPtr<UAvaPlayableGroup>> UAvaPlayableGroupManager::GetPlayableG
 	return Groups;
 }
 
+UAvaPlayableGroup* UAvaPlayableGroupManager::FindPlayableGroupForWorld(const UWorld* InWorld) const
+{
+	for (const TPair<FName, TObjectPtr<UAvaPlayableGroupChannelManager>>& ChannelManager : ChannelManagers)
+	{
+		if (ChannelManager.Value)
+		{
+			if (UAvaPlayableGroup* PlayableGroup = ChannelManager.Value->FindPlayableGroupForWorld(InWorld))
+			{
+				return PlayableGroup;
+			}
+		}
+	}
+	return nullptr;
+}
+
 void UAvaPlayableGroupManager::BeginDestroy()
 {
 	Shutdown();
diff --git a/Engine/Plugins/Experimental/Avalanche/Source/AvalancheMedia/Private/Playable/AvaPlayableGroupManager.h b/Engine/Plugins/Experimental/Avalanche/Source/AvalancheMedia/Private/Playable/AvaPlayableGroupManager.h
--- a/Engine/Plugins/Experimental/Avalanche/Source/AvalancheMedia/Private/Playable/AvaPlayableGroupManager.h
+++ b/Engine/Plugins/Experimental/Avalanche/Source/AvalancheMedia/Private/Playable/AvaPlayableGroupManager.h
@@ -5,8 +5,10 @@
 #include "UObject/Object.h"
 #include "UObject/ObjectPtr.h"
 #include "UObject/SoftObjectPtr.h"
+
 #include "AvaPlayableGroupManager.generated.h"
 
+class IAvaMediaSynchronizedEventDispatcher;
 class UAvaGameInstance;
 class UAvaPlayableGroup;
 class UAvaPlayableGroupManager;
@@ -26,6 +28,8 @@ public:
 	UAvaPlayableGroupManager* GetPlayableGroupManager() const;
 
 	void GetPlayableGroups(TArray<TWeakObjectPtr<UAvaPlayableGroup>>& OutGroups) const;
+
+	UAvaPlayableGroup* FindPlayableGroupForWorld(const UWorld* InWorld) const;
 	
 protected:
 	//~ Begin UObject
@@ -64,6 +68,10 @@ public:
 
 	void Tick(double InDeltaSeconds);
 
+	void PushSynchronizedEvent(FString&& InEventSignature, TUniqueFunction<void()> InFunction);
+
+	bool IsSynchronizedEventPushed(const FString& InEventSignature) const;
+
 	UAvaPlayableGroupChannelManager* FindChannelManager(const FName& InChannelName) const
 	{
 		const TObjectPtr<UAvaPlayableGroupChannelManager>* ChannelManager = ChannelManagers.Find(InChannelName);
@@ -91,6 +99,11 @@ public:
 	 */
 	TArray<TWeakObjectPtr<UAvaPlayableGroup>> GetPlayableGroups(FName InChannelName = NAME_None) const;
 
+	/**
+	 * Return the playable group corresponding to the given world.
+	 */
+	UAvaPlayableGroup* FindPlayableGroupForWorld(const UWorld* InWorld) const;
+
 protected:
 	//~ Begin UObject
 	virtual void BeginDestroy() override;
@@ -112,4 +125,8 @@ protected:
 
 	bool bIsTickingTransitions = false;
 	TSet<TWeakObjectPtr<UAvaPlayableGroup>> GroupsToTickTransitions;
+
+	TArray<TWeakObjectPtr<UAvaPlayableGroup>> GroupsToTickEvents;
+	
+	TSharedPtr<IAvaMediaSynchronizedEventDispatcher> SynchronizedEventDispatcher;
 };
\ No newline at end of file
diff --git a/Engine/Plugins/Experimental/Avalanche/Source/AvalancheMedia/Private/Playable/AvaPlayableGroupSceneViewExtension.cpp b/Engine/Plugins/Experimental/Avalanche/Source/AvalancheMedia/Private/Playable/AvaPlayableGroupSceneViewExtension.cpp
new file mode 100644
--- /dev/null
+++ b/Engine/Plugins/Experimental/Avalanche/Source/AvalancheMedia/Private/Playable/AvaPlayableGroupSceneViewExtension.cpp
@@ -0,0 +1,65 @@
+// Copyright Epic Games, Inc. All Rights Reserved.
+
+#include "AvaPlayableGroupSceneViewExtension.h"
+
+#include "AvaPlayableGroupManager.h"
+#include "AvaPlayableGroupSubsystem.h"
+#include "Engine/GameInstance.h"
+#include "IAvaMediaModule.h"
+#include "Playable/AvaPlayableGroup.h"
+#include "Playback/AvaPlaybackManager.h"
+#include "Playback/IAvaPlaybackServer.h"
+#include "SceneView.h"
+
+FAvaPlayableGroupSceneViewExtension::FAvaPlayableGroupSceneViewExtension(const FAutoRegister& InAutoReg)
+	: FSceneViewExtensionBase(InAutoReg)
+{
+}
+
+void FAvaPlayableGroupSceneViewExtension::SetupView(FSceneViewFamily& InViewFamily, FSceneView& InView)
+{
+	if (!InViewFamily.Scene)
+	{
+		return;
+	}
+	
+	const UWorld* ViewWorld = InViewFamily.Scene->GetWorld();
+	if (!ViewWorld)
+	{
+		return;
+	}
+
+	// Ideally, we would like a direct link from GameInstance (or World) to it's owning PlayableGroup and PlayableGroupManager.
+	// Todo: Refactor the AvaGameInstance path to have a link to PlayableGroup.
+	const UAvaPlayableGroupManager* PlayableGroupManager = nullptr;
+
+	// For Game Viewport output, the sub system will give us the playable group manager directly.
+	if (const UGameInstance* GameInstance = ViewWorld->GetGameInstance())
+	{
+		if (const UAvaPlayableGroupSubsystem* PlayableGroupSubsystem = GameInstance->GetSubsystem<UAvaPlayableGroupSubsystem>())
+		{
+			PlayableGroupManager = PlayableGroupSubsystem->PlayableGroupManager;
+		}
+	}
+
+	if (!PlayableGroupManager)
+	{
+		const FAvaPlaybackManager& PlaybackManager = IAvaMediaModule::Get().GetLocalPlaybackManager();
+		PlayableGroupManager = PlaybackManager.GetPlayableGroupManager();
+	}
+
+	// Search in the local playable group manager for that world.
+	UAvaPlayableGroup* ViewPlayableGroup = PlayableGroupManager->FindPlayableGroupForWorld(ViewWorld);
+	
+	// If not found, search in the playback server's playback manager.
+	if (!ViewPlayableGroup && IAvaMediaModule::Get().IsPlaybackServerStarted())
+	{
+		PlayableGroupManager = IAvaMediaModule::Get().GetPlaybackServer()->GetPlaybackManager().GetPlayableGroupManager();
+		ViewPlayableGroup = PlayableGroupManager->FindPlayableGroupForWorld(ViewWorld);
+	}
+
+	if (ViewPlayableGroup)
+	{
+		ViewPlayableGroup->SetupView(InViewFamily, InView);
+	}
+}
\ No newline at end of file
diff --git a/Engine/Plugins/Experimental/Avalanche/Source/AvalancheMedia/Private/Playable/AvaPlayableGroupSceneViewExtension.h b/Engine/Plugins/Experimental/Avalanche/Source/AvalancheMedia/Private/Playable/AvaPlayableGroupSceneViewExtension.h
new file mode 100644
--- /dev/null
+++ b/Engine/Plugins/Experimental/Avalanche/Source/AvalancheMedia/Private/Playable/AvaPlayableGroupSceneViewExtension.h
@@ -0,0 +1,18 @@
+// Copyright Epic Games, Inc. All Rights Reserved.
+
+#pragma once
+
+#include "Containers/Set.h"
+#include "SceneTypes.h"
+#include "SceneViewExtension.h"
+
+class FAvaPlayableGroupSceneViewExtension final : public FSceneViewExtensionBase
+{
+public:
+	FAvaPlayableGroupSceneViewExtension(const FAutoRegister& InAutoReg);
+	virtual ~FAvaPlayableGroupSceneViewExtension() override = default;
+
+	void SetupViewFamily(FSceneViewFamily& InViewFamily) override {}
+	void SetupView(FSceneViewFamily& InViewFamily, FSceneView& InView) override;
+	void BeginRenderViewFamily(FSceneViewFamily& InViewFamily) override {}
+};
diff --git a/Engine/Plugins/Experimental/Avalanche/Source/AvalancheMedia/Private/Playable/Playables/AvaPlayableLevelStreaming.cpp b/Engine/Plugins/Experimental/Avalanche/Source/AvalancheMedia/Private/Playable/Playables/AvaPlayableLevelStreaming.cpp
--- a/Engine/Plugins/Experimental/Avalanche/Source/AvalancheMedia/Private/Playable/Playables/AvaPlayableLevelStreaming.cpp
+++ b/Engine/Plugins/Experimental/Avalanche/Source/AvalancheMedia/Private/Playable/Playables/AvaPlayableLevelStreaming.cpp
@@ -6,6 +6,7 @@
 #include "AvaRemoteControlUtils.h"
 #include "AvaScene.h"
 #include "Camera/CameraActor.h"
+#include "Components/PrimitiveComponent.h"
 #include "Engine/Level.h"
 #include "Engine/LevelStreamingDynamic.h"
 #include "Engine/LocalPlayer.h"
@@ -17,6 +18,7 @@
 #include "Playable/AvaPlayableGroup.h"
 #include "Playable/AvaPlayableGroupManager.h"
 #include "Playback/AvaPlaybackUtils.h"
+#include "SceneView.h"
 #include "Streaming/LevelStreamingDelegates.h"
 #include "UObject/Package.h"
 #include "Viewport/AvaCameraManager.h"
@@ -102,14 +104,9 @@ namespace UE::AvaMedia::LevelStreamingPlayable::Private
 		return FindFirstCameraActor(InLevel);
 	}
 
-	EAvaPlayableStatus GetPlayableStatusFromLevelStreaming(const ULevelStreaming* InLevelStreaming)
+	EAvaPlayableStatus GetPlayableStatusFromLevelStreamingState(ELevelStreamingState InState, bool bInShouldBeLoaded)
 	{
-		if (!InLevelStreaming)
-		{
-			return EAvaPlayableStatus::Unloaded;
-		}
-
-		switch (InLevelStreaming->GetLevelStreamingState())
+		switch (InState)
 		{
 			case ELevelStreamingState::Removed:
 				return EAvaPlayableStatus::Unloaded;
@@ -117,11 +114,11 @@ namespace UE::AvaMedia::LevelStreamingPlayable::Private
 			case ELevelStreamingState::Unloaded:
 				// If the LevelStreaming was not loaded and has just been made to be loading, the status will still be "unloaded" but
 				// we consider it is loading.
-				return InLevelStreaming->ShouldBeLoaded() ? EAvaPlayableStatus::Loading : EAvaPlayableStatus::Unloaded;
+				return bInShouldBeLoaded ? EAvaPlayableStatus::Loading : EAvaPlayableStatus::Unloaded;
 			
 			case ELevelStreamingState::FailedToLoad:
 				return EAvaPlayableStatus::Error;
-			
+
 			case ELevelStreamingState::Loading:
 				return EAvaPlayableStatus::Loading;
 			
@@ -137,7 +134,22 @@ namespace UE::AvaMedia::LevelStreamingPlayable::Private
 				return EAvaPlayableStatus::Error;
 		}
 	}
-	
+
+	EAvaPlayableStatus GetPlayableStatusFromLevelStreaming(const ULevelStreaming* InLevelStreaming)
+	{
+		if (!InLevelStreaming)
+		{
+			return EAvaPlayableStatus::Unloaded;
+		}
+
+		return GetPlayableStatusFromLevelStreamingState(InLevelStreaming->GetLevelStreamingState(), InLevelStreaming->ShouldBeLoaded());
+	}
+
+	ELevelStreamingState GetLevelStreamingState(const ULevelStreaming* InLevelStreaming)
+	{
+		return InLevelStreaming ? InLevelStreaming->GetLevelStreamingState() : ELevelStreamingState::Unloaded;
+	}
+
 	static int32 GetPlayableStatusRelevanceShouldBeUnloaded(EAvaPlayableStatus InStatus)
 	{
 		switch (InStatus)
@@ -194,6 +206,29 @@ namespace UE::AvaMedia::LevelStreamingPlayable::Private
 		
 		return GetPlayableStatusRelevanceShouldBeUnloaded(InStatus);
 	}
+
+	void BuildPrimitiveComponentIdList(const AActor* InActor, TSet<FPrimitiveComponentId>& OutComponentIds)
+	{
+		TInlineComponentArray<UPrimitiveComponent*> Components;
+		InActor->GetComponents(Components);
+		for (int32 ComponentIndex = 0; ComponentIndex < Components.Num(); ComponentIndex++)
+		{
+			const UPrimitiveComponent* PrimitiveComponent = Components[ComponentIndex];
+			if (PrimitiveComponent->IsRegistered())
+			{
+				OutComponentIds.Add(PrimitiveComponent->GetPrimitiveSceneId());
+
+				for (USceneComponent* AttachedChild : PrimitiveComponent->GetAttachChildren())
+				{						
+					const UPrimitiveComponent* AttachChildPC = Cast<UPrimitiveComponent>(AttachedChild);
+					if (AttachChildPC && AttachChildPC->IsRegistered())
+					{
+						OutComponentIds.Add(AttachChildPC->GetPrimitiveSceneId());
+					}
+				}
+			}
+		}
+	}
 }
 
 bool UAvaPlayableLevelStreaming::LoadAsset(const FAvaSoftAssetPtr& InSourceAsset, bool bInInitiallyVisible)
@@ -213,6 +248,9 @@ bool UAvaPlayableLevelStreaming::LoadAsset(const FAvaSoftAssetPtr& InSourceAsset
 	const bool bAssetLoading = LoadLevel(TSoftObjectPtr<UWorld>(InSourceAsset.ToSoftObjectPath()), bInInitiallyVisible); 
 	if (bAssetLoading)
 	{
+		// Refresh status immediately (assumes it won't be loaded nor visible immediately, so shouldn't affect TL)
+		SynchronizedLevelStreamingState = UE::AvaMedia::LevelStreamingPlayable::Private::GetLevelStreamingState(LevelStreaming);
+		UpdatePlayableStatus(SynchronizedLevelStreamingState);
 		PlayableGroup->NotifyLevelStreaming(this);
 	}
 	return bAssetLoading;
@@ -247,10 +285,15 @@ bool UAvaPlayableLevelStreaming::UnloadAsset()
 	LevelStreaming = nullptr;
 	Scene = nullptr;
 	SourceLevel.Reset();
+
+	// Refresh status immediately. (not sure if level streaming is going to be called)
+	SynchronizedLevelStreamingState = UE::AvaMedia::LevelStreamingPlayable::Private::GetLevelStreamingState(LevelStreaming);
+	UpdatePlayableStatus(SynchronizedLevelStreamingState);
+
 	return true;
 }
 
-EAvaPlayableStatus UAvaPlayableLevelStreaming::GetPlayableStatus() const
+void UAvaPlayableLevelStreaming::UpdatePlayableStatus(ELevelStreamingState InNewState)
 {
 	using namespace UE::AvaMedia::LevelStreamingPlayable::Private;
 
@@ -276,9 +319,13 @@ EAvaPlayableStatus UAvaPlayableLevelStreaming::GetPlayableStatus() const
 		CompareAndSetMostRelevantStatus(SubPlayable->GetPlayableStatus());
 	}
 
-	CompareAndSetMostRelevantStatus(GetPlayableStatusFromLevelStreaming(LevelStreaming));
+	CompareAndSetMostRelevantStatus(GetPlayableStatusFromLevelStreamingState(InNewState, bShouldBeLoaded));
 
-	return MostRelevantStatus;
+	if (PlayableStatus != MostRelevantStatus)
+	{
+		PlayableStatus = MostRelevantStatus;
+		NotifyPlayableStatusChanged();
+	}
 }
 
 IAvaSceneInterface* UAvaPlayableLevelStreaming::GetSceneInterface() const
@@ -359,6 +406,44 @@ void UAvaPlayableLevelStreaming::SetShouldBeVisible(bool bInShouldBeVisible)
 	}
 }
 
+void UAvaPlayableLevelStreaming::SetupView(FSceneViewFamily& InViewFamily, FSceneView& InView)
+{
+	if (!LevelStreaming)
+	{
+		return;
+	}
+	
+	const ULevel* const Level = LevelStreaming->GetLoadedLevel();
+	if (!IsValid(Level))
+	{
+		return;
+	}
+
+	// Not using synchronized state here. We want to react to the actual state and compensate.
+	const ELevelStreamingState StreamingState = LevelStreaming->GetLevelStreamingState();
+	const bool bIsVisible = StreamingState == ELevelStreamingState::MakingVisible || StreamingState == ELevelStreamingState::LoadedVisible;
+	
+	if (bIsVisible && bShouldBeHidden)
+	{
+		TSet<FPrimitiveComponentId> HiddenPrimitives;	// Todo(opt): cache this?
+
+		using namespace UE::AvaMedia::LevelStreamingPlayable::Private;
+
+		for (TObjectPtr<AActor> Actor : Level->Actors)
+		{
+			if (IsValid(Actor))
+			{
+				BuildPrimitiveComponentIdList(Actor, HiddenPrimitives);
+			}
+		}
+
+		if (!HiddenPrimitives.IsEmpty())
+		{
+			InView.HiddenPrimitives.Append(HiddenPrimitives);
+		}
+	}
+}
+
 bool UAvaPlayableLevelStreaming::LoadLevel(const TSoftObjectPtr<UWorld>& InSourceLevel, bool bInInitiallyVisible)
 {
 	if (SourceLevel == InSourceLevel)
@@ -395,6 +480,14 @@ bool UAvaPlayableLevelStreaming::LoadLevel(const TSoftObjectPtr<UWorld>& InSourc
 	return true;
 }
 
+void UAvaPlayableLevelStreaming::HandleTransitionEvent(UAvaPlayable* InPlayable, UAvaPlayableTransition* InTransition, EAvaPlayableTransitionEventFlags InTransitionFlags)
+{
+	if (InPlayable == this && EnumHasAnyFlags(InTransitionFlags, EAvaPlayableTransitionEventFlags::ShowPlayable))
+	{
+		bShouldBeHidden = false;
+	}
+}
+
 void UAvaPlayableLevelStreaming::OnLevelStreamingStateChanged(UWorld* InWorld
 	, const ULevelStreaming* InLevelStreaming
 	, ULevel* InLevelIfLoaded
@@ -406,10 +499,46 @@ void UAvaPlayableLevelStreaming::OnLevelStreamingStateChanged(UWorld* InWorld
 	{
 		return;
 	}
-	
+
+	// Package the event handler for queueing.
+	TWeakObjectPtr<UAvaPlayableLevelStreaming> ThisPlayableWeak(this);
+	auto SyncEventHandler = [ThisPlayableWeak, InNewState]()
+	{
+		if (UAvaPlayableLevelStreaming* ThisPlayable = ThisPlayableWeak.Get())
+		{
+			ThisPlayable->OnLevelStreamingStateChanged_Synchronized(InNewState);
+		}
+	};
+
+	// Sub-playables don't have an instanceId. But they are supposed to be unique in the playable group, so the
+	// source asset path should uniquely identify them.
+	const FString InstanceIdString = GetInstanceId().IsValid() ? GetInstanceId().ToString() : GetSourceAssetPath().ToString();
+
+	// Build unique signature for this event.
+	FString SyncEventSignature = FString(TEXT("Playable_")) + InstanceIdString + FString(TEXT("_LevelStreaming_")) + EnumToString(InNewState);
+
+	// The same level streaming events are usually sent twice, but we only want to push it once.
+	// This avoids generating a warning in the sync event logs. 
+	if (!GetPlayableGroup()->IsSynchronizedEventPushed(SyncEventSignature))
+	{
+		GetPlayableGroup()->PushSynchronizedEvent(MoveTemp(SyncEventSignature), MoveTemp(SyncEventHandler));
+	}
+}
+
+void UAvaPlayableLevelStreaming::OnLevelStreamingStateChanged_Synchronized(ELevelStreamingState InNewState)
+{
+	if (!LevelStreaming)
+	{
+		return;
+	}
+
+	// Using state from synchronized event for status updates.
+	// The one from LevelStreaming is not synchronized and may cause transitions to start early on some nodes.
+	SynchronizedLevelStreamingState = InNewState;
+
 	if (InNewState == ELevelStreamingState::FailedToLoad)
 	{
-		UE_LOG(LogAvaPlayable, Error, TEXT("Level \"%s\" failed to load."), *InLevelStreaming->PackageNameToLoad.ToString());		
+		UE_LOG(LogAvaPlayable, Error, TEXT("Level \"%s\" failed to load."), *LevelStreaming->PackageNameToLoad.ToString());		
 	}
 	else if (InNewState == ELevelStreamingState::LoadedNotVisible || InNewState == ELevelStreamingState::LoadedVisible)
 	{
@@ -439,23 +568,21 @@ void UAvaPlayableLevelStreaming::OnLevelStreamingStateChanged(UWorld* InWorld
 		
 		// Resolve the ava scene for the other operations.
 		ResolveScene(Level);
-
-		OnPlayableStatusChanged().Broadcast(this);
-
-		for (const TObjectKey<UAvaPlayableLevelStreaming>& ParentPlayableKey : ParentPlayables)
-		{
-			if (UAvaPlayableLevelStreaming* ParentPlayable = ParentPlayableKey.ResolveObjectPtr())
-			{
-				ParentPlayable->OnLevelStreamingPlayableStatusChanged(this);
-			}
-		}
-
-		OnLevelStreamingPlayableStatusChanged(this);
 	}
+	
+	// Important: playable status gets updated in the synchronized event handler.
+	UpdatePlayableStatus(SynchronizedLevelStreamingState);
 }
 
-void UAvaPlayableLevelStreaming::OnLevelStreamingPlayableStatusChanged(UAvaPlayableLevelStreaming* InPlayable)
+void UAvaPlayableLevelStreaming::NotifyPlayableStatusChanged()
 {
+	using namespace UE::AvaPlayback::Utils;
+
+	UE_LOG(LogAvaPlayable, Verbose, TEXT("%s Playable \"%s\" (id:%s) Status Changed: %s"),
+		*GetBriefFrameInfo(),
+		*GetSourceAssetPath().GetAssetName(), *GetInstanceId().ToString(),
+		*StaticEnumToString(PlayableStatus));
+
 	// OnPlay (camera setup, animations, etc) can only be done when the level is visible (components must be active).
 	// With camera rig, we also need to make sure the rig level is loaded and visible.
 
@@ -467,6 +594,18 @@ void UAvaPlayableLevelStreaming::OnLevelStreamingPlayableStatusChanged(UAvaPlaya
 			OnPlay();
 		}
 	}
+	
+	OnPlayableStatusChanged().Broadcast(this);
+
+	// Parent playables must be informed of the status change too.
+	for (const TObjectKey<UAvaPlayableLevelStreaming>& ParentPlayableKey : ParentPlayables)
+	{
+		if (UAvaPlayableLevelStreaming* ParentPlayable = ParentPlayableKey.ResolveObjectPtr())
+		{
+			// Note: using the synchronized streaming state to avoid spurious states.
+			ParentPlayable->UpdatePlayableStatus(ParentPlayable->SynchronizedLevelStreamingState);
+		}
+	}	
 }
 
 void UAvaPlayableLevelStreaming::BindDelegates()
@@ -475,11 +614,16 @@ void UAvaPlayableLevelStreaming::BindDelegates()
 	{
 		FLevelStreamingDelegates::OnLevelStreamingStateChanged.AddUObject(this, &UAvaPlayableLevelStreaming::OnLevelStreamingStateChanged);
 	}
+	if (!UAvaPlayable::OnTransitionEvent().IsBoundToObject(this))
+	{
+		UAvaPlayable::OnTransitionEvent().AddUObject(this, &UAvaPlayableLevelStreaming::HandleTransitionEvent);
+	}
 }
 
 void UAvaPlayableLevelStreaming::UnbindDelegates()
 {
 	FLevelStreamingDelegates::OnLevelStreamingStateChanged.RemoveAll(this);
+	UAvaPlayable::OnTransitionEvent().RemoveAll(this);
 }
 
 ULevel* UAvaPlayableLevelStreaming::GetLoadedLevel() const
@@ -822,5 +966,4 @@ void UAvaPlayableLevelStreaming::UpdateVisibilityFromParents()
 	SetShouldBeVisible(bShouldBeVisible);
 }
 
-
 #undef LOCTEXT_NAMESPACE
diff --git a/Engine/Plugins/Experimental/Avalanche/Source/AvalancheMedia/Private/Playable/Playables/AvaPlayableLevelStreaming.h b/Engine/Plugins/Experimental/Avalanche/Source/AvalancheMedia/Private/Playable/Playables/AvaPlayableLevelStreaming.h
--- a/Engine/Plugins/Experimental/Avalanche/Source/AvalancheMedia/Private/Playable/Playables/AvaPlayableLevelStreaming.h
+++ b/Engine/Plugins/Experimental/Avalanche/Source/AvalancheMedia/Private/Playable/Playables/AvaPlayableLevelStreaming.h
@@ -3,15 +3,14 @@
 #pragma once
 
 #include "AvaMediaDefines.h"
+#include "Engine/LevelStreaming.h"
 #include "Playable/AvaPlayable.h"
 #include "AvaPlayableLevelStreaming.generated.h"
 
 class AActor;
 class AAvaScene;
-class ULevel;
-class ULevelStreaming;
+class UAvaPlayableTransition;
 class ULevelStreamingDynamic;
-enum class ELevelStreamingState : uint8;
 
 UCLASS(NotBlueprintable, BlueprintType, ClassGroup = "Motion Design Playable",
 	meta = (DisplayName = "Motion Design Level Streaming Playable"))
@@ -23,11 +22,12 @@ public:
 	virtual bool LoadAsset(const FAvaSoftAssetPtr& InSourceAsset, bool bInInitiallyVisible) override;
 	virtual bool UnloadAsset() override;
 	virtual const FSoftObjectPath& GetSourceAssetPath() const override { return SourceLevel.ToSoftObjectPath(); }
-	virtual EAvaPlayableStatus GetPlayableStatus() const override;
+	virtual EAvaPlayableStatus GetPlayableStatus() const override { return PlayableStatus; }
 	virtual IAvaSceneInterface* GetSceneInterface() const override;
 	virtual bool ApplyCamera() override;
 	virtual bool GetShouldBeVisible() const override;
 	virtual void SetShouldBeVisible(bool bInShouldBeVisible) override;
+	virtual void SetupView(FSceneViewFamily& InViewFamily, FSceneView& InView) override;
 
 protected:
 	virtual bool InitPlayable(const FPlayableCreationInfo& InPlayableInfo) override;
@@ -46,13 +46,18 @@ public:
 protected:
 	bool LoadLevel(const TSoftObjectPtr<UWorld>& InSourceLevel, bool bInInitiallyVisible);
 
+	void HandleTransitionEvent(UAvaPlayable* InPlayable, UAvaPlayableTransition* InTransition, EAvaPlayableTransitionEventFlags InTransitionFlags);
+
 	void OnLevelStreamingStateChanged(UWorld* InWorld
 		, const ULevelStreaming* InLevelStreaming
 		, ULevel* InLevelIfLoaded
 		, ELevelStreamingState InPreviousState
 		, ELevelStreamingState InNewState);
 
-	void OnLevelStreamingPlayableStatusChanged(UAvaPlayableLevelStreaming* InSubPlayable);
+	void OnLevelStreamingStateChanged_Synchronized(ELevelStreamingState InNewState);	
+
+	void UpdatePlayableStatus(ELevelStreamingState InNewState);
+	void NotifyPlayableStatusChanged();
 	
 	void BindDelegates();
 	void UnbindDelegates();
@@ -90,6 +95,11 @@ protected:
 	TObjectPtr<AAvaScene> Scene; 
 
 	bool bLoadSubPlayables = false;
+
+	EAvaPlayableStatus PlayableStatus = EAvaPlayableStatus::Unloaded;
+
+	/** Keep track of the synchronized level streaming state. */
+	ELevelStreamingState SynchronizedLevelStreamingState = ELevelStreamingState::Unloaded;
 	
 	/**
 	 * Dependent playables loaded from secondary streaming levels.
@@ -106,4 +116,7 @@ protected:
 	TSet<TObjectKey<UAvaPlayableLevelStreaming>> ParentPlayables;
 
 	bool bOnPlayQueued = false;
+	
+	/** Enter Playables should be hidden until the transition has started. */
+	bool bShouldBeHidden = true;
 };
diff --git a/Engine/Plugins/Experimental/Avalanche/Source/AvalancheMedia/Private/Playable/Playables/AvaPlayableRemoteProxy.cpp b/Engine/Plugins/Experimental/Avalanche/Source/AvalancheMedia/Private/Playable/Playables/AvaPlayableRemoteProxy.cpp
--- a/Engine/Plugins/Experimental/Avalanche/Source/AvalancheMedia/Private/Playable/Playables/AvaPlayableRemoteProxy.cpp
+++ b/Engine/Plugins/Experimental/Avalanche/Source/AvalancheMedia/Private/Playable/Playables/AvaPlayableRemoteProxy.cpp
@@ -7,6 +7,7 @@
 #include "IAvaMediaModule.h"
 #include "Playable/AvaPlayableGroupManager.h"
 #include "Playback/AvaPlaybackClientDelegates.h"
+#include "Playback/AvaPlaybackUtils.h"
 #include "Playback/IAvaPlaybackClient.h"
 
 #define LOCTEXT_NAMESPACE "AvaPlayableRemoteProxy"
@@ -19,18 +20,27 @@ bool UAvaPlayableRemoteProxy::LoadAsset(const FAvaSoftAssetPtr& InSourceAsset, b
 	}
 	
 	SourceAssetPath = InSourceAsset.ToSoftObjectPath();
-
-	IAvaPlaybackClient& PlaybackClient = IAvaMediaModule::Get().GetPlaybackClient();
+	IAvaPlaybackClient& PlaybackClient = IAvaMediaModule::Get().GetPlaybackClient();	
+	const TArray<FString> OnlineServers = PlaybackClient.GetOnlineServersForChannel(PlayingChannelFName);
 	
-	if (PlaybackClient.HasAnyServerOnlineForChannel(PlayingChannelFName))
+	bool bShouldRequestLoad = false;
+
+	// Reconcile the status per-server.
+	for (const FString& Server : OnlineServers)
 	{
-		const FString ChannelName = PlayingChannelName;
-		const FSoftObjectPath& AssetPath = InSourceAsset.ToSoftObjectPath();
-		const TOptional<EAvaPlaybackStatus> RemoteStatusOpt = PlaybackClient.GetRemotePlaybackStatus(InstanceId, AssetPath, ChannelName);
+		const TOptional<EAvaPlaybackStatus> RemoteStatusOpt =
+			PlaybackClient.GetRemotePlaybackStatus(InstanceId, SourceAssetPath, PlayingChannelName, Server);
+
+		if (!RemoteStatusOpt.IsSet())
+		{
+			PlaybackClient.RequestPlaybackAssetStatus(SourceAssetPath, Server, /*bInForceRefresh*/ false);
+		}
+		
 		const EAvaPlaybackStatus RemoteStatus = RemoteStatusOpt.IsSet() ? RemoteStatusOpt.GetValue() : EAvaPlaybackStatus::Unknown;
 
 		const bool bCanLoad = RemoteStatus == EAvaPlaybackStatus::Available
 			|| RemoteStatus == EAvaPlaybackStatus::Unknown;
+		
 		const bool bIsLoaded = RemoteStatus == EAvaPlaybackStatus::Loading
 			|| RemoteStatus == EAvaPlaybackStatus::Loaded
 			|| RemoteStatus == EAvaPlaybackStatus::Starting
@@ -38,10 +48,21 @@ bool UAvaPlayableRemoteProxy::LoadAsset(const FAvaSoftAssetPtr& InSourceAsset, b
 
 		if (bCanLoad && !bIsLoaded)
 		{
-			PlaybackClient.RequestPlayback(InstanceId, AssetPath, ChannelName, EAvaPlaybackAction::Load);
-			PlaybackClient.RequestPlayback(InstanceId, AssetPath, ChannelName, EAvaPlaybackAction::SetUserData, UserData);
+			bShouldRequestLoad = true;
 		}
 	}
+
+	// We have at least one server that was not in the proper state, so we issue the load request.
+	if (bShouldRequestLoad)
+	{
+		// Todo(opt): Combine actions: LoadWithUserData
+		// Server-side requirement: a load request on an already loaded/playing asset will not do anything.
+		PlaybackClient.RequestPlayback(InstanceId, SourceAssetPath, PlayingChannelName, EAvaPlaybackAction::Load);
+		PlaybackClient.RequestPlayback(InstanceId, SourceAssetPath, PlayingChannelName, EAvaPlaybackAction::SetUserData, UserData);
+	}
+
+	// After this, we should expect the asset to be loaded.
+	bShouldBeLoaded = true;	
 	return true;
 }
 
@@ -52,6 +73,7 @@ bool UAvaPlayableRemoteProxy::UnloadAsset()
 	{
 		PlaybackClient.RequestPlayback(InstanceId, SourceAssetPath, PlayingChannelName, EAvaPlaybackAction::Unload);
 	}
+	bShouldBeLoaded = false;
 	return true;
 }
 
@@ -84,28 +106,80 @@ namespace UE::AvaMedIaRemoteProxyPlayable::Private
 			return EAvaPlayableStatus::Error;
 		}
 	}
+
+	// During the loading process, going towards "loaded/visible" is only
+	// reached when all replicated playables are in that state.
+	int32 PlayableStatusPriorityForLoaded(EAvaPlayableStatus InStatus)
+	{
+		switch (InStatus)
+		{
+		case EAvaPlayableStatus::Unknown: return 0; // ignore unknown status
+		case EAvaPlayableStatus::Error: return 5;	// wins over everything
+		case EAvaPlayableStatus::Unloaded: return 4;
+		case EAvaPlayableStatus::Loading: return 3;
+		case EAvaPlayableStatus::Loaded: return 2;
+		case EAvaPlayableStatus::Visible: return 1;	// destination state is the "weakest"
+		default:
+			return -1;
+		}
+	}
+
+	// During the unloading process, going towards "unloaded" is only
+	// reached when all replicated playables are in that state.
+	int32 PlayableStatusPriorityForUnloaded(EAvaPlayableStatus InStatus)
+	{
+		switch (InStatus)
+		{
+		case EAvaPlayableStatus::Unknown: return 0; // ignore unknown status
+		case EAvaPlayableStatus::Error: return 5;	// wins over everything
+		case EAvaPlayableStatus::Unloaded: return 1; // destination state is the "weakest"
+		case EAvaPlayableStatus::Loading: return 2;
+		case EAvaPlayableStatus::Loaded: return 3;
+		case EAvaPlayableStatus::Visible: return 4;
+		default:
+			return -1;
+		}
+	}
+
+	int32 PlayableStatusPriority(EAvaPlayableStatus InStatus, bool bInShouldBeLoaded)
+	{
+		return bInShouldBeLoaded ? PlayableStatusPriorityForLoaded(InStatus) : PlayableStatusPriorityForUnloaded(InStatus);
+	}
+
+	// Contextually reconcile the playable statuses
+	EAvaPlayableStatus ReconcilePlayableStatus(EAvaPlayableStatus InStatus, EAvaPlayableStatus InOtherStatus, bool bInShouldBeLoaded)
+	{
+		if (PlayableStatusPriority(InStatus, bInShouldBeLoaded) > PlayableStatusPriority(InOtherStatus, bInShouldBeLoaded))
+		{
+			return InStatus;
+		}
+		return InOtherStatus;
+	}
 }
 
 EAvaPlayableStatus UAvaPlayableRemoteProxy::GetPlayableStatus() const
 {
+	using namespace UE::AvaMedIaRemoteProxyPlayable;
 	IAvaPlaybackClient& PlaybackClient = IAvaMediaModule::Get().GetPlaybackClient();
-	const TArray<FString> OnlineServers = GetOnlineServerForChannel(PlayingChannelFName);
+	const TArray<FString> OnlineServers = PlaybackClient.GetOnlineServersForChannel(PlayingChannelFName);
+
+	EAvaPlayableStatus PlayableStatus = EAvaPlayableStatus::Unknown;
 	
 	for (const FString& Server : OnlineServers)
 	{
 		TOptional<EAvaPlaybackStatus> PlaybackStatus = PlaybackClient.GetRemotePlaybackStatus(InstanceId, SourceAssetPath, PlayingChannelName, Server);
 
-		if (!PlaybackStatus.IsSet())
+		if (PlaybackStatus.IsSet())
+		{
+			PlayableStatus = Private::ReconcilePlayableStatus(PlayableStatus, Private::GetPlayableStatus(PlaybackStatus.GetValue()), bShouldBeLoaded);
+		}
+		else
 		{
 			PlaybackClient.RequestPlayback(InstanceId, SourceAssetPath, PlayingChannelName, EAvaPlaybackAction::Status);
-			PlaybackStatus = EAvaPlaybackStatus::Unknown;
 		}
-
-		// TODO: reconcile forked channels.
-		return UE::AvaMedIaRemoteProxyPlayable::Private::GetPlayableStatus(PlaybackStatus.GetValue());
 	}
 	
-	return EAvaPlayableStatus::Unknown;
+	return PlayableStatus;
 }
 
 IAvaSceneInterface* UAvaPlayableRemoteProxy::GetSceneInterface() const
@@ -137,10 +211,10 @@ EAvaPlayableCommandResult UAvaPlayableRemoteProxy::ExecuteAnimationCommand(EAvaP
 			break;
 
 		default:
+			using namespace UE::AvaPlayback::Utils;
 			UE_LOG(LogAvaPlayable, Warning,
-				TEXT("Animation command action \"%s\" for asset \"%s\" on channel \"%s\" is not implemented."),
-				*StaticEnum<EAvaPlaybackAnimAction>()->GetValueAsString(InAnimAction),
-				*SourceAssetPath.ToString(), *PlayingChannelName);
+				TEXT("%s Animation command action \"%s\" for asset \"%s\" on channel \"%s\" is not implemented."),
+				*GetBriefFrameInfo(), *StaticEnumToString(InAnimAction), *SourceAssetPath.ToString(), *PlayingChannelName);
 			break;
 		}
 	}
@@ -166,10 +240,26 @@ void UAvaPlayableRemoteProxy::SetUserData(const FString& InUserData)
 {
 	if (UserData != InUserData)
 	{
-		IAvaPlaybackClient& PlaybackClient = IAvaMediaModule::Get().GetPlaybackClient();
-		if (PlaybackClient.HasAnyServerOnlineForChannel(PlayingChannelFName))
+		// Setting user data is only allowed if source path is specified (even if InstanceId is ok).
+		if (!GetSourceAssetPath().IsNull())	// Should be set in LoadAsset.
 		{
-			PlaybackClient.RequestPlayback(InstanceId, GetSourceAssetPath(), PlayingChannelName, EAvaPlaybackAction::SetUserData, InUserData);
+			IAvaPlaybackClient& PlaybackClient = IAvaMediaModule::Get().GetPlaybackClient();
+			if (!GetSourceAssetPath().IsNull() && PlaybackClient.HasAnyServerOnlineForChannel(PlayingChannelFName))
+			{
+				PlaybackClient.RequestPlayback(InstanceId, GetSourceAssetPath(), PlayingChannelName, EAvaPlaybackAction::SetUserData, InUserData);
+			}
+		}
+		else
+		{
+			// If LoadAsset hasn't been called yet, user data is going to be sent along with the Load command.
+			
+			if (bShouldBeLoaded)	// LoadAsset already called.
+			{
+				using namespace UE::AvaPlayback::Utils;
+				UE_LOG(LogAvaPlayable, Warning,
+					TEXT("%s Failed to set user data. Source Asset Path is not specified (it should be) for instance (id:%s) on channel \"%s\"."),
+					*GetBriefFrameInfo(), *InstanceId.ToString(), *PlayingChannelName);
+			}
 		}
 	}
 	
@@ -224,6 +314,7 @@ void UAvaPlayableRemoteProxy::OnPlay()
 			|| RemoteStatus == EAvaPlaybackStatus::Unknown || RemoteStatus == EAvaPlaybackStatus::Stopping
 			|| RemoteStatus == EAvaPlaybackStatus::Unloading)
 		{
+			// Todo: Combine actions: StartWithUserData
 			Client.RequestPlayback(InstanceId, SourceAssetPath, ChannelName, EAvaPlaybackAction::Start);
 		}
 	}
@@ -248,40 +339,36 @@ void UAvaPlayableRemoteProxy::RegisterClientEventHandlers()
 {
 	using namespace UE::AvaPlaybackClient::Delegates;
 	GetOnPlaybackSequenceEvent().RemoveAll(this);
-	GetOnPlaybackSequenceEvent().AddUObject(this, &UAvaPlayableRemoteProxy::HandleAvaPlaybackSequenceEvent);
+	GetOnPlaybackSequenceEvent().AddUObject(this, &UAvaPlayableRemoteProxy::HandlePlaybackSequenceEvent);
+	GetOnPlaybackStatusChanged().RemoveAll(this);
+	GetOnPlaybackStatusChanged().AddUObject(this, &UAvaPlayableRemoteProxy::HandlePlaybackStatusChanged);
 }
 
 void UAvaPlayableRemoteProxy::UnregisterClientEventHandlers() const
 {
 	using namespace UE::AvaPlaybackClient::Delegates;
 	GetOnPlaybackSequenceEvent().RemoveAll(this);
+	GetOnPlaybackStatusChanged().RemoveAll(this);
 }
 
-TArray<FString> UAvaPlayableRemoteProxy::GetOnlineServerForChannel(const FName& InChannelName)
+void UAvaPlayableRemoteProxy::HandlePlaybackSequenceEvent(IAvaPlaybackClient& InPlaybackClient,
+	const UE::AvaPlaybackClient::Delegates::FPlaybackSequenceEventArgs& InEventArgs)
 {
-	const FAvaBroadcastOutputChannel& Channel = UAvaBroadcast::Get().GetCurrentProfile().GetChannel(InChannelName);
-	const TArray<UMediaOutput*>& Outputs = Channel.GetMediaOutputs();
-	TArray<FString> OnlineServers;
-	OnlineServers.Reserve(Outputs.Num());
-	
-	for (const UMediaOutput* Output : Outputs)
+	if (InEventArgs.InstanceId == InstanceId && InEventArgs.ChannelName == PlayingChannelName)
 	{
-		if (Channel.IsMediaOutputRemote(Output) && Channel.GetMediaOutputState(Output) != EAvaBroadcastOutputState::Offline)
-		{
-			OnlineServers.AddUnique(Channel.GetMediaOutputServerName(Output));
-		}
+		const FName SequenceName(InEventArgs.SequenceName);
+		OnSequenceEventDelegate.Broadcast(this, SequenceName, InEventArgs.EventType);
 	}
-	return OnlineServers;
 }
 
-void UAvaPlayableRemoteProxy::HandleAvaPlaybackSequenceEvent(IAvaPlaybackClient& InPlaybackClient,
-	const UE::AvaPlaybackClient::Delegates::FPlaybackSequenceEventArgs& InEventArgs)
+void UAvaPlayableRemoteProxy::HandlePlaybackStatusChanged(IAvaPlaybackClient& InPlaybackClient,
+	const UE::AvaPlaybackClient::Delegates::FPlaybackStatusChangedArgs& InEventArgs)
 {
 	if (InEventArgs.InstanceId == InstanceId && InEventArgs.ChannelName == PlayingChannelName)
 	{
-		const FName SequenceName(InEventArgs.SequenceName);
-		OnSequenceEventDelegate.Broadcast(this, SequenceName, InEventArgs.EventType);
+		OnPlayableStatusChanged().Broadcast(this);
 	}
 }
 
+
 #undef LOCTEXT_NAMESPACE
diff --git a/Engine/Plugins/Experimental/Avalanche/Source/AvalancheMedia/Private/Playable/Playables/AvaPlayableRemoteProxy.h b/Engine/Plugins/Experimental/Avalanche/Source/AvalancheMedia/Private/Playable/Playables/AvaPlayableRemoteProxy.h
--- a/Engine/Plugins/Experimental/Avalanche/Source/AvalancheMedia/Private/Playable/Playables/AvaPlayableRemoteProxy.h
+++ b/Engine/Plugins/Experimental/Avalanche/Source/AvalancheMedia/Private/Playable/Playables/AvaPlayableRemoteProxy.h
@@ -10,6 +10,7 @@ class IAvaPlaybackClient;
 namespace UE::AvaPlaybackClient::Delegates
 {
 	struct FPlaybackSequenceEventArgs;
+	struct FPlaybackStatusChangedArgs;
 }
 
 UCLASS(NotBlueprintable, BlueprintType, ClassGroup = "Motion Design Playable",
@@ -45,15 +46,25 @@ protected:
 	void RegisterClientEventHandlers();
 	void UnregisterClientEventHandlers() const;
 
-	static TArray<FString> GetOnlineServerForChannel(const FName& InChannelName);
-	
-	void HandleAvaPlaybackSequenceEvent(IAvaPlaybackClient& InPlaybackClient,
+	void HandlePlaybackSequenceEvent(IAvaPlaybackClient& InPlaybackClient,
 		const UE::AvaPlaybackClient::Delegates::FPlaybackSequenceEventArgs& InEventArgs);
-	
+
+	void HandlePlaybackStatusChanged(IAvaPlaybackClient& InPlaybackClient,
+		const UE::AvaPlaybackClient::Delegates::FPlaybackStatusChangedArgs& InEventArgs);
+
 protected:
 	/** Channel name this playable is playing on. */
 	FName PlayingChannelFName;
 	FString PlayingChannelName;
 
 	FSoftObjectPath SourceAssetPath;
+
+	/**
+	 * Internal forked channel's status reconciling needs to keep track of the expected series of events.
+	 * If should be loaded, expected sequence is: loading, loaded, making visible, visible.
+	 * If it should be unloaded the sequence is inverted.
+	 * The reconciling logic will select the "slowest" node. Example: for loading process, the status
+	 * is loading as long as there is one node still loading.
+	 */
+	bool bShouldBeLoaded = false;
 };
diff --git a/Engine/Plugins/Experimental/Avalanche/Source/AvalancheMedia/Private/Playable/Transition/AvaPlayableLocalTransition.cpp b/Engine/Plugins/Experimental/Avalanche/Source/AvalancheMedia/Private/Playable/Transition/AvaPlayableLocalTransition.cpp
--- a/Engine/Plugins/Experimental/Avalanche/Source/AvalancheMedia/Private/Playable/Transition/AvaPlayableLocalTransition.cpp
+++ b/Engine/Plugins/Experimental/Avalanche/Source/AvalancheMedia/Private/Playable/Transition/AvaPlayableLocalTransition.cpp
@@ -263,7 +263,7 @@ bool UAvaPlayableLocalTransition::Start()
 
 	using namespace UE::AvaPlayableTransition::Private;
 
-	FBuilderHelper Helper(GetFullName(), this);
+	FBuilderHelper Helper(GetInstanceName(), this);
 
 	TArray<UAvaPlayable*> EnterPlayables = Pin(EnterPlayablesWeak);
 	TArray<UAvaPlayable*> PlayingPlayables = Pin(PlayingPlayablesWeak);
@@ -282,7 +282,7 @@ bool UAvaPlayableLocalTransition::Start()
 	{
 		if (EnterPlayableValues.IsValidIndex(ArrayIndex) && EnterPlayableValues[ArrayIndex].IsValid())
 		{
-			Playable->UpdateRemoteControlCommand(EnterPlayableValues[ArrayIndex].ToSharedRef());	
+			Playable->UpdateRemoteControlCommand(EnterPlayableValues[ArrayIndex].ToSharedRef());
 		}
 		
 		if (!Helper.AddTransitionBehaviorInstance(Playable, EAvaPlayableTransitionEntryRole::Enter))
@@ -291,6 +291,9 @@ bool UAvaPlayableLocalTransition::Start()
 			InitializeSequences(Playable);
 			PostExecutorSequencePlayablesWeak.Add(Playable);
 		}
+
+		// Signal that this playable can be shown now that sequences are initialized and RC is injected.
+		UAvaPlayable::OnTransitionEvent().Broadcast(Playable, this, EAvaPlayableTransitionEventFlags::ShowPlayable);
 		++ArrayIndex;
 	}
 
@@ -341,7 +344,7 @@ bool UAvaPlayableLocalTransition::Start()
 		PostTransitionExecutorPhase();
 	}
 
-	// Todo: reconsider if this is needed and if so where it should go. Not used in code for now.
+	// Allows client-side to implement forked/clustered channel state reconciling.
 	UAvaPlayable::OnTransitionEvent().Broadcast(nullptr, this, EAvaPlayableTransitionEventFlags::Starting);
 
 	return true;
@@ -410,7 +413,7 @@ void UAvaPlayableLocalTransition::OnTransitionExecutorEnded()
 
 		TransitionExecutor.Reset();
 		
-		UE_LOG(LogAvaPlayable, Verbose, TEXT("%s Transition Executor \"%s\" ended."), *GetBriefFrameInfo(), *GetFullName());
+		UE_LOG(LogAvaPlayable, Verbose, TEXT("%s Transition Executor \"%s\" ended."), *GetBriefFrameInfo(), *GetInstanceName());
 	}
 
 	PostTransitionExecutorPhase();
@@ -448,7 +451,7 @@ void UAvaPlayableLocalTransition::PostTransitionExecutorPhase()
 void UAvaPlayableLocalTransition::NotifyTransitionFinished()
 {
 	using namespace UE::AvaPlayback::Utils;
-	UE_LOG(LogAvaPlayable, Verbose, TEXT("%s Playable Transition \"%s\" ended."), *GetBriefFrameInfo(), *GetFullName());
+	UE_LOG(LogAvaPlayable, Verbose, TEXT("%s Playable Transition \"%s\" ended."), *GetBriefFrameInfo(), *GetInstanceName());
 
 	// This will indicate the playable transition is completed and can be cleaned up.
 	// Combo templates break the one page one playable rule, so we need a dedicated event to signal the end of the playable transition.
@@ -489,4 +492,9 @@ void UAvaPlayableLocalTransition::OnPlayableSequenceEvent(UAvaPlayable* InPlayab
 			NotifyTransitionFinished();
 		}
 	}
-}
\ No newline at end of file
+}
+
+FString UAvaPlayableLocalTransition::GetInstanceName() const
+{
+	return TransitionId.IsValid() ? TransitionId.ToString() : GetFullName();
+}
diff --git a/Engine/Plugins/Experimental/Avalanche/Source/AvalancheMedia/Private/Playable/Transition/AvaPlayableLocalTransition.h b/Engine/Plugins/Experimental/Avalanche/Source/AvalancheMedia/Private/Playable/Transition/AvaPlayableLocalTransition.h
--- a/Engine/Plugins/Experimental/Avalanche/Source/AvalancheMedia/Private/Playable/Transition/AvaPlayableLocalTransition.h
+++ b/Engine/Plugins/Experimental/Avalanche/Source/AvalancheMedia/Private/Playable/Transition/AvaPlayableLocalTransition.h
@@ -41,6 +41,8 @@ protected:
 
 	void StartWaitOnPostExecutorSequences();
 	void FinishWaitOnPostExecutorSequences();
+
+	FString GetInstanceName() const;
 	
 protected:
 	/** Transition tree executor. */
diff --git a/Engine/Plugins/Experimental/Avalanche/Source/AvalancheMedia/Private/Playable/Transition/AvaPlayableRemoteTransition.cpp b/Engine/Plugins/Experimental/Avalanche/Source/AvalancheMedia/Private/Playable/Transition/AvaPlayableRemoteTransition.cpp
--- a/Engine/Plugins/Experimental/Avalanche/Source/AvalancheMedia/Private/Playable/Transition/AvaPlayableRemoteTransition.cpp
+++ b/Engine/Plugins/Experimental/Avalanche/Source/AvalancheMedia/Private/Playable/Transition/AvaPlayableRemoteTransition.cpp
@@ -20,7 +20,6 @@ bool UAvaPlayableRemoteTransition::Start()
 	
 	if (!AvaMediaModule.IsPlaybackClientStarted())
 	{
-		TransitionId.Invalidate();
 		return false;
 	}
 
@@ -29,7 +28,10 @@ bool UAvaPlayableRemoteTransition::Start()
 		return false;
 	}
 
-	TransitionId = FGuid::NewGuid();
+	if (!TransitionId.IsValid())
+	{
+		TransitionId = FGuid::NewGuid();
+	}
 	
 	using namespace UE::AvaPlayableTransition::Private;
 	TArray<FGuid> EnterInstanceIds = GetInstanceIds(EnterPlayablesWeak);
@@ -44,23 +46,68 @@ bool UAvaPlayableRemoteTransition::Start()
 	
 	IAvaPlaybackClient& PlaybackClient = AvaMediaModule.GetPlaybackClient();
 	PlaybackClient.RequestPlayableTransitionStart(TransitionId, MoveTemp(EnterInstanceIds), MoveTemp(PlayingInstanceIds), MoveTemp(ExitInstanceIds), MoveTemp(EnterValues), ChannelName, TransitionFlags);
+
+	TArray<FString> Servers = PlaybackClient.GetOnlineServersForChannel(ChannelName);
+	for (const FString& Server : Servers)
+	{
+		FRemoteStatusInfo& StatusInfo = RemoteStatusPerServer.Add(Server);
+		StatusInfo.Status = ERemoteStatus::StartRequest;
+	}
+
 	RegisterToPlaybackClientDelegates();	// to get the transition events from the server side.
 
 	using namespace UE::AvaPlayback::Utils;
-	UE_LOG(LogAvaPlayable, Verbose, TEXT("%s Remote Playable Transition \"%s\" Id:\"%s\" starting."), *GetBriefFrameInfo(), *GetFullName(), *TransitionId.ToString());
+	UE_LOG(LogAvaPlayable, Verbose, TEXT("%s Remote Playable Transition \"%s\" starting."), *GetBriefFrameInfo(), *GetInstanceName());
 	return true;
 }
 
 void UAvaPlayableRemoteTransition::Stop()
 {
-	if (TransitionId.IsValid())
+	IAvaMediaModule& AvaMediaModule = IAvaMediaModule::Get();
+
+	// There is no need to request a stop if we have been notified the transition is already ended on all servers (avoids server side warning).
+	if (AvaMediaModule.IsPlaybackClientStarted() && !IsTransitionFinishedOnAllServers())
 	{
-		IAvaMediaModule& AvaMediaModule = IAvaMediaModule::Get();
-	
-		if (AvaMediaModule.IsPlaybackClientStarted())
+		AvaMediaModule.GetPlaybackClient().RequestPlayableTransitionStop(TransitionId, ChannelName);
+	}
+
+	// Forked/Clustered channels: Check for desyncs.
+	if (RemoteStatusPerServer.Num() > 1 && IsTransitionFinishedOnAllServers())
+	{
+		auto CheckForDesync = [this](const FString& InValueName, const TFunction<int32(const FRemoteStatusInfo& InStatusInfo)>& InGetValueFunction)
 		{
-			AvaMediaModule.GetPlaybackClient().RequestPlayableTransitionStop(TransitionId, ChannelName);
-		}
+			bool bDesyncDetected = false;
+			int32 DesyncDeltaFrame = 0;
+			TOptional<int32> PreviousValue;
+			for (const TPair<FString, FRemoteStatusInfo>& StatusInfo : RemoteStatusPerServer)
+			{
+				const int32 CurrentValue = InGetValueFunction(StatusInfo.Value);
+				if (PreviousValue.IsSet() && PreviousValue.GetValue() != CurrentValue)
+				{
+					DesyncDeltaFrame = FMath::Abs(PreviousValue.GetValue() - CurrentValue);
+					bDesyncDetected = true;
+					break;
+				}
+				PreviousValue = CurrentValue;
+			}
+
+			// Report desync
+			if (bDesyncDetected)
+			{
+				using namespace UE::AvaPlayback::Utils;
+				UE_LOG(LogAvaPlayable, Warning, TEXT("%s Remote Playable Transition \"%s\" %s Desync Detected: Delta Frames: %d."),
+					*GetBriefFrameInfo(), *GetInstanceName(), *InValueName, DesyncDeltaFrame);
+			}
+		};
+
+		CheckForDesync(TEXT("Start"), [](const FRemoteStatusInfo& InStatusInfo)
+		{
+			return InStatusInfo.StartFrameNumber;
+		});
+		CheckForDesync(TEXT("Finish"), [](const FRemoteStatusInfo& InStatusInfo)
+		{
+			return InStatusInfo.FinishFrameNumber;
+		});
 	}
 
 	Super::Stop();
@@ -70,49 +117,165 @@ void UAvaPlayableRemoteTransition::RegisterToPlaybackClientDelegates()
 {
 	UE::AvaPlaybackClient::Delegates::GetOnPlaybackTransitionEvent().RemoveAll(this);
 	UE::AvaPlaybackClient::Delegates::GetOnPlaybackTransitionEvent().AddUObject(this, &UAvaPlayableRemoteTransition::HandlePlaybackTransitionEvent);
+	UE::AvaPlaybackClient::Delegates::GetOnConnectionEvent().RemoveAll(this);
+	UE::AvaPlaybackClient::Delegates::GetOnConnectionEvent().AddUObject(this, &UAvaPlayableRemoteTransition::HandleRemoteConnectionEvent);
 }
 
 void UAvaPlayableRemoteTransition::UnregisterFromPlaybackClientDelegates() const
 {
 	UE::AvaPlaybackClient::Delegates::GetOnPlaybackTransitionEvent().RemoveAll(this);
+	UE::AvaPlaybackClient::Delegates::GetOnConnectionEvent().RemoveAll(this);
 }
 
 bool UAvaPlayableRemoteTransition::IsRunning() const
 {
-	return TransitionId.IsValid();	
+	return !RemoteStatusPerServer.IsEmpty() && !IsTransitionFinishedOnAllServers();
 }
 
 void UAvaPlayableRemoteTransition::HandlePlaybackTransitionEvent(IAvaPlaybackClient& InPlaybackClient,
 	const UE::AvaPlaybackClient::Delegates::FPlaybackTransitionEventArgs& InArgs)
 {
+	using namespace UE::AvaPlayback::Utils;
+
 	if (InArgs.TransitionId != TransitionId)
-	{
+	{ 
 		return;
 	}
 
+	const bool bIsTransitionStartingEvent = EnumHasAnyFlags(InArgs.EventFlags, EAvaPlayableTransitionEventFlags::Starting);
+	const bool bIsTransitionFinishedEvent = EnumHasAnyFlags(InArgs.EventFlags, EAvaPlayableTransitionEventFlags::Finished);
+	
+	if (bIsTransitionStartingEvent)
+	{
+		// Don't stomp finished state.
+		if (GetRemoteStatus(InArgs.ServerName) != ERemoteStatus::Finished)
+		{
+			SetRemoteStatus(InArgs.ServerName, ERemoteStatus::Started);
+		}
+		else
+		{
+			UE_LOG(LogAvaPlayable, Warning,
+				TEXT("%s Remote Playable Transition \"%s\" discarding \"Starting\" event from server \"%s\" because it is already finished."),
+				*GetBriefFrameInfo(), *GetInstanceName(), *InArgs.ServerName);
+		}
+
+		SetRemoteStartFrame(InArgs.ServerName, InArgs.ServerFrameNumber);
+	}
+
+	if (bIsTransitionFinishedEvent)	// Finished has priority over all other states.
+	{
+		SetRemoteStatus(InArgs.ServerName, ERemoteStatus::Finished);
+		SetRemoteFinishFrame(InArgs.ServerName, InArgs.ServerFrameNumber);
+
+		UE_LOG(LogAvaPlayable, Verbose, TEXT("%s Remote Playable Transition \"%s\" ended on server \"%s\" at frame [%d]."),
+			*GetBriefFrameInfo(), *GetInstanceName(), *InArgs.ServerName, InArgs.ServerFrameNumber);
+	}
+
 	if (InArgs.InstanceId.IsValid())
 	{
 		if (UAvaPlayable* Playable = FindPlayable(InArgs.InstanceId))
 		{
+			const bool bIsStopPlayableEvent = EnumHasAnyFlags(InArgs.EventFlags, EAvaPlayableTransitionEventFlags::StopPlayable);
+
+			// Reconciling stop playable event for forked channels.
+			// We only propagate the stop event if all forked channels have stopped.
+			if (bIsStopPlayableEvent)
+			{
+				MarkPlayableForStop(InArgs.ServerName, InArgs.InstanceId);
+			}
+
 			// Relay locally through playable event.
-			UAvaPlayable::OnTransitionEvent().Broadcast(Playable, this, InArgs.EventFlags);
+			if (!bIsStopPlayableEvent || IsPlayableMarkedForStopOnAllServers(InArgs.InstanceId))
+			{
+				UAvaPlayable::OnTransitionEvent().Broadcast(Playable, this, InArgs.EventFlags);
+			}
 		}
 		else
 		{
 			UE_LOG(LogAvaPlayable, Error,
-				TEXT("Remote Playable Transition \"%s\" doesn't have playable instance Id \"%s\"."),
-				*TransitionId.ToString(), *InArgs.InstanceId.ToString());
+				TEXT("%s Remote Playable Transition \"%s\" doesn't have playable instance Id \"%s\"."),
+				*GetBriefFrameInfo(), *GetInstanceName(), *InArgs.InstanceId.ToString());
 		}
 	}
 	else
 	{
-		UAvaPlayable::OnTransitionEvent().Broadcast(nullptr, this, InArgs.EventFlags);
+		// Event propagation reconciling for forked channels.
+		// Only propagate "finish" event when all remote transitions are finished.
+		if ((bIsTransitionFinishedEvent && IsTransitionFinishedOnAllServers()) || bIsTransitionStartingEvent)
+		{
+			UAvaPlayable::OnTransitionEvent().Broadcast(nullptr, this, InArgs.EventFlags);
+		}
+	}
+}
+
+void UAvaPlayableRemoteTransition::HandleRemoteConnectionEvent(IAvaPlaybackClient& InPlaybackClient,
+	const UE::AvaPlaybackClient::Delegates::FConnectionEventArgs& InArgs)
+{
+	using namespace UE::AvaPlaybackClient::Delegates;
+	if (InArgs.Event == EConnectionEvent::ServerDisconnected)
+	{
+		RemoteStatusPerServer.Remove(InArgs.ServerName);
+		if (IsTransitionFinishedOnAllServers())
+		{
+			UAvaPlayable::OnTransitionEvent().Broadcast(nullptr, this, EAvaPlayableTransitionEventFlags::Finished);
+		}
+	}
+}
+
+UAvaPlayableRemoteTransition::ERemoteStatus UAvaPlayableRemoteTransition::GetRemoteStatus(const FString& InServer) const
+{
+	const FRemoteStatusInfo* FoundRemoteStatus = RemoteStatusPerServer.Find(InServer);
+	return FoundRemoteStatus ? FoundRemoteStatus->Status : ERemoteStatus::Unknown;
+}
+
+void UAvaPlayableRemoteTransition::SetRemoteStatus(const FString& InServer, ERemoteStatus InStatus)
+{
+	RemoteStatusPerServer.FindOrAdd(InServer).Status = InStatus;
+}
+
+void UAvaPlayableRemoteTransition::SetRemoteStartFrame(const FString& InServer, int32 InFrameNumber)
+{
+	RemoteStatusPerServer.FindOrAdd(InServer).StartFrameNumber = InFrameNumber;
+}
+
+void UAvaPlayableRemoteTransition::SetRemoteFinishFrame(const FString& InServer, int32 InFrameNumber)
+{
+	RemoteStatusPerServer.FindOrAdd(InServer).FinishFrameNumber = InFrameNumber;
+}
+
+bool UAvaPlayableRemoteTransition::IsTransitionFinishedOnAllServers() const
+{
+	for (const TPair<FString, FRemoteStatusInfo>& StatusInfo : RemoteStatusPerServer)
+	{
+		if (StatusInfo.Value.Status != ERemoteStatus::Finished)
+		{
+			return false;
+		}
 	}
+	return true;
+}
 
-	if (EnumHasAnyFlags(InArgs.EventFlags, EAvaPlayableTransitionEventFlags::Finished))
+void UAvaPlayableRemoteTransition::MarkPlayableForStop(const FString& InServer, const FGuid& InInstanceId)
+{
+	if (FRemoteStatusInfo* FoundRemoteStatus = RemoteStatusPerServer.Find(InServer))
 	{
-		using namespace UE::AvaPlayback::Utils;
-		UE_LOG(LogAvaPlayable, Verbose, TEXT("%s Remote Playable Transition \"%s\" Id:\"%s\" ended."), *GetBriefFrameInfo(), *GetFullName(), *TransitionId.ToString());
+		FoundRemoteStatus->PlayablesMarkedForStop.Add(InInstanceId);
 	}
 }
 
+bool UAvaPlayableRemoteTransition::IsPlayableMarkedForStopOnAllServers(const FGuid& InInstanceId) const
+{
+	for (const TPair<FString, FRemoteStatusInfo>& StatusInfo : RemoteStatusPerServer)
+	{
+		if (!StatusInfo.Value.PlayablesMarkedForStop.Contains(InInstanceId))
+		{
+			return false;
+		}
+	}
+	return true; // was marked on all servers.
+}
+
+FString UAvaPlayableRemoteTransition::GetInstanceName() const
+{
+	return TransitionId.ToString();
+}
diff --git a/Engine/Plugins/Experimental/Avalanche/Source/AvalancheMedia/Private/Playable/Transition/AvaPlayableRemoteTransition.h b/Engine/Plugins/Experimental/Avalanche/Source/AvalancheMedia/Private/Playable/Transition/AvaPlayableRemoteTransition.h
--- a/Engine/Plugins/Experimental/Avalanche/Source/AvalancheMedia/Private/Playable/Transition/AvaPlayableRemoteTransition.h
+++ b/Engine/Plugins/Experimental/Avalanche/Source/AvalancheMedia/Private/Playable/Transition/AvaPlayableRemoteTransition.h
@@ -10,6 +10,7 @@ class IAvaPlaybackClient;
 namespace UE::AvaPlaybackClient::Delegates
 {
 	struct FPlaybackTransitionEventArgs;
+	struct FConnectionEventArgs;
 }
 
 UCLASS()
@@ -32,10 +33,46 @@ protected:
 	void HandlePlaybackTransitionEvent(IAvaPlaybackClient& InPlaybackClient,
 		const UE::AvaPlaybackClient::Delegates::FPlaybackTransitionEventArgs& InArgs);
 
+	void HandleRemoteConnectionEvent(IAvaPlaybackClient& InPlaybackClient,
+		const UE::AvaPlaybackClient::Delegates::FConnectionEventArgs& InArgs);
+
 	void RegisterToPlaybackClientDelegates();
 	void UnregisterFromPlaybackClientDelegates() const;
 
+	enum class ERemoteStatus : uint8
+	{
+		Unknown,
+		StartRequest,
+		Started,
+		Finished
+	};
+
+	ERemoteStatus GetRemoteStatus(const FString& InServer) const;
+	void SetRemoteStatus(const FString& InServer, ERemoteStatus InStatus);
+	void SetRemoteStartFrame(const FString& InServer, int32 InFrameNumber);
+	void SetRemoteFinishFrame(const FString& InServer, int32 InFrameNumber);
+
+	bool IsTransitionFinishedOnAllServers() const;
+
+	void MarkPlayableForStop(const FString& InServer, const FGuid& InInstanceId);
+	bool IsPlayableMarkedForStopOnAllServers(const FGuid& InInstanceId) const;
+	
+	FString GetInstanceName() const;
+	
 protected:
-	FGuid TransitionId;
 	FName ChannelName;
+
+	/**
+	 * Forked channel support: tracking the status of the transition per
+	 * server to be able to reconcile local event propagation.
+	 */
+	struct FRemoteStatusInfo
+	{
+		ERemoteStatus Status = ERemoteStatus::Unknown;
+		int32 StartFrameNumber = 0;
+		int32 FinishFrameNumber = 0;
+		TSet<FGuid> PlayablesMarkedForStop;
+	};
+	
+	TMap<FString, FRemoteStatusInfo> RemoteStatusPerServer;
 };
\ No newline at end of file
diff --git a/Engine/Plugins/Experimental/Avalanche/Source/AvalancheMedia/Private/Playable/Transition/AvaPlayableTransition.cpp b/Engine/Plugins/Experimental/Avalanche/Source/AvalancheMedia/Private/Playable/Transition/AvaPlayableTransition.cpp
--- a/Engine/Plugins/Experimental/Avalanche/Source/AvalancheMedia/Private/Playable/Transition/AvaPlayableTransition.cpp
+++ b/Engine/Plugins/Experimental/Avalanche/Source/AvalancheMedia/Private/Playable/Transition/AvaPlayableTransition.cpp
@@ -49,6 +49,11 @@ void UAvaPlayableTransition::Stop()
 	}
 }
 
+void UAvaPlayableTransition::SetTransitionId(const FGuid& InTransitionId)
+{
+	TransitionId = InTransitionId;
+}
+
 void UAvaPlayableTransition::SetTransitionFlags(EAvaPlayableTransitionFlags InFlags)
 {
 	TransitionFlags = InFlags;
@@ -239,7 +244,7 @@ bool FAvaPlayableTransitionBuilder::AddExitPlayable(UAvaPlayable* InPlayable)
 	return true;
 }
 
-UAvaPlayableTransition* FAvaPlayableTransitionBuilder::MakeTransition(UObject* InOuter)
+UAvaPlayableTransition* FAvaPlayableTransitionBuilder::MakeTransition(UObject* InOuter, const FGuid& InTransitionId)
 {
 	// Determine if we have remote playables.
 	using namespace UE::AvaPlayableTransition::Private;
@@ -303,6 +308,7 @@ UAvaPlayableTransition* FAvaPlayableTransitionBuilder::MakeTransition(UObject* I
 	PlayableTransition->SetEnterPlayables(MoveTemp(EnterPlayablesWeak));
 	PlayableTransition->SetPlayingPlayables(MoveTemp(PlayingPlayablesWeak));
 	PlayableTransition->SetExitPlayables(MoveTemp(ExitPlayablesWeak));
+	PlayableTransition->SetTransitionId(InTransitionId);
 	
 	return PlayableTransition;
 }
diff --git a/Engine/Plugins/Experimental/Avalanche/Source/AvalancheMedia/Private/Playable/Transition/AvaPlayableTransitionPrivate.cpp b/Engine/Plugins/Experimental/Avalanche/Source/AvalancheMedia/Private/Playable/Transition/AvaPlayableTransitionPrivate.cpp
--- a/Engine/Plugins/Experimental/Avalanche/Source/AvalancheMedia/Private/Playable/Transition/AvaPlayableTransitionPrivate.cpp
+++ b/Engine/Plugins/Experimental/Avalanche/Source/AvalancheMedia/Private/Playable/Transition/AvaPlayableTransitionPrivate.cpp
@@ -4,6 +4,7 @@
 
 #include "Playable/AvaPlayable.h"
 #include "Playable/Playables/AvaPlayableRemoteProxy.h"
+#include "Playback/AvaPlaybackUtils.h"
 
 namespace UE::AvaPlayableTransition::Private
 {
@@ -74,7 +75,7 @@ namespace UE::AvaPlayableTransition::Private
 			return FString::Printf(TEXT("Id:%s, Asset:%s, Status:%s"),
 				*InPlayable->GetInstanceId().ToString(),
 				*InPlayable->GetSourceAssetPath().ToString(),
-				*StaticEnum<EAvaPlayableStatus>()->GetNameByValue(static_cast<int32>(InPlayable->GetPlayableStatus())).ToString());
+				*AvaPlayback::Utils::StaticEnumToString(InPlayable->GetPlayableStatus()));
 		}
 		return TEXT("(nullptr)");
 	}
diff --git a/Engine/Plugins/Experimental/Avalanche/Source/AvalancheMedia/Private/Playback/AvaPlaybackClient.cpp b/Engine/Plugins/Experimental/Avalanche/Source/AvalancheMedia/Private/Playback/AvaPlaybackClient.cpp
--- a/Engine/Plugins/Experimental/Avalanche/Source/AvalancheMedia/Private/Playback/AvaPlaybackClient.cpp
+++ b/Engine/Plugins/Experimental/Avalanche/Source/AvalancheMedia/Private/Playback/AvaPlaybackClient.cpp
@@ -18,6 +18,7 @@
 #include "Misc/CoreDelegates.h"
 #include "Misc/Paths.h"
 #include "Playback/AvaPlaybackClientDelegates.h"
+#include "Playback/AvaPlaybackUtils.h"
 #include "UObject/ObjectSaveContext.h"
 #include "UObject/Package.h"
 
@@ -25,11 +26,46 @@ DEFINE_LOG_CATEGORY_STATIC(LogAvaPlaybackClient, Log, All);
 
 namespace UE::AvaPlaybackClient::Private
 {
-	template<typename InEnumType>
-	FString EnumToString(InEnumType InValue)
+	bool IsPlaybackActionAllowedForNullAsset(EAvaPlaybackAction InAction)
 	{
-		return StaticEnum<InEnumType>()->GetNameStringByValue(static_cast<int64>(InValue));
+		switch (InAction)
+		{
+		case EAvaPlaybackAction::None:
+			return false;
+		case EAvaPlaybackAction::Load:
+		case EAvaPlaybackAction::Start:
+		case EAvaPlaybackAction::Stop:
+		case EAvaPlaybackAction::Unload:
+			return true;
+		case EAvaPlaybackAction::Status:
+		case EAvaPlaybackAction::SetUserData:
+		case EAvaPlaybackAction::GetUserData:
+	default:
+		return false;
+		}
 	}
+
+	bool IsAssetMissing(EAvaPlaybackAssetStatus InAssetStatus)
+	{
+		return InAssetStatus == EAvaPlaybackAssetStatus::Missing;
+	}
+
+	bool IsAssetAvailable(EAvaPlaybackAssetStatus InAssetStatus)
+	{
+		return InAssetStatus == EAvaPlaybackAssetStatus::Available
+		|| InAssetStatus == EAvaPlaybackAssetStatus::MissingDependencies;
+	}
+
+	FString GetCommandActionString(EAvaPlaybackAction InAction, const FString& InArguments)
+	{
+		FString ActionString = AvaPlayback::Utils::StaticEnumToString(InAction); 
+		if (InArguments.IsEmpty())
+		{
+			return ActionString;
+		}
+
+		return FString::Printf(TEXT("%s \"%s\""), *ActionString, *InArguments);
+	};
 }
 
 FAvaPlaybackClient::FAvaPlaybackClient(IAvaMediaModule* InParentModule)
@@ -240,6 +276,17 @@ void FAvaPlaybackClient::RequestPlaybackAssetStatus(const FSoftObjectPath& InAss
 void FAvaPlaybackClient::RequestPlayback(const FGuid& InInstanceId, const FSoftObjectPath& InAssetPath, const FString& InChannelName,
 	EAvaPlaybackAction InAction, const FString& InArguments)
 {
+	using namespace UE::AvaPlayback::Utils;
+	using namespace UE::AvaPlaybackClient::Private;
+
+	if (InAssetPath.IsNull() && !IsPlaybackActionAllowedForNullAsset(InAction))
+	{
+		UE_LOG(LogAvaPlaybackClient, Error,
+		   TEXT("Playback request [%s] on channel \"%s\" with no asset specified is not allowed."),
+		   *GetCommandActionString(InAction, InArguments), *InChannelName);
+		return;
+	}
+	
 	// Possibly TMP. Maybe find a generic way to deal with pending requests.
 	// Status requests can be spammed every frame by the UI so we block them
 	// here to avoid spamming the message bus and servers.
@@ -279,7 +326,8 @@ void FAvaPlaybackClient::RequestPlayback(const FGuid& InInstanceId, const FSoftO
 			{
 				// The command will be sent in a batch on the next Tick.
 				ServerInfo->PendingPlaybackCommands.Add({InInstanceId, InAssetPath, InChannelName, InAction, InArguments});
-				
+
+				// Todo: Ideally, don't override the remote status. Need another way to know if a request has been issued.
 				switch (InAction)
 				{
 				case EAvaPlaybackAction::None:
@@ -308,8 +356,7 @@ void FAvaPlaybackClient::RequestPlayback(const FGuid& InInstanceId, const FSoftO
 		// We will wait for feedback from server, if it managed to play the asset correctly.
 		UE_LOG(LogAvaPlaybackClient, Warning,
 			   TEXT("Playback request \"%s\" on channel \"%s\" with no asset specified. Unable to update playback status."),
-			   *UE::AvaPlaybackClient::Private::EnumToString(InAction),
-			   *InChannelName);
+			   *StaticEnumToString(InAction), *InChannelName);
 	}
 }
 
@@ -334,6 +381,7 @@ void FAvaPlaybackClient::RequestAnimPlayback(const FGuid& InInstanceId, const FS
 void FAvaPlaybackClient::RequestAnimAction(const FGuid& InInstanceId, const FSoftObjectPath& InAssetPath, const FString& InChannelName,
 												  const FString& InAnimationName, EAvaPlaybackAnimAction InAction)
 {
+	using namespace UE::AvaPlayback::Utils;
 	FAvaPlaybackAnimPlaybackRequest* Request = FMessageEndpoint::MakeMessage<FAvaPlaybackAnimPlaybackRequest>();
 	Request->InstanceId = InInstanceId;
 	Request->AssetPath = InAssetPath;
@@ -343,7 +391,7 @@ void FAvaPlaybackClient::RequestAnimAction(const FGuid& InInstanceId, const FSof
 	if (!InAssetPath.IsValid())
 	{
 		UE_LOG(LogAvaPlaybackClient, Warning, TEXT("Animation \"%s\" %s request with invalid asset path on channel \"%s\"."),
-			*InAnimationName, *StaticEnum<EAvaPlaybackAnimAction>()->GetValueAsString(InAction), *InChannelName);
+			*InAnimationName, *StaticEnumToString(InAction), *InChannelName);
 	}
 	
 	SendRequest(Request, GetServerAddressesForChannel(InChannelName));
@@ -611,6 +659,11 @@ bool FAvaPlaybackClient::HasAnyServerOnlineForChannel(const FName& InChannelName
 	return false;
 }
 
+TArray<FString> FAvaPlaybackClient::GetOnlineServersForChannel(const FName& InChannelName) const
+{
+	return GetServerNamesForChannel(InChannelName, /*bInOnlineOnly*/ true);
+}
+
 TOptional<EAvaPlaybackStatus> FAvaPlaybackClient::GetRemotePlaybackStatus(
 	const FGuid& InInstanceId, const FSoftObjectPath& InAssetPath, const FString& InChannelName, const FString& InServerName) const
 {
@@ -883,11 +936,11 @@ void FAvaPlaybackClient::HandleDeviceProviderDataListMessage(const FAvaBroadcast
 void FAvaPlaybackClient::HandleBroadcastStatusMessage(const FAvaBroadcastStatus& InMessage,
 														   const TSharedRef<IMessageContext, ESPMode::ThreadSafe>& InContext)
 {
-	using namespace UE::AvaPlaybackClient::Private;
+	using namespace UE::AvaPlayback::Utils;
 	UE_LOG(LogAvaPlaybackClient, Verbose,
 		   TEXT("Received broadcast update from server \"%s\" for channel \"%s\": Status: \"%s\"."),
 		   *InMessage.ServerName, *InMessage.ChannelName,
-		   *EnumToString(InMessage.ChannelState));
+		   *StaticEnumToString(InMessage.ChannelState));
 
 	// Log details for the output status.
 	for (const TPair<FGuid, FAvaBroadcastOutputStatus>& OutputStatus : InMessage.MediaOutputStatuses)
@@ -895,8 +948,8 @@ void FAvaPlaybackClient::HandleBroadcastStatusMessage(const FAvaBroadcastStatus&
 		UE_LOG(LogAvaPlaybackClient, Verbose,
 			   TEXT("Playback Server: \"%s\" Channel: \"%s\" OutputId \"%s\" Status: \"%s\" Severity: \"%s\"."),
 			   *InMessage.ServerName, *InMessage.ChannelName, *OutputStatus.Key.ToString(),
-			   *EnumToString(OutputStatus.Value.MediaOutputState),
-			   *EnumToString(OutputStatus.Value.MediaIssueSeverity));
+			   *StaticEnumToString(OutputStatus.Value.MediaOutputState),
+			   *StaticEnumToString(OutputStatus.Value.MediaIssueSeverity));
 		
 		for (const FString& Message : OutputStatus.Value.MediaIssueMessages)
 		{
@@ -1014,13 +1067,14 @@ void FAvaPlaybackClient::HandleBroadcastStatusMessage(const FAvaBroadcastStatus&
 void FAvaPlaybackClient::HandlePlaybackAssetStatusMessage(
 	const FAvaPlaybackAssetStatus& InMessage, const TSharedRef<IMessageContext, ESPMode::ThreadSafe>& InContext)
 {
-	using namespace UE::AvaPlaybackClient::Private;
+	using namespace UE::AvaPlayback::Utils;
 	using namespace UE::AvaPlaybackClient::Delegates;
+
 	if (!InMessage.AssetPath.IsValid())
 	{
 		UE_LOG(LogAvaPlaybackClient, Error,
-			   TEXT("Received Playback Asset Status \"%s\" from server \"%s\" with invalid asset path, can't update status."),
-			   *EnumToString(InMessage.Status), *InMessage.ServerName);
+		   TEXT("%s Received Remote Playback Asset \"%s\" from \"%s\" Status: %s. Invalid asset path, can't update status."),
+		   *GetBriefFrameInfo(), *InMessage.AssetPath.GetAssetName(), *InMessage.ServerName, *StaticEnumToString(InMessage.Status));
 		return;
 	}
 	
@@ -1031,91 +1085,126 @@ void FAvaPlaybackClient::HandlePlaybackAssetStatusMessage(
 	// Clear the pending request (if any).
 	ServerInfo.PendingPlaybackAssetStatusRequests.Remove(InMessage.AssetPath);
 
+	UE_LOG(LogAvaPlaybackClient, Verbose,
+	   TEXT("%s Received Remote Playback Asset \"%s\" from \"%s\" Status Changed: %s."),
+	   *GetBriefFrameInfo(), *InMessage.AssetPath.GetAssetName(), *InMessage.ServerName, *StaticEnumToString(InMessage.Status));
+	
 	GetOnPlaybackAssetStatusChanged().Broadcast(*this,{InMessage.AssetPath, InMessage.ServerName, InMessage.Status});
 
 	// TODO - Further refactoring needed to untangle playback state and asset state. 
 	// Because some of the states of the playback status reflect the state of the asset on disk,
 	// we need to make sure the playback state properly reflects the asset state for those cases.
-	const bool bAssetIsMissing = InMessage.Status == EAvaPlaybackAssetStatus::Missing;
-	const bool bAssetIsAvailable = InMessage.Status == EAvaPlaybackAssetStatus::Available
-									|| InMessage.Status == EAvaPlaybackAssetStatus::MissingDependencies;
+	using namespace UE::AvaPlaybackClient::Private;
+	const bool bAssetIsMissing = IsAssetMissing(InMessage.Status);
+	const bool bAssetIsAvailable = IsAssetAvailable(InMessage.Status);
 
 	for (const TPair<FString, TUniquePtr<FPlaybackChannelInfo>>& ChannelInfo : ServerInfo.PlaybackChannelInfosByName)
 	{
 		check(ChannelInfo.Value.IsValid());
 		if (FPlaybackAssetInfo* AssetInfo = ChannelInfo.Value->GetAssetInfo(InMessage.AssetPath))
 		{
-			auto ConditionalChangeStatus = [AssetInfo, &InMessage, &ChannelInfo, bAssetIsMissing, bAssetIsAvailable, this](const FGuid& InInstanceId, EAvaPlaybackStatus InPlaybackStatus)
+			auto SetPlaybackInstanceStatus = [this, AssetInfo, &InMessage, &ChannelInfo](const FGuid& InInstanceId, EAvaPlaybackStatus InNewStatus)
+			{
+				const TOptional<EAvaPlaybackStatus> PreviousStatusOpt = AssetInfo->GetInstanceStatus(InInstanceId);
+				const EAvaPlaybackStatus PreviousStatus = PreviousStatusOpt.IsSet() ? PreviousStatusOpt.GetValue() : EAvaPlaybackStatus::Unknown;
+				
+				AssetInfo->SetInstanceStatus(InInstanceId, InNewStatus);
+
+				UE_LOG(LogAvaPlaybackClient, Verbose,
+				   TEXT("%s Triggered Remote Playback \"%s\" (id:%s) for \"%s\" (%s) Status Changed: %s -> %s because of asset status change."),
+				   *GetBriefFrameInfo(), 
+				   *InMessage.AssetPath.GetAssetName(), *InInstanceId.ToString(), *InMessage.ServerName, *ChannelInfo.Key,
+				   *StaticEnumToString(PreviousStatus), *StaticEnumToString(InNewStatus));
+
+				const FPlaybackStatusChangedArgs Args =
+				{
+					InInstanceId,
+					InMessage.AssetPath,
+					ChannelInfo.Key,
+					InMessage.ServerName,
+					PreviousStatus,
+					InNewStatus
+				};
+				GetOnPlaybackStatusChanged().Broadcast(*this, Args);
+			};
+
+			for (TPair<FGuid, FPlaybackInstanceInfo>& InstanceInfo : AssetInfo->InstanceByIds)
 			{
 				// A missing asset leads to a missing playback.
-				if (InPlaybackStatus == EAvaPlaybackStatus::Available && bAssetIsMissing)
+				if (InstanceInfo.Value.Status == EAvaPlaybackStatus::Available && bAssetIsMissing)
 				{
-					AssetInfo->SetInstanceStatus(InInstanceId, EAvaPlaybackStatus::Missing);
-					GetOnPlaybackStatusChanged().Broadcast(*this, {InInstanceId, InMessage.AssetPath, ChannelInfo.Key, InMessage.ServerName, InPlaybackStatus,EAvaPlaybackStatus::Missing});
+					SetPlaybackInstanceStatus(InstanceInfo.Key, EAvaPlaybackStatus::Missing);
 				}
 				// If the playback was missing, and the asset becomes available, update the playback to available too.
-				else if (InPlaybackStatus == EAvaPlaybackStatus::Missing && bAssetIsAvailable)
+				else if (InstanceInfo.Value.Status == EAvaPlaybackStatus::Missing && bAssetIsAvailable)
 				{
-					AssetInfo->SetInstanceStatus(InInstanceId, EAvaPlaybackStatus::Available);
-					GetOnPlaybackStatusChanged().Broadcast(*this, {InInstanceId, InMessage.AssetPath, ChannelInfo.Key, InMessage.ServerName, InPlaybackStatus,EAvaPlaybackStatus::Available});
+					SetPlaybackInstanceStatus(InstanceInfo.Key, EAvaPlaybackStatus::Available);
 				}
-			};
-
-			for (TPair<FGuid, FPlaybackInstanceInfo>& InstanceInfo : AssetInfo->InstanceByIds)
-			{
-				ConditionalChangeStatus(InstanceInfo.Key, InstanceInfo.Value.Status);
 			}
 		}
 	}
-	
-	UE_LOG(LogAvaPlaybackClient, Verbose,
-		   TEXT("Received Playback Asset Status \"%s\" from server \"%s\" for \"%s\"."),
-		   *EnumToString(InMessage.Status),
-		   *InMessage.ServerName, *InMessage.AssetPath.ToString());
 }
 
-void FAvaPlaybackClient::HandlePlaybackStatusMessage(const FAvaPlaybackStatus& InMessage,
-														  const TSharedRef<IMessageContext, ESPMode::ThreadSafe>& InContext)
+void FAvaPlaybackClient::HandlePlaybackStatus(FServerInfo& InServerInfo,
+	const FGuid& InInstanceId, const FSoftObjectPath& InAssetPath, const FString& InChannelName,
+	EAvaPlaybackStatus InStatus, const FString& InUserData, bool bInValidUserData)
 {
-	FServerInfo& ServerInfo = GetOrCreateServerInfo(InMessage.ServerName, InContext->GetSender());
-	ServerInfo.ResetPingTimeout();
+	using namespace UE::AvaPlayback::Utils;
 	
-	if (InMessage.AssetPath.IsValid())
+	if (!InAssetPath.IsValid())
 	{
-		TOptional<EAvaPlaybackStatus> PrevPlaybackStatus =
-			ServerInfo.GetInstanceStatus(InMessage.ChannelName, InMessage.InstanceId, InMessage.AssetPath);
-		if (!PrevPlaybackStatus.IsSet())
-		{
-			PrevPlaybackStatus = EAvaPlaybackStatus::Unknown;
-		}
-		
-		ServerInfo.SetInstanceStatus(InMessage.ChannelName, InMessage.InstanceId, InMessage.AssetPath, InMessage.Status);
-		
-		if (InMessage.bValidUserData)
-		{
-			ServerInfo.SetInstanceUserData(InMessage.ChannelName, InMessage.InstanceId, InMessage.AssetPath, InMessage.UserData);
-		}
+		UE_LOG(LogAvaPlaybackClient, Error,
+		   TEXT("%s Received Remote Playback \"%s\" (id:%s) from \"%s\" (%s) Status: %s. Invalid asset path, can't update status."),
+		   *GetBriefFrameInfo(), *InAssetPath.GetAssetName(), *InInstanceId.ToString(), *InServerInfo.ServerName, *InChannelName, *StaticEnumToString(InStatus));
+		return;
+	}
 
-		// Clear the pending request (if any).
-		PendingPlaybackStatusRequests.Remove(GetPlaybackStatusRequestKey(InMessage.InstanceId, InMessage.AssetPath, InMessage.ChannelName));
+	TOptional<EAvaPlaybackStatus> PrevPlaybackStatus = InServerInfo.GetInstanceStatus(InChannelName, InInstanceId, InAssetPath);
+	if (!PrevPlaybackStatus.IsSet())
+	{
+		PrevPlaybackStatus = EAvaPlaybackStatus::Unknown;
+	}
+	
+	InServerInfo.SetInstanceStatus(InChannelName, InInstanceId, InAssetPath, InStatus);
+	
+	if (bInValidUserData)
+	{
+		InServerInfo.SetInstanceUserData(InChannelName, InInstanceId, InAssetPath, InUserData);
 
-		UE::AvaPlaybackClient::Delegates::GetOnPlaybackStatusChanged().Broadcast(
-			*this, {InMessage.InstanceId, InMessage.AssetPath, InMessage.ChannelName, InMessage.ServerName, PrevPlaybackStatus.GetValue(), InMessage.Status});
-		
 		UE_LOG(LogAvaPlaybackClient, Verbose,
-			   TEXT("Received Playback Status \"%s\" from server.channel \"%s.%s\" for \"%s\"."),
-			   *StaticEnum<EAvaPlaybackStatus>()->GetValueAsString(InMessage.Status),
-			   *InMessage.ServerName, *InMessage.ChannelName, *InMessage.AssetPath.ToString());
+		   TEXT("%s Received Remote Playback \"%s\" (id:%s) from \"%s\" (%s) User Data: %s."),
+		   *GetBriefFrameInfo(), *InAssetPath.GetAssetName(), *InInstanceId.ToString(), *InServerInfo.ServerName, *InChannelName, *InUserData);
 	}
-	else
+
+	// Clear the pending request (if any).
+	PendingPlaybackStatusRequests.Remove(GetPlaybackStatusRequestKey(InInstanceId, InAssetPath, InChannelName));
+
+	UE_LOG(LogAvaPlaybackClient, Verbose,
+		   TEXT("%s Received Remote Playback \"%s\" (id:%s) from \"%s\" (%s) Status Change: %s -> %s."),
+		   *GetBriefFrameInfo(), 
+		   *InAssetPath.GetAssetName(), *InInstanceId.ToString(), *InServerInfo.ServerName, *InChannelName,
+		   *StaticEnumToString(PrevPlaybackStatus.GetValue()), *StaticEnumToString(InStatus));
+
+	using namespace UE::AvaPlaybackClient::Delegates;
+	const FPlaybackStatusChangedArgs Args =
 	{
-		UE_LOG(LogAvaPlaybackClient, Error,
-			   TEXT(
-				   "Received Playback Status \"%s\" from server.channel \"%s.%s\" with invalid asset path, can't update status."
-			   ),
-			   *StaticEnum<EAvaPlaybackStatus>()->GetValueAsString(InMessage.Status), *InMessage.ServerName,
-			   *InMessage.ChannelName);
-	}
+		InInstanceId,
+		InAssetPath,
+		InChannelName,
+		InServerInfo.ServerName,
+		PrevPlaybackStatus.GetValue(),
+		InStatus
+	};	
+	GetOnPlaybackStatusChanged().Broadcast(*this, Args);
+}
+
+void FAvaPlaybackClient::HandlePlaybackStatusMessage(const FAvaPlaybackStatus& InMessage,
+														  const TSharedRef<IMessageContext, ESPMode::ThreadSafe>& InContext)
+{
+	FServerInfo& ServerInfo = GetOrCreateServerInfo(InMessage.ServerName, InContext->GetSender());
+	ServerInfo.ResetPingTimeout();
+
+	HandlePlaybackStatus(ServerInfo, InMessage.InstanceId, InMessage.AssetPath, InMessage.ChannelName, InMessage.Status, InMessage.UserData, InMessage.bValidUserData);
 }
 
 void FAvaPlaybackClient::HandlePlaybackStatusesMessage(const FAvaPlaybackStatuses& InMessage,
@@ -1129,37 +1218,8 @@ void FAvaPlaybackClient::HandlePlaybackStatusesMessage(const FAvaPlaybackStatuse
 	{
 		const FSoftObjectPath AssetPath = InMessage.AssetPaths.IsValidIndex(StatusIndex) ? InMessage.AssetPaths[StatusIndex] : FSoftObjectPath();
 		const FGuid InstanceId = InMessage.InstanceIds.IsValidIndex(StatusIndex) ? InMessage.InstanceIds[StatusIndex] : FGuid();
-		
-		if (AssetPath.IsValid())
-		{
-			TOptional<EAvaPlaybackStatus> PrevPlaybackStatus = ServerInfo.GetInstanceStatus(InMessage.ChannelName, InstanceId, AssetPath);
-			if (!PrevPlaybackStatus.IsSet())
-			{
-				PrevPlaybackStatus = EAvaPlaybackStatus::Unknown;
-			}
-
-			ServerInfo.SetInstanceStatus(InMessage.ChannelName, InstanceId, AssetPath, InMessage.Status);
 
-			// Clear the pending request (if any).
-			PendingPlaybackStatusRequests.Remove(GetPlaybackStatusRequestKey(InstanceId, AssetPath, InMessage.ChannelName));
-
-			UE::AvaPlaybackClient::Delegates::GetOnPlaybackStatusChanged().Broadcast(
-				*this, {InstanceId, AssetPath, InMessage.ChannelName, InMessage.ServerName, PrevPlaybackStatus.GetValue(), InMessage.Status});
-			
-			UE_LOG(LogAvaPlaybackClient, Verbose,
-				   TEXT("Received Playback Status \"%s\" from server.channel \"%s.%s\" for \"%s\"."),
-				   *StaticEnum<EAvaPlaybackStatus>()->GetValueAsString(InMessage.Status),
-				   *InMessage.ServerName, *InMessage.ChannelName, *AssetPath.ToString());
-		}
-		else
-		{
-			UE_LOG(LogAvaPlaybackClient, Error,
-				   TEXT(
-					   "Received Playback Status \"%s\" from server.channel \"%s.%s\" with invalid asset path, can't update status."
-				   ),
-				   *StaticEnum<EAvaPlaybackStatus>()->GetValueAsString(InMessage.Status), *InMessage.ServerName,
-				   *InMessage.ChannelName);
-		}
+		HandlePlaybackStatus(ServerInfo, InstanceId, AssetPath, InMessage.ChannelName, InMessage.Status);
 	}
 }
 
@@ -1173,7 +1233,8 @@ void FAvaPlaybackClient::HandlePlaybackSequenceEventMessage(const FAvaPlaybackSe
 		InMessage.ChannelName,
 		InMessage.ServerName,
 		InMessage.SequenceName,
-		InMessage.EventType
+		InMessage.EventType,
+		InMessage.FrameNumber
 	};
 	GetOnPlaybackSequenceEvent().Broadcast(*this, Args);
 }
@@ -1187,7 +1248,8 @@ void FAvaPlaybackClient::HandlePlaybackTransitionEventMessage(const FAvaPlayback
 		InMessage.InstanceId,
 		InMessage.ChannelName,
 		InMessage.ServerName,
-		InMessage.GetEventFlags()
+		InMessage.GetEventFlags(),
+		InMessage.FrameNumber
 	};
 	GetOnPlaybackTransitionEvent().Broadcast(*this, Args);
 }
@@ -1411,7 +1473,7 @@ void FAvaPlaybackClient::SetUserDataCommand(const TArray<FString>& InArgs)
 
 void FAvaPlaybackClient::ShowStatusCommand(const TArray<FString>& InArgs)
 {
-	using namespace UE::AvaPlaybackClient::Private;
+	using namespace UE::AvaPlayback::Utils;
 	UE_LOG(LogAvaPlaybackClient, Display, TEXT("Playback Client: \"%s\""), *ComputerName);
 	UE_LOG(LogAvaPlaybackClient, Display, TEXT("- Endpoint Bus Address: \"%s\""), MessageEndpoint.IsValid() ? *MessageEndpoint->GetAddress().ToString() : TEXT("Invalid"));
 	UE_LOG(LogAvaPlaybackClient, Display, TEXT("- ProcessId: %d"), ProcessId);
@@ -1441,18 +1503,18 @@ void FAvaPlaybackClient::ShowStatusCommand(const TArray<FString>& InArgs)
 			check(ChannelInfo.Value.IsValid());
 			UE_LOG(LogAvaPlaybackClient, Display, TEXT("   - Channel(\"%s\") Channel State: \"%s\"."),
 				*ChannelInfo.Key,
-				*EnumToString(ChannelInfo.Value->ChannelState));
+				*StaticEnumToString(ChannelInfo.Value->ChannelState));
 			UE_LOG(LogAvaPlaybackClient, Display, TEXT("   - Channel(\"%s\") Channel Issue Severity: \"%s\"."),
 				*ChannelInfo.Key,
-				*EnumToString(ChannelInfo.Value->ChannelIssueSeverity));
+				*StaticEnumToString(ChannelInfo.Value->ChannelIssueSeverity));
 			for (const TPair<FGuid, FAvaBroadcastOutputStatus>& MediaOutputStatus : ChannelInfo.Value->MediaOutputStatuses)
 			{
 				UE_LOG(LogAvaPlaybackClient, Display, TEXT("   - Channel(\"%s\") Media Output \"%s\" State: \"%s\"."),
 					*ChannelInfo.Key, *MediaOutputStatus.Key.ToString(),
-					*EnumToString(MediaOutputStatus.Value.MediaOutputState));
+					*StaticEnumToString(MediaOutputStatus.Value.MediaOutputState));
 				UE_LOG(LogAvaPlaybackClient, Display, TEXT("   - Channel(\"%s\") Media Output \"%s\" Issue Severity: \"%s\"."),
 					*ChannelInfo.Key, *MediaOutputStatus.Key.ToString(),
-					*EnumToString(MediaOutputStatus.Value.MediaIssueSeverity));
+					*StaticEnumToString(MediaOutputStatus.Value.MediaIssueSeverity));
 				for (const FString& MediaIssueMessage : MediaOutputStatus.Value.MediaIssueMessages)
 				{
 					UE_LOG(LogAvaPlaybackClient, Display, TEXT("   - Channel(\"%s\") Media Output \"%s\" Message: \"%s\"."),
@@ -1477,7 +1539,7 @@ void FAvaPlaybackClient::ShowStatusCommand(const TArray<FString>& InArgs)
 						*ChannelInfoEntry.Key,
 						*AssetInfo.Key.ToString(),
 						*InstanceInfo.Key.ToString(),
-						*EnumToString(InstanceInfo.Value.Status),
+						*StaticEnumToString(InstanceInfo.Value.Status),
 						*PrettyPrintUserData);
 				}
 			}
@@ -1486,7 +1548,7 @@ void FAvaPlaybackClient::ShowStatusCommand(const TArray<FString>& InArgs)
 		{
 			UE_LOG(LogAvaPlaybackClient, Display, TEXT("   - Asset (\"%s\") Status \"%s\"."),
 				*AssetStatus.Key.ToString(),
-				*EnumToString(AssetStatus.Value));
+				*StaticEnumToString(AssetStatus.Value));
 		}
 	}
 }
@@ -1583,7 +1645,7 @@ void FAvaPlaybackClient::OnPackageSaved(const FString& InPackageFileName, UPacka
 			if (InServerInfo.PlaybackAssetStatuses.Contains(Assets.ToSoftObjectPath()))
 			{
 				// Note: we force a refresh (i.e. disregard the server's cached value) because the asset has changed on client side.
-				RequestPlaybackAssetStatusForServer(Assets.ToSoftObjectPath(), InServerInfo.ServerName, true);
+				RequestPlaybackAssetStatusForServer(Assets.ToSoftObjectPath(), InServerInfo.ServerName, /*bInForceRefresh*/ true);
 			}
 		}
 	});
@@ -1705,8 +1767,7 @@ void FAvaPlaybackClient::RemoveDeadServers(const FDateTime& InCurrentTime)
 {
 	bool bServerRemoved = false;
 
-	for (TMap<FString, TSharedPtr<FServerInfo>>::TIterator ServerIter = Servers.CreateIterator(); ServerIter; ++
-		 ServerIter)
+	for (TMap<FString, TSharedPtr<FServerInfo>>::TIterator ServerIter(Servers); ServerIter; ++ServerIter)
 	{
 		check(ServerIter.Value().IsValid());
 		if (ServerIter.Value()->HasTimedOut(InCurrentTime))
@@ -1756,48 +1817,49 @@ TArray<FMessageAddress> FAvaPlaybackClient::GetServerAddressesForChannel(const F
 	return AllServerAddresses;
 }
 
-TArray<FString> FAvaPlaybackClient::GetServerNamesForChannel(const FName& InChannelName) const
+TArray<FString> FAvaPlaybackClient::GetServerNamesForChannel(const FName& InChannelName, bool bInOnlineOnly) const
 {
-	TArray<FString> ServerNames;
+	if (InChannelName.IsNone())
+	{
+		// If no channel is specified, we will return all the servers that are connected.
+		return GetServerNames();
+	}
+	
+	TArray<FString> OutServerNames;
 
-	if (!InChannelName.IsNone())
+	const FAvaBroadcastOutputChannel& Channel = UAvaBroadcast::Get().GetCurrentProfile().GetChannel(InChannelName);
+	const TArray<UMediaOutput*>& RemoteOutputs = Channel.GetRemoteMediaOutputs();
+	OutServerNames.Reserve(RemoteOutputs.Num());
+	for (const UMediaOutput* RemoteOutput : RemoteOutputs)
 	{
-		const FAvaBroadcastOutputChannel& Channel = UAvaBroadcast::Get().GetCurrentProfile().GetChannel(InChannelName);
-		const TArray<UMediaOutput*>& RemoteOutputs = Channel.GetRemoteMediaOutputs();
-		ServerNames.Reserve(RemoteOutputs.Num());
-		for (const UMediaOutput* RemoteOutput : RemoteOutputs)
+		const FAvaBroadcastMediaOutputInfo& OutputInfo = Channel.GetMediaOutputInfo(RemoteOutput);
+		if (OutputInfo.IsValid())
+		{
+			if (!bInOnlineOnly || Servers.Contains(OutputInfo.ServerName))
+			{
+				OutServerNames.AddUnique(OutputInfo.ServerName);
+			}
+		}
+		else
 		{
-			const FAvaBroadcastMediaOutputInfo& OutputInfo = Channel.GetMediaOutputInfo(RemoteOutput);
-			if (OutputInfo.IsValid())
+			UE_LOG(LogAvaPlaybackClient, Warning, TEXT("MediaOutputInfo invalid for channel \"%s\"."),
+				   *InChannelName.ToString());
+
+			// Try to find the server name from the device name. The server has to be online for this to work.
+			const FString ServerName = GetServerNameForMediaOutputFallback(RemoteOutput);
+			if (!ServerName.IsEmpty())
 			{
-				ServerNames.Add(OutputInfo.ServerName);
+				OutServerNames.AddUnique(ServerName);
 			}
 			else
 			{
-				UE_LOG(LogAvaPlaybackClient, Warning, TEXT("MediaOutputInfo invalid for channel \"%s\"."),
+				UE_LOG(LogAvaPlaybackClient, Error,
+					   TEXT("Unable to find server name for remote MediaOutput for channel \"%s\"."),
 					   *InChannelName.ToString());
-
-				// Try to find the server name from the device name. The server has to be online for this to work.
-				const FString ServerName = GetServerNameForMediaOutputFallback(RemoteOutput);
-				if (!ServerName.IsEmpty())
-				{
-					ServerNames.Add(ServerName);
-				}
-				else
-				{
-					UE_LOG(LogAvaPlaybackClient, Error,
-						   TEXT("Unable to find server name for remote MediaOutput for channel \"%s\"."),
-						   *InChannelName.ToString());
-				}
 			}
 		}
 	}
-	else
-	{
-		// If no channel is specified, we will return all the servers that are connected.
-		ServerNames = GetServerNames();
-	}
-	return ServerNames;
+	return OutServerNames;
 }
 
 FString FAvaPlaybackClient::GetServerNameForMediaOutputFallback(const UMediaOutput* InMediaOutput) const
diff --git a/Engine/Plugins/Experimental/Avalanche/Source/AvalancheMedia/Private/Playback/AvaPlaybackClient.h b/Engine/Plugins/Experimental/Avalanche/Source/AvalancheMedia/Private/Playback/AvaPlaybackClient.h
--- a/Engine/Plugins/Experimental/Avalanche/Source/AvalancheMedia/Private/Playback/AvaPlaybackClient.h
+++ b/Engine/Plugins/Experimental/Avalanche/Source/AvalancheMedia/Private/Playback/AvaPlaybackClient.h
@@ -45,6 +45,7 @@ public:
 	virtual const TArray<FString>& GetMediaOutputIssueMessages(const FString& InServerName, const FString& InChannelName, const FGuid& InOutputGuid) const override;
 	virtual EAvaBroadcastOutputState GetMediaOutputState(const FString& InServerName, const FString& InChannelName, const FGuid& InOutputGuid) const override;
 	virtual bool HasAnyServerOnlineForChannel(const FName& InChannelName) const override;
+	virtual TArray<FString> GetOnlineServersForChannel(const FName& InChannelName) const override;
 	virtual TOptional<EAvaPlaybackStatus> GetRemotePlaybackStatus(const FGuid& InInstanceId, const FSoftObjectPath& InAssetPath, const FString& InChannelName, const FString& InServerName) const override;
 	virtual const FString* GetRemotePlaybackUserData(const FGuid& InInstanceId, const FSoftObjectPath& InAssetPath, const FString& InChannelName, const FString& InServerName) const override;
 	virtual TOptional<EAvaPlaybackAssetStatus> GetRemotePlaybackAssetStatus(const FSoftObjectPath& InAssetPath, const FString& InServerName) const override;
@@ -76,7 +77,15 @@ protected:
 	void HandlePlaybackStatusesMessage(const FAvaPlaybackStatuses& InMessage, const TSharedRef<IMessageContext, ESPMode::ThreadSafe>& InContext);
 	void HandlePlaybackSequenceEventMessage(const FAvaPlaybackSequenceEvent& InMessage, const TSharedRef<IMessageContext, ESPMode::ThreadSafe>& InContext);
 	void HandlePlaybackTransitionEventMessage(const FAvaPlaybackTransitionEvent& InMessage, const TSharedRef<IMessageContext, ESPMode::ThreadSafe>& InContext);
+
+private:
+	class FServerInfo;
 	
+	void HandlePlaybackStatus(FServerInfo& InServerInfo,
+		const FGuid& InInstanceId, const FSoftObjectPath& InAssetPath, const FString& InChannelName,
+		EAvaPlaybackStatus InStatus, const FString& InUserData = FString(), bool bInValidUserData = false);
+
+protected:	
 	void RegisterCommands();
 	
 	// Command handlers
@@ -155,7 +164,7 @@ protected:
 	 * This is based on the channel's MediaOutputInfos. It will properly work
 	 * regardless of the channel's status (idle or even offline).
 	 */
-	TArray<FString> GetServerNamesForChannel(const FName& InChannelName) const;
+	TArray<FString> GetServerNamesForChannel(const FName& InChannelName, bool bInOnlineOnly = false) const;
 
 	FString GetServerNameForMediaOutputFallback(const UMediaOutput* InMediaOutput) const;
 
diff --git a/Engine/Plugins/Experimental/Avalanche/Source/AvalancheMedia/Private/Playback/AvaPlaybackClientDummy.h b/Engine/Plugins/Experimental/Avalanche/Source/AvalancheMedia/Private/Playback/AvaPlaybackClientDummy.h
--- a/Engine/Plugins/Experimental/Avalanche/Source/AvalancheMedia/Private/Playback/AvaPlaybackClientDummy.h
+++ b/Engine/Plugins/Experimental/Avalanche/Source/AvalancheMedia/Private/Playback/AvaPlaybackClientDummy.h
@@ -39,6 +39,7 @@ public:
 	virtual const TArray<FString>& GetMediaOutputIssueMessages(const FString& InServerName, const FString& InChannelName, const FGuid& InOutputGuid) const override { return EmptyStringArray; }
 	virtual EAvaBroadcastOutputState GetMediaOutputState(const FString& InServerName, const FString& InChannelName, const FGuid& InOutputGuid) const override { return EAvaBroadcastOutputState::Error; }
 	virtual bool HasAnyServerOnlineForChannel(const FName& InChannelName) const override { return false;}
+	virtual TArray<FString> GetOnlineServersForChannel(const FName& InChannelName) const override { return TArray<FString>(); };
    	virtual TOptional<EAvaPlaybackStatus> GetRemotePlaybackStatus(const FGuid& InInstanceId, const FSoftObjectPath& InAssetPath, const FString& InChannelName, const FString& InServerName) const override
 	{
 		return EAvaPlaybackStatus::Unknown;
diff --git a/Engine/Plugins/Experimental/Avalanche/Source/AvalancheMedia/Private/Playback/AvaPlaybackGraph.cpp b/Engine/Plugins/Experimental/Avalanche/Source/AvalancheMedia/Private/Playback/AvaPlaybackGraph.cpp
--- a/Engine/Plugins/Experimental/Avalanche/Source/AvalancheMedia/Private/Playback/AvaPlaybackGraph.cpp
+++ b/Engine/Plugins/Experimental/Avalanche/Source/AvalancheMedia/Private/Playback/AvaPlaybackGraph.cpp
@@ -14,6 +14,7 @@
 #include "Playable/AvaPlayable.h"
 #include "Playable/AvaPlayableGroup.h"
 #include "Playback/AvaPlaybackManager.h"
+#include "Playback/AvaPlaybackUtils.h"
 #include "Playback/IAvaPlaybackGraphEditor.h"
 #include "Playback/Nodes/AvaPlaybackNode.h"
 
@@ -188,6 +189,7 @@ UAvaPlayable* UAvaPlaybackGraph::FindOrLoadPlayable(const FAvaSoftAssetPtr& InSo
 		return *FoundPlayable;
 	}
 
+	// TODO: Extension point IAvaPlayableFactory -- Example: IMediaPlayerFactory. Using URL (> FAvaSoftAssetPath).
 	UAvaPlayable* NewPlayable = UAvaPlayable::Create(this, {GetPlayableGroupManager(), InSourceAsset, InChannelName});
 	if (NewPlayable)
 	{
@@ -654,8 +656,8 @@ void UAvaPlaybackGraph::ExecutePendingAnimationCommands()
 		if (AnimCommand.HasTimedOut(CurrentTime))
 		{
 			UE_LOG(LogAvaPlayback, Warning,
-				TEXT("Animation command for asset \"%s\" on channel \"%s\" timed out and is discarded."),
-				*AnimCommand.SourcePath.ToString(), *AnimCommand.ChannelName);
+				TEXT("%s Animation command for asset \"%s\" on channel \"%s\" timed out and is discarded."),
+				*UE::AvaPlayback::Utils::GetBriefFrameInfo(), *AnimCommand.SourcePath.ToString(), *AnimCommand.ChannelName);
 			continue;
 		}
 
@@ -696,8 +698,8 @@ void UAvaPlaybackGraph::ExecutePendingRemoteControlCommands()
 		if (RemoteControlCommand.HasTimedOut(CurrentTime))
 		{
 			UE_LOG(LogAvaPlayback, Warning,
-				TEXT("Remote Control command for asset \"%s\" on channel \"%s\" timed out and is discarded."),
-				*RemoteControlCommand.SourcePath.ToString(), *RemoteControlCommand.ChannelName);
+				TEXT("%s Remote Control command for asset \"%s\" on channel \"%s\" timed out and is discarded."),
+				*UE::AvaPlayback::Utils::GetBriefFrameInfo(), *RemoteControlCommand.SourcePath.ToString(), *RemoteControlCommand.ChannelName);
 			continue;
 		}
 		
diff --git a/Engine/Plugins/Experimental/Avalanche/Source/AvalancheMedia/Private/Playback/AvaPlaybackManager.cpp b/Engine/Plugins/Experimental/Avalanche/Source/AvalancheMedia/Private/Playback/AvaPlaybackManager.cpp
--- a/Engine/Plugins/Experimental/Avalanche/Source/AvalancheMedia/Private/Playback/AvaPlaybackManager.cpp
+++ b/Engine/Plugins/Experimental/Avalanche/Source/AvalancheMedia/Private/Playback/AvaPlaybackManager.cpp
@@ -128,6 +128,19 @@ namespace UE::AvaPlaybackManager::Private
 			return true;
 		}
 	};
+
+	FString GetPrettyPlaybackInstanceInfo(const FAvaPlaybackInstance* InPlaybackInstance)
+	{
+		if (InPlaybackInstance)
+		{
+			return FString::Printf(TEXT("Id:%s, Asset:%s, Channel:%s, UserData:\"%s\""),
+				*InPlaybackInstance->GetInstanceId().ToString(),
+				*InPlaybackInstance->GetSourcePath().GetAssetName(),
+				*InPlaybackInstance->GetChannelName(),
+				*InPlaybackInstance->GetInstanceUserData());
+		}
+		return TEXT("");
+	}
 }
 
 FAvaPlaybackManager::FAvaPlaybackManager()
@@ -847,67 +860,90 @@ void FAvaPlaybackInstance::SetInstanceUserData(const FString& InUserData)
 bool FAvaPlaybackInstance::UpdateStatus()
 {
 	using namespace UE::AvaPlaybackManager::Private;
-	if (IsPlaybackValid(Playback))
+	using namespace UE::AvaPlayback::Utils;
+	
+	if (!IsPlaybackValid(Playback))
 	{
-		EAvaPlaybackStatus NewStatus = Status;
+		return false;
+	}
+	
+	EAvaPlaybackStatus NewStatus = Status;
 
-		if (Playback->IsPlaying())
+	if (Playback->IsPlaying())
+	{
+		if (const UAvaPlayable* Playable = Playback->FindPlayable(SourcePath, ChannelFName))
 		{
-			if (const UAvaPlayable* Playable = Playback->FindPlayable(SourcePath, ChannelFName))
+			switch (Playable->GetPlayableStatus())
 			{
-				switch (Playable->GetPlayableStatus())
-				{
-				case EAvaPlayableStatus::Unloaded:
-					NewStatus = EAvaPlaybackStatus::Available;
-					break;
-				case EAvaPlayableStatus::Loading:
-					NewStatus = EAvaPlaybackStatus::Loading;
-					break;
-				case EAvaPlayableStatus::Loaded:
-					NewStatus = EAvaPlaybackStatus::Loaded;
-					break;
-				case EAvaPlayableStatus::Visible:
-					NewStatus = Playable->GetPlayableGroup()->IsRenderTargetReady() ? EAvaPlaybackStatus::Started : EAvaPlaybackStatus::Starting;
-					break;
-				}
-			}
-			else
-			{
-				// The game instance may not be created yet. The RefreshPlayback is done on the next tick.
+			case EAvaPlayableStatus::Unloaded:
+				NewStatus = EAvaPlaybackStatus::Available;
+				break;
+			case EAvaPlayableStatus::Loading:
 				NewStatus = EAvaPlaybackStatus::Loading;
+				break;
+			case EAvaPlayableStatus::Loaded:
+				NewStatus = EAvaPlaybackStatus::Loaded;
+				break;
+			case EAvaPlayableStatus::Visible:
+				NewStatus = Playable->GetPlayableGroup()->IsRenderTargetReady() ? EAvaPlaybackStatus::Started : EAvaPlaybackStatus::Starting;
+				break;
 			}
 		}
 		else
 		{
-			// Even if not playing, we could have a game instance already, it can be preloaded now.
-			if (const UAvaPlayable* Playable = Playback->FindPlayable(SourcePath, ChannelFName))
+			// The game instance may not be created yet. The RefreshPlayback is done on the next tick.
+			NewStatus = EAvaPlaybackStatus::Loading;
+
+			// Reporting possible edge case:
+			if (Status != NewStatus)
 			{
-				switch (Playable->GetPlayableStatus())
-				{
-				case EAvaPlayableStatus::Unloaded:
-					NewStatus = EAvaPlaybackStatus::Available;
-					break;
-				case EAvaPlayableStatus::Loading:
-					NewStatus = EAvaPlaybackStatus::Loading;
-					break;
-				case EAvaPlayableStatus::Loaded:
-				case EAvaPlayableStatus::Visible:
-					NewStatus = EAvaPlaybackStatus::Loaded;
-					break;
-				}
+				UE_LOG(LogAvaPlaybackManager, Verbose,
+					TEXT("%s Playback Instance {%s} Status Changed: %s -> %s because Playable \"%s\" (%s) was not found in PLAYING instance."),
+					*GetBriefFrameInfo(), *GetPrettyPlaybackInstanceInfo(this), *StaticEnumToString(Status), *StaticEnumToString(NewStatus),
+					*SourcePath.GetAssetName(), *ChannelName);
 			}
-			else
+		}
+	}
+	else
+	{
+		// Even if not playing, we could have a game instance already, it can be preloaded now.
+		if (const UAvaPlayable* Playable = Playback->FindPlayable(SourcePath, ChannelFName))
+		{
+			switch (Playable->GetPlayableStatus())
 			{
+			case EAvaPlayableStatus::Unloaded:
+				NewStatus = EAvaPlaybackStatus::Available;
+				break;
+			case EAvaPlayableStatus::Loading:
 				NewStatus = EAvaPlaybackStatus::Loading;
+				break;
+			case EAvaPlayableStatus::Loaded:
+			case EAvaPlayableStatus::Visible:
+				NewStatus = EAvaPlaybackStatus::Loaded;
+				break;
 			}
 		}
-
-		if (NewStatus != Status)
+		else
 		{
-			Status = NewStatus;
-			return true;
+			NewStatus = EAvaPlaybackStatus::Loading;
+			
+			// Reporting possible edge case:
+			if (Status != NewStatus)
+			{
+				UE_LOG(LogAvaPlaybackManager, Verbose,
+					TEXT("%s Playback Instance {%s} Status Changed: %s -> %s because Playable \"%s\" (%s) was not found in instance."),
+					*GetBriefFrameInfo(), *GetPrettyPlaybackInstanceInfo(this), *StaticEnumToString(Status), *StaticEnumToString(NewStatus),
+					*SourcePath.GetAssetName(), *ChannelName);
+			}
 		}
 	}
+
+	if (NewStatus != Status)
+	{
+		Status = NewStatus;
+		return true;
+	}
+	
 	return false;
 }
 
diff --git a/Engine/Plugins/Experimental/Avalanche/Source/AvalancheMedia/Private/Playback/AvaPlaybackMessages.h b/Engine/Plugins/Experimental/Avalanche/Source/AvalancheMedia/Private/Playback/AvaPlaybackMessages.h
--- a/Engine/Plugins/Experimental/Avalanche/Source/AvalancheMedia/Private/Playback/AvaPlaybackMessages.h
+++ b/Engine/Plugins/Experimental/Avalanche/Source/AvalancheMedia/Private/Playback/AvaPlaybackMessages.h
@@ -422,6 +422,9 @@ struct FAvaPlaybackSequenceEvent : public FAvaPlaybackServerMessageBase
 
 	UPROPERTY()
 	EAvaPlayableSequenceEventType EventType = EAvaPlayableSequenceEventType::None;
+
+	UPROPERTY()
+	int32 FrameNumber = 0;
 };
 
 USTRUCT()
@@ -510,6 +513,9 @@ struct FAvaPlaybackTransitionEvent : public FAvaPlaybackServerMessageBase
 	UPROPERTY()
 	uint8 EventFlags = static_cast<uint8>(EAvaPlayableTransitionEventFlags::None);
 
+	UPROPERTY()
+	int32 FrameNumber = 0;
+
 	// Because the enum is used as flags, we need to convert to uint8 manually. (Can't use TEnumAsByte for this apparently.)
 	EAvaPlayableTransitionEventFlags GetEventFlags() const { return static_cast<EAvaPlayableTransitionEventFlags>(EventFlags);}
 	void SetEventFlags(EAvaPlayableTransitionEventFlags InFlags) { EventFlags = static_cast<uint8>(InFlags);}
diff --git a/Engine/Plugins/Experimental/Avalanche/Source/AvalancheMedia/Private/Playback/AvaPlaybackServer.cpp b/Engine/Plugins/Experimental/Avalanche/Source/AvalancheMedia/Private/Playback/AvaPlaybackServer.cpp
--- a/Engine/Plugins/Experimental/Avalanche/Source/AvalancheMedia/Private/Playback/AvaPlaybackServer.cpp
+++ b/Engine/Plugins/Experimental/Avalanche/Source/AvalancheMedia/Private/Playback/AvaPlaybackServer.cpp
@@ -2,7 +2,6 @@
 
 #include "Playback/AvaPlaybackServer.h"
 
-#include "Async/Async.h"
 #include "AvaMediaMessageUtils.h"
 #include "AvaMediaSettings.h"
 #include "AvaPlaybackSyncManager.h"
@@ -10,6 +9,7 @@
 #include "Broadcast/OutputDevices/AvaBroadcastDeviceProviderData.h"
 #include "Broadcast/OutputDevices/AvaBroadcastDeviceProviderProxy.h"
 #include "Broadcast/OutputDevices/AvaBroadcastOutputUtils.h"
+#include "Containers/Ticker.h"
 #include "IAvaModule.h"
 #include "MediaOutput.h"
 #include "MessageEndpointBuilder.h"
@@ -24,29 +24,62 @@ DEFINE_LOG_CATEGORY(LogAvaPlaybackServer);
 namespace UE::AvaPlaybackServer::Private
 {
 	FAvaPlaybackStatus* MakePlaybackStatusMessage(const FGuid& InInstanceId, const FString& InChannelName,
-		const FSoftObjectPath& InAssetPath, EAvaPlaybackStatus InStatus)
+		const FSoftObjectPath& InAssetPath, EAvaPlaybackStatus InStatus, const FString& InUserData, const bool bInUserDataValid = true)
 	{
 		FAvaPlaybackStatus* Message = FMessageEndpoint::MakeMessage<FAvaPlaybackStatus>();
 		Message->InstanceId = InInstanceId;
 		Message->ChannelName = InChannelName;
 		Message->AssetPath = InAssetPath;
 		Message->Status = InStatus;
-		Message->bValidUserData = false;
+		Message->UserData = InUserData;
+		Message->bValidUserData = bInUserDataValid;
 		return Message;
 	}
-	
+
 	FAvaPlaybackStatus* MakePlaybackStatusMessage(const FGuid& InInstanceId, const FString& InChannelName,
-		const FSoftObjectPath& InAssetPath, EAvaPlaybackStatus InStatus, const FString& InUserData)
+		const FSoftObjectPath& InAssetPath, EAvaPlaybackStatus InStatus)
 	{
-		FAvaPlaybackStatus* Message = FMessageEndpoint::MakeMessage<FAvaPlaybackStatus>();
-		Message->InstanceId = InInstanceId;
-		Message->ChannelName = InChannelName;
-		Message->AssetPath = InAssetPath;
-		Message->Status = InStatus;
-		Message->bValidUserData = true;
-		Message->UserData = InUserData;
-		return Message;
+		return MakePlaybackStatusMessage(InInstanceId, InChannelName, InAssetPath, InStatus, FString(), /*bInUserDataValid*/ false);
+	}
+
+	FString GetCommandActionString(const FAvaPlaybackCommand& InCommand)
+	{
+		FString ActionString = AvaPlayback::Utils::StaticEnumToString(InCommand.Action); 
+		if (InCommand.Arguments.IsEmpty())
+		{
+			return ActionString;
+		}
+
+		return FString::Printf(TEXT("%s \"%s\""), *ActionString, *InCommand.Arguments);
+	};
+
+	/** Returns the command priority order of execution. */
+	int32 GetCommandActionPriority(EAvaPlaybackAction InAction)
+	{
+		switch (InAction)
+		{
+		case EAvaPlaybackAction::None:			return 10;
+		case EAvaPlaybackAction::Load:			return 1;
+		case EAvaPlaybackAction::Start:			return 2;
+		case EAvaPlaybackAction::Stop:			return 5;
+		case EAvaPlaybackAction::Unload:		return 6;
+		case EAvaPlaybackAction::Status:		return 7;		// We want the status after all other commands have been executed.
+		case EAvaPlaybackAction::SetUserData:	return 3;		// Should be after Load and Start.
+		case EAvaPlaybackAction::GetUserData:	return 4;		// Should be after "set user data". 
+		default:
+			return 10;
+		}
 	}
+
+	// Simulate network latency by delaying playback commands a random amount.
+	// This is used to cause desynchronization between the nodes on clustered rendering
+	// and see how well the synchronization handles it.
+	TAutoConsoleVariable<float> CVarTestMaxRandomWaitForPlaybackCommands(
+		TEXT("MotionDesignPlaybackServer.Test.MaxRandomWaitForPlaybackCommands")
+		, 0.0f
+		, TEXT("if not zero, the server will wait a random duration between 0 and the specified delay before executing playback commands. Unit: seconds"), ECVF_Default);
+
+	FRandomStream GPlaybackServerRandomStream((int32)FPlatformTime::Cycles());
 }
 
 /**
@@ -400,6 +433,7 @@ void FAvaPlaybackServer::SendPlayableTransitionEvent(
 	Message->ChannelName = InChannelName.ToString();
 	Message->TransitionId = InTransitionId;
 	Message->InstanceId = InInstanceId;
+	Message->FrameNumber = GFrameNumber;
 	Message->SetEventFlags(InFlags);
 	SendResponse(Message, GetClientAddressSafe(InClientName));
 }
@@ -566,9 +600,28 @@ void FAvaPlaybackServer::HandlePlaybackAssetStatusRequest(const FAvaPlaybackAsse
 
 void FAvaPlaybackServer::HandlePlaybackRequest(const FAvaPlaybackRequest& InMessage, const TSharedRef<IMessageContext, ESPMode::ThreadSafe>& InContext)
 {
+	using namespace UE::AvaPlaybackServer::Private;
+	const float MaxRandomWait = CVarTestMaxRandomWaitForPlaybackCommands.GetValueOnAnyThread();
+
+	const FDateTime UtcNow = FDateTime::UtcNow();
+	
 	for (const FAvaPlaybackCommand& Command : InMessage.Commands)
 	{
-		PendingPlaybackCommands.Add({InContext->GetSender(), Command});
+		TSharedPtr<FPendingPlaybackCommand> PendingCommand
+			= MakeShared<FPendingPlaybackCommand>(UtcNow, GFrameNumber, GetCommandActionPriority(Command.Action), InContext->GetSender(), Command);
+
+		if (MaxRandomWait > 0.0f && (Command.Action == EAvaPlaybackAction::Load || Command.Action == EAvaPlaybackAction::Start))
+		{
+			FTSTicker::GetCoreTicker().AddTicker(FTickerDelegate::CreateSPLambda(this, [this, PendingCommand] (float) 
+			{
+				PendingPlaybackCommands.Add(PendingCommand);	
+				return false;
+			}), GPlaybackServerRandomStream.GetFraction() * MaxRandomWait);
+		}
+		else
+		{
+			PendingPlaybackCommands.Add(MoveTemp(PendingCommand));
+		}
 	}
 }
 
@@ -604,10 +657,10 @@ void FAvaPlaybackServer::HandlePlayableTransitionStartRequest(const FAvaPlayback
 	Transition->SetTransitionFlags(InMessage.GetTransitionFlags());
 
 	// Enter Instances are likely not loaded yet.
-	Transition->SetEnterInstanceIds(InMessage.EnterInstanceIds);
+	Transition->AddPendingEnterInstanceIds(InMessage.EnterInstanceIds);
 	Transition->SetEnterValues(InMessage.EnterValues);
 
-	// We can resolve the playing instances since they should be loaded.
+	// We try to resolve the playing instances since they should be loaded (unless delayed).
 	for (const FGuid& PlayingInstanceId : InMessage.PlayingInstanceIds)
 	{
 		if (TSharedPtr<FAvaPlaybackInstance> Instance = FindActivePlaybackInstance(PlayingInstanceId))
@@ -616,13 +669,15 @@ void FAvaPlaybackServer::HandlePlayableTransitionStartRequest(const FAvaPlayback
 		}
 		else
 		{
-			UE_LOG(LogAvaPlaybackServer, Error,
+			Transition->AddPendingPlayingInstanceId(PlayingInstanceId);	// Will be resolved later.
+
+			UE_LOG(LogAvaPlaybackServer, Warning,
 				TEXT("Transition \"%s\" from client \"%s\": \"Playing\" Instance Id \"%s\" was not found in active playback instances."),
 				*InMessage.TransitionId.ToString(), *GetClientNameSafe(InContext->GetSender()), *PlayingInstanceId.ToString());
 		}
 	}
 	
-	// We can resolve the exit instances since they should be loaded.
+	// We try to resolve the exit instances since they should be loaded (unless delayed).
 	for (const FGuid& ExitInstanceId : InMessage.ExitInstanceIds)
 	{
 		if (TSharedPtr<FAvaPlaybackInstance> Instance = FindActivePlaybackInstance(ExitInstanceId))
@@ -631,7 +686,9 @@ void FAvaPlaybackServer::HandlePlayableTransitionStartRequest(const FAvaPlayback
 		}
 		else
 		{
-			UE_LOG(LogAvaPlaybackServer, Error,
+			Transition->AddPendingExitInstanceId(ExitInstanceId);	// Will be resolved later.
+			
+			UE_LOG(LogAvaPlaybackServer, Warning,
 				TEXT("Transition \"%s\" from client \"%s\": \"Exit\" Instance Id \"%s\" was not found in active playback instances."),
 				*InMessage.TransitionId.ToString(), *GetClientNameSafe(InContext->GetSender()), *ExitInstanceId.ToString());
 		}
@@ -772,8 +829,8 @@ FString FAvaPlaybackServer::GetMessageEndpointAddressId() const
 
 void FAvaPlaybackServer::Tick()
 {
-	const FDateTime CurrentTime = FDateTime::UtcNow();
-	RemoveDeadClients(CurrentTime);
+	const FDateTime CurrentTimeUtc = FDateTime::UtcNow();
+	RemoveDeadClients(CurrentTimeUtc);
 
 	for (const TPair<FString, TSharedPtr<FClientInfo>>& Client : Clients)
 	{
@@ -784,7 +841,7 @@ void FAvaPlaybackServer::Tick()
 	}
 
 	// Execute the pending commands in batch for this tick.
-	ExecutePendingPlaybackCommands();
+	ExecutePendingPlaybackCommands(CurrentTimeUtc);
 
 	// Try to resolve the instance for loaded transitions.
 	for (const TPair<FGuid, TObjectPtr<UAvaPlaybackServerTransition>>& Transition : PlaybackInstanceTransitions->Transitions)
@@ -947,7 +1004,7 @@ void FAvaPlaybackServer::ShowStatusCommand(const TArray<FString>& InArgs)
 			ClientInfo.BroadcastSettings.Settings.ChannelDefaultResolution.Y);
 		UE_LOG(LogAvaPlaybackServer, Display, TEXT("   - BroadcastSettings.bDrawPlaceholderWidget: %s"),
 			ClientInfo.BroadcastSettings.Settings.bDrawPlaceholderWidget ? TEXT("true") : TEXT("false"));
-		UE_LOG(LogAvaPlaybackServer, Display, TEXT("   - BroadcastSettings.PlaceholderWidgetClass: "),
+		UE_LOG(LogAvaPlaybackServer, Display, TEXT("   - BroadcastSettings.PlaceholderWidgetClass: %s"),
 			*ClientInfo.BroadcastSettings.Settings.PlaceholderWidgetClass.ToString());
 
 		UE_LOG(LogAvaPlaybackServer, Display, TEXT("   - MediaSyncManager: %s."),
@@ -972,7 +1029,7 @@ void FAvaPlaybackServer::ShowStatusCommand(const TArray<FString>& InArgs)
 		const TSharedPtr<FAvaPlaybackInstance>& Instance = ActivePlaybackInstance.Value;
 		UE_LOG(LogAvaPlaybackServer, Display, TEXT("   - Id:%s, Channel: %s, Asset: %s, Status: %s, UserData: %s ."),
 			*Instance->GetInstanceId().ToString(), *Instance->GetChannelName(), *Instance->GetSourcePath().ToString(),
-			*StaticEnum<EAvaPlaybackStatus>()->GetNameByValue(static_cast<int32>(Instance->GetStatus())).ToString(),
+			*UE::AvaPlayback::Utils::StaticEnumToString(Instance->GetStatus()),
 			*Instance->GetInstanceUserData());
 	}
 	
@@ -1077,6 +1134,7 @@ void FAvaPlaybackServer::OnPlayableSequenceEvent(UAvaPlayable* InPlayable, const
 	Message->ChannelName = PlaybackInstance->GetChannelName();
 	Message->SequenceName = SequenceName.ToString();
 	Message->EventType = InEventType;
+	Message->FrameNumber = GFrameNumber;
 	SendResponse(Message, ClientAddresses);
 }
 
@@ -1099,17 +1157,6 @@ void FAvaPlaybackServer::ApplyAvaMediaSettings()
 	{
 		ReplicationOutputDevice.Reset();
 	}
-
-#if !NO_LOGGING
-	if (Settings.bVerbosePlaybackServerLogging)
-	{
-		LogAvaPlaybackServer.SetVerbosity(ELogVerbosity::Verbose);
-	}
-	else
-	{
-		LogAvaPlaybackServer.SetVerbosity(ELogVerbosity::Log);
-	}
-#endif
 }
 
 void FAvaPlaybackServer::SendUserDataUpdate(const TArray<FMessageAddress>& InRecipients)
@@ -1193,47 +1240,85 @@ void FAvaPlaybackServer::SendLogMessage(const TCHAR* InText, ELogVerbosity::Type
 	SendResponse(Message, ClientAddresses);
 }
 
-void FAvaPlaybackServer::ExecutePendingPlaybackCommands()
+void FAvaPlaybackServer::ExecutePendingPlaybackCommands(const FDateTime& InUtcNow)
 {
-	for (const FPendingPlaybackCommand& PendingCommand : PendingPlaybackCommands)
+	if (PendingPlaybackCommands.IsEmpty())
 	{
-		const FAvaPlaybackCommand& Command = PendingCommand.Command;
+		return;
+	}
+
+	// Sort pending commands according to priority of execution. (Should minimize the amount of re-scheduling)
+	Algo::Sort(PendingPlaybackCommands, [](const TSharedPtr<FPendingPlaybackCommand>& InA, const TSharedPtr<FPendingPlaybackCommand>& InB)
+	{
+		return InA->Priority < InB->Priority;
+	});
+
+	for (TArray<TSharedPtr<FPendingPlaybackCommand>>::TIterator PendingCommandIt(PendingPlaybackCommands); PendingCommandIt; ++PendingCommandIt)
+	{
+		const FAvaPlaybackCommand& Command = (*PendingCommandIt)->Command;
+		const FMessageAddress& ReplyTo = (*PendingCommandIt)->ReplyTo;
+		bool bReschedule = false;
+		
 		switch (Command.Action)
 		{
 		case EAvaPlaybackAction::None:
 			break;
 			
 		case EAvaPlaybackAction::Load:
-			LoadPlayback(PendingCommand.ReplyTo, Command.InstanceId, Command.ChannelName, Command.AssetPath);
+			LoadPlayback(ReplyTo, Command.InstanceId, Command.ChannelName, Command.AssetPath);
 			break;
 			
 		case EAvaPlaybackAction::Start:
-			StartPlayback(PendingCommand.ReplyTo, Command.InstanceId, Command.ChannelName, Command.AssetPath);
+			StartPlayback(ReplyTo, Command.InstanceId, Command.ChannelName, Command.AssetPath);
 			break;
 			
 		case EAvaPlaybackAction::Stop:
-			StopPlayback(PendingCommand.ReplyTo, Command.InstanceId, Command.ChannelName, Command.AssetPath);
+			StopPlayback(ReplyTo, Command.InstanceId, Command.ChannelName, Command.AssetPath);
 			break;
 			
 		case EAvaPlaybackAction::Unload:
-			UnloadPlayback(PendingCommand.ReplyTo, Command.InstanceId, Command.ChannelName, Command.AssetPath);
+			UnloadPlayback(ReplyTo, Command.InstanceId, Command.ChannelName, Command.AssetPath);
 			break;
 			
 		case EAvaPlaybackAction::Status:
-			SendPlaybackStatus(PendingCommand.ReplyTo, Command.InstanceId, Command.ChannelName, Command.AssetPath);
+			SendPlaybackStatus(ReplyTo, Command.InstanceId, Command.ChannelName, Command.AssetPath);
 			break;
 			
 		case EAvaPlaybackAction::SetUserData:	
-			SetPlaybackUserData(PendingCommand.ReplyTo, Command.InstanceId, Command.Arguments);
+			bReschedule = !SetPlaybackUserData(ReplyTo, Command.InstanceId, Command.Arguments);
 			break;
 
 		case EAvaPlaybackAction::GetUserData:
-			SendPlaybackUserData(PendingCommand.ReplyTo, Command.InstanceId);
+			bReschedule = !SendPlaybackUserData(ReplyTo, Command.InstanceId);
 			break;
 		}
-	}
+		using namespace UE::AvaPlaybackServer::Private;
+		using namespace UE::AvaPlayback::Utils;
+		
+		if (bReschedule)
+		{
+			const UAvaMediaSettings& Settings = UAvaMediaSettings::Get();
+			const float CommandWaitTime = (InUtcNow - (*PendingCommandIt)->ReceivedUtc).GetTotalSeconds();
+			
+			if (CommandWaitTime > Settings.ServerPendingPlaybackCommandTimeout)
+			{
+				UE_LOG(LogAvaPlaybackServer, Warning, TEXT("%s Discarding Playback Command [%s] (Timed out after %f seconds) for asset: \"%s\" (id:%s) on channel %s"),
+					*GetBriefFrameInfo(), *GetCommandActionString(Command), CommandWaitTime,
+					*Command.AssetPath.GetAssetName(), *Command.InstanceId.ToString(), *Command.ChannelName);
 
-	PendingPlaybackCommands.Reset();
+				PendingCommandIt.RemoveCurrent();
+			}
+		}
+		else
+		{
+			UE_LOG(LogAvaPlaybackServer, Verbose, TEXT("%s Playback Command [%s] Executed for asset: \"%s\" (id:%s) on channel \"%s\", received frame [%d], wait time: %.2f ms"),
+				*GetBriefFrameInfo(), *GetCommandActionString(Command),
+				*Command.AssetPath.GetAssetName(), *Command.InstanceId.ToString(), *Command.ChannelName,
+				(*PendingCommandIt)->ReceivedFrameNumber, (InUtcNow - (*PendingCommandIt)->ReceivedUtc).GetTotalMilliseconds());
+
+			PendingCommandIt.RemoveCurrent();
+		}
+	}
 }
 
 TSharedPtr<FAvaPlaybackInstance> FAvaPlaybackServer::GetOrLoadPlaybackInstance(const FGuid& InInstanceId, const FString& InChannelName, const FSoftObjectPath& InAssetPath)
@@ -1297,7 +1382,8 @@ void FAvaPlaybackServer::LoadPlayback(const FMessageAddress& InReplyToAddress, c
 			{
 				PlaybackInstance->GetPlayback()->LoadInstances();
 			}
-			SendPlaybackStatus(InReplyToAddress, InInstanceId, InChannelName, InAssetPath, EAvaPlaybackStatus::Loading);
+			PlaybackInstance->UpdateStatus();
+			SendPlaybackStatus(InReplyToAddress, InInstanceId, InChannelName, InAssetPath, PlaybackInstance->GetStatus());
 		}
 		else
 		{
@@ -1383,6 +1469,20 @@ void FAvaPlaybackServer::UnloadPlayback(const FMessageAddress& InReplyToAddress,
 	{
 		if (const TSharedPtr<FAvaPlaybackInstance> Instance = FindActivePlaybackInstance(InInstanceId))
 		{
+			// Validation of the operation.
+			// Unloading an instance that is part of a transition is an error state.
+			if (PlaybackInstanceTransitions)
+			{
+				for (const TPair<FGuid, TObjectPtr<UAvaPlaybackServerTransition>>& Transition : PlaybackInstanceTransitions->Transitions)
+				{
+					if (Transition.Value && Transition.Value->ContainsInstance(InInstanceId))
+					{
+						UE_LOG(LogAvaPlaybackServer, Error, TEXT("%s Unloading instance \"%s\" (id:%s) while it is part of transition %s."),
+							*UE::AvaPlayback::Utils::GetBriefFrameInfo(), *Instance->GetSourcePath().GetAssetName(), *InInstanceId.ToString(), *Transition.Value->GetTransitionId().ToString());
+					}
+				}
+			}
+			
 			Instance->Unload();
 			ActivePlaybackInstances.Remove(InInstanceId);
 			SendPlaybackStatus(InReplyToAddress, InInstanceId, Instance->GetChannelName(), Instance->GetSourcePath(), GetUnloadedPlaybackStatus(Instance->GetSourcePath()));
@@ -1421,34 +1521,46 @@ void FAvaPlaybackServer::UnloadPlayback(const FMessageAddress& InReplyToAddress,
 	}
 }
 
-void FAvaPlaybackServer::SetPlaybackUserData(const FMessageAddress& InReplyToAddress, const FGuid& InInstanceId, const FString& InUserData)
+bool FAvaPlaybackServer::SetPlaybackUserData(const FMessageAddress& InReplyToAddress, const FGuid& InInstanceId, const FString& InUserData)
 {
 	if (InInstanceId.IsValid())
 	{
-		if (const TSharedPtr<FAvaPlaybackInstance> Instance = FindActivePlaybackInstance(InInstanceId))
+		const TSharedPtr<FAvaPlaybackInstance> Instance = FindActivePlaybackInstance(InInstanceId);
+
+		// Instance may not be loaded yet.
+		if (!Instance)
 		{
-			Instance->SetInstanceUserData(InUserData);
-			
-			using namespace UE::AvaPlaybackServer::Private;
-			SendResponse( MakePlaybackStatusMessage(InInstanceId,
-				Instance->GetChannelName(), Instance->GetSourcePath(),
-				Instance->GetStatus(), Instance->GetInstanceUserData()), InReplyToAddress);
+			return false;
 		}
+		
+		Instance->SetInstanceUserData(InUserData);
+			
+		using namespace UE::AvaPlaybackServer::Private;
+		SendResponse( MakePlaybackStatusMessage(InInstanceId,
+			Instance->GetChannelName(), Instance->GetSourcePath(),
+			Instance->GetStatus(), Instance->GetInstanceUserData()), InReplyToAddress);
 	}
+	return true;
 }
 
-void FAvaPlaybackServer::SendPlaybackUserData(const FMessageAddress& InReplyToAddress, const FGuid& InInstanceId)
+bool FAvaPlaybackServer::SendPlaybackUserData(const FMessageAddress& InReplyToAddress, const FGuid& InInstanceId)
 {
 	if (InInstanceId.IsValid())
 	{
-		if (const TSharedPtr<FAvaPlaybackInstance> Instance = FindActivePlaybackInstance(InInstanceId))
+		const TSharedPtr<FAvaPlaybackInstance> Instance = FindActivePlaybackInstance(InInstanceId);
+
+		// Instance may not be loaded yet.
+		if (!Instance)
 		{
-			using namespace UE::AvaPlaybackServer::Private;
-			SendResponse( MakePlaybackStatusMessage(InInstanceId,
-				Instance->GetChannelName(), Instance->GetSourcePath(),
-				Instance->GetStatus(), Instance->GetInstanceUserData()), InReplyToAddress);
+			return false;
 		}
+
+		using namespace UE::AvaPlaybackServer::Private;
+		SendResponse( MakePlaybackStatusMessage(InInstanceId,
+			Instance->GetChannelName(), Instance->GetSourcePath(),
+			Instance->GetStatus(), Instance->GetInstanceUserData()), InReplyToAddress);
 	}
+	return true;
 }
 
 void FAvaPlaybackServer::SendPlaybackStatus(const FMessageAddress& InReplyToAddress, const FGuid& InInstanceId, const FString& InChannelName, const FSoftObjectPath& InAssetPath)
diff --git a/Engine/Plugins/Experimental/Avalanche/Source/AvalancheMedia/Private/Playback/AvaPlaybackServer.h b/Engine/Plugins/Experimental/Avalanche/Source/AvalancheMedia/Private/Playback/AvaPlaybackServer.h
--- a/Engine/Plugins/Experimental/Avalanche/Source/AvalancheMedia/Private/Playback/AvaPlaybackServer.h
+++ b/Engine/Plugins/Experimental/Avalanche/Source/AvalancheMedia/Private/Playback/AvaPlaybackServer.h
@@ -179,15 +179,15 @@ protected:
 	void SendLogMessage(const TCHAR* V, ELogVerbosity::Type Verbosity, const FName& Category, double Time);
 
 	// Playback Commands
-	void ExecutePendingPlaybackCommands();
+	void ExecutePendingPlaybackCommands(const FDateTime& InUtcNow);
 
 	TSharedPtr<FAvaPlaybackInstance> GetOrLoadPlaybackInstance(const FGuid& InInstanceId, const FString& InChannelName, const FSoftObjectPath& InAssetPath);
 	void LoadPlayback(const FMessageAddress& InReplyToAddress, const FGuid& InInstanceId, const FString& InChannelName, const FSoftObjectPath& InAssetPath);
 	void StartPlayback(const FMessageAddress& InReplyToAddress, const FGuid& InInstanceId, const FString& InChannelName, const FSoftObjectPath& InAssetPath);
 	void StopPlayback(const FMessageAddress& InReplyToAddress, const FGuid& InInstanceId, const FString& InChannelName, const FSoftObjectPath& InAssetPath);	
 	void UnloadPlayback(const FMessageAddress& InReplyToAddress, const FGuid& InInstanceId, const FString& InChannelName, const FSoftObjectPath& InAssetPath);
-	void SetPlaybackUserData(const FMessageAddress& InReplyToAddress, const FGuid& InInstanceId, const FString& InUserData);
-	void SendPlaybackUserData(const FMessageAddress& InReplyToAddress, const FGuid& InInstanceId);
+	bool SetPlaybackUserData(const FMessageAddress& InReplyToAddress, const FGuid& InInstanceId, const FString& InUserData);
+	bool SendPlaybackUserData(const FMessageAddress& InReplyToAddress, const FGuid& InInstanceId);
 	void SendPlaybackStatus(const FMessageAddress& InReplyToAddress, const FGuid& InInstanceId, const FString& InChannelName, const FSoftObjectPath& InAssetPath);
 	
 	void SendPlaybackStatus(const FMessageAddress& InSendTo, const FGuid& InInstanceId, const FString& InChannelName, const FSoftObjectPath& InAssetPath, EAvaPlaybackStatus InStatus);
@@ -225,12 +225,15 @@ private:
 
 	struct FPendingPlaybackCommand
 	{
+		FDateTime ReceivedUtc;
+		uint32 ReceivedFrameNumber;
+		int32 Priority;
 		FMessageAddress ReplyTo;
 		FAvaPlaybackCommand Command;
 	};
 
 	/** Accumulate all the playback commands and execute them all in one batch on the next tick. */
-	TArray<FPendingPlaybackCommand> PendingPlaybackCommands;
+	TArray<TSharedPtr<FPendingPlaybackCommand>> PendingPlaybackCommands;
 
 	/** Keep an map of active instances per id for fast lookup. */
 	TMap<FGuid, TSharedPtr<FAvaPlaybackInstance>> ActivePlaybackInstances;
diff --git a/Engine/Plugins/Experimental/Avalanche/Source/AvalancheMedia/Private/Playback/Transition/AvaPlaybackServerTransition.cpp b/Engine/Plugins/Experimental/Avalanche/Source/AvalancheMedia/Private/Playback/Transition/AvaPlaybackServerTransition.cpp
--- a/Engine/Plugins/Experimental/Avalanche/Source/AvalancheMedia/Private/Playback/Transition/AvaPlaybackServerTransition.cpp
+++ b/Engine/Plugins/Experimental/Avalanche/Source/AvalancheMedia/Private/Playback/Transition/AvaPlaybackServerTransition.cpp
@@ -5,6 +5,7 @@
 #include "Playable/Transition/AvaPlayableTransition.h"
 #include "Playback/AvaPlaybackManager.h"
 #include "Playback/AvaPlaybackServer.h"
+#include "Playback/AvaPlaybackUtils.h"
 
 namespace UE::AvaPlaybackServerTransition::Private
 {
@@ -14,7 +15,7 @@ namespace UE::AvaPlaybackServerTransition::Private
 		{
 			return FString::Printf(TEXT("Id:%s, Asset:%s, Channel:%s, UserData:\"%s\""),
 				*InPlaybackInstance->GetInstanceId().ToString(),
-				*InPlaybackInstance->GetSourcePath().ToString(),
+				*InPlaybackInstance->GetSourcePath().GetAssetName(),
 				*InPlaybackInstance->GetChannelName(),
 				*InPlaybackInstance->GetInstanceUserData());
 		}
@@ -42,6 +43,84 @@ namespace UE::AvaPlaybackServerTransition::Private
 		}
 		return nullptr;
 	}
+
+	TSharedPtr<FAvaPlaybackInstance> FindInstance(const TArray<TWeakPtr<FAvaPlaybackInstance>>& InPlaybackInstancesWeak, const FGuid& InInstanceId)
+	{
+		for (const TWeakPtr<FAvaPlaybackInstance>& InstanceWeak : InPlaybackInstancesWeak)
+		{
+			TSharedPtr<FAvaPlaybackInstance> Instance = InstanceWeak.Pin();
+			if (Instance && Instance->GetInstanceId() == InInstanceId)
+			{
+				return Instance;
+			}
+		}
+		return nullptr;
+	}
+
+	
+	// Check if the given set of playback instances are preventing the start of the transition.
+	bool CanStartTransition(const UAvaPlaybackServerTransition* InTransition,
+		const TArray<TWeakPtr<FAvaPlaybackInstance>>& InPlaybackInstancesWeak, bool& bOutShouldDiscard)
+	{
+		using namespace UE::AvaPlayback::Utils;
+		for (const TWeakPtr<FAvaPlaybackInstance>& InstanceWeak : InPlaybackInstancesWeak)
+		{
+			const TSharedPtr<FAvaPlaybackInstance> Instance = InstanceWeak.Pin();
+			if (!Instance)
+			{
+				// For now, we discard transitions with invalid instance.
+				UE_LOG(LogAvaPlaybackServer, Warning,
+					TEXT("%s Discarding Playback Transition {%s}. Reason: Invalid Instance. "),
+					*GetBriefFrameInfo(), *InTransition->GetPrettyTransitionInfo());
+
+				bOutShouldDiscard = true;
+				return false;
+			}
+			
+			const UAvaPlayable* Playable = GetPlayable(Instance.Get());
+			if (!Playable)
+			{
+				bOutShouldDiscard = false;
+				return false;	// Playable not yet created.
+			}
+
+			const EAvaPlayableStatus PlayableStatus = Playable->GetPlayableStatus();
+
+			if (PlayableStatus == EAvaPlayableStatus::Unknown
+				|| PlayableStatus == EAvaPlayableStatus::Error)
+			{
+				UE_LOG(LogAvaPlaybackServer, Warning,
+            		TEXT("%s Discarding Playback Transition {%s}. Reason: Playable status: \"%s\". "),
+            		*GetBriefFrameInfo(), *InTransition->GetPrettyTransitionInfo(), *StaticEnumToString(PlayableStatus));
+
+				// Discard the command
+				bOutShouldDiscard = true;
+				return false;
+			}
+
+			// todo: this might cause commands to become stale and fill the pending command list
+			if (PlayableStatus == EAvaPlayableStatus::Unloaded)
+			{
+				UE_LOG(LogAvaPlaybackServer, Warning,
+					TEXT("%s Playback Transition {%s}: Playable \"%s\" (Id:%s) is unloaded."),
+					*GetBriefFrameInfo(), *InTransition->GetPrettyTransitionInfo(), *Playable->GetSourceAssetPath().GetAssetName(), *Playable->GetInstanceId().ToString());
+
+				return false;
+			}
+
+			// Asset status must be visible to run the command.
+			// If not visible, the components are not yet added to the world.
+			if (PlayableStatus != EAvaPlayableStatus::Visible)
+			{
+				// Keep the command in the queue for next tick.
+				bOutShouldDiscard = false;
+				return false;
+			}
+		}
+
+		bOutShouldDiscard = true;
+		return true;
+	}
 }
 
 UAvaPlaybackServerTransition* UAvaPlaybackServerTransition::MakeNew(const TSharedPtr<FAvaPlaybackServer>& InPlaybackServer)
@@ -51,6 +130,26 @@ UAvaPlaybackServerTransition* UAvaPlaybackServerTransition::MakeNew(const TShare
 	return NewTransition;
 }
 
+void UAvaPlaybackServerTransition::AddPendingEnterInstanceIds(const TArray<FGuid>& InInstanceIds)
+{
+	PendingEnterInstanceIds.Reserve(InInstanceIds.Num());
+	
+	for (const FGuid& InstanceId : InInstanceIds)
+	{
+		PendingEnterInstanceIds.AddUnique(InstanceId);		
+	}
+}
+
+void UAvaPlaybackServerTransition::AddPendingPlayingInstanceId(const FGuid& InInstanceId)
+{
+	PendingPlayingInstanceIds.AddUnique(InInstanceId);
+}
+	
+void UAvaPlaybackServerTransition::AddPendingExitInstanceId(const FGuid& InInstanceId)
+{
+	PendingExitInstanceIds.AddUnique(InInstanceId);
+}
+
 void UAvaPlaybackServerTransition::SetEnterValues(const TArray<FAvaPlayableRemoteControlValues>& InEnterValues)
 {
 	EnterValues.Reserve(InEnterValues.Num());
@@ -97,19 +196,51 @@ bool UAvaPlaybackServerTransition::AddExitInstance(const TSharedPtr<FAvaPlayback
 
 void UAvaPlaybackServerTransition::TryResolveInstances(const FAvaPlaybackServer& InPlaybackServer)
 {
-	if (EnterPlaybackInstancesWeak.Num() < EnterInstanceIds.Num())
+	for (TArray<FGuid>::TIterator InstanceIdIt(PendingEnterInstanceIds); InstanceIdIt; ++InstanceIdIt)
 	{
-		for (const FGuid& InstanceId : EnterInstanceIds)
+		if (TSharedPtr<FAvaPlaybackInstance> Instance = InPlaybackServer.FindActivePlaybackInstance(*InstanceIdIt))
 		{
-			if (TSharedPtr<FAvaPlaybackInstance> Instance = InPlaybackServer.FindActivePlaybackInstance(InstanceId))
-			{
-				if (!EnterPlaybackInstancesWeak.Contains(Instance.ToWeakPtr()))
-				{
-					AddEnterInstance(Instance);
-				}
-			}
+			AddEnterInstance(Instance);
+			InstanceIdIt.RemoveCurrent();
+		}
+	}
+	
+	for (TArray<FGuid>::TIterator InstanceIdIt(PendingPlayingInstanceIds); InstanceIdIt; ++InstanceIdIt)
+	{
+		if (TSharedPtr<FAvaPlaybackInstance> Instance = InPlaybackServer.FindActivePlaybackInstance(*InstanceIdIt))
+		{
+			AddPlayingInstance(Instance);
+			InstanceIdIt.RemoveCurrent();
 		}
-	}	
+	}
+
+	for (TArray<FGuid>::TIterator InstanceIdIt(PendingExitInstanceIds); InstanceIdIt; ++InstanceIdIt)
+	{
+		if (TSharedPtr<FAvaPlaybackInstance> Instance = InPlaybackServer.FindActivePlaybackInstance(*InstanceIdIt))
+		{
+			AddExitInstance(Instance);
+			InstanceIdIt.RemoveCurrent();
+		}
+	}
+}
+
+bool UAvaPlaybackServerTransition::ContainsInstance(const FGuid& InInstanceId) const
+{
+	if (PendingEnterInstanceIds.Contains(InInstanceId)
+		|| PendingPlayingInstanceIds.Contains(InInstanceId)
+		|| PendingExitInstanceIds.Contains(InInstanceId))
+	{
+		return true;
+	}
+
+	using namespace UE::AvaPlaybackServerTransition::Private;
+	if (FindInstance(EnterPlaybackInstancesWeak, InInstanceId).IsValid()
+		|| FindInstance(PlayingPlaybackInstancesWeak, InInstanceId).IsValid()
+		|| FindInstance(ExitPlaybackInstancesWeak, InInstanceId).IsValid())
+	{
+		return true;
+	}
+	return false;
 }
 
 bool UAvaPlaybackServerTransition::IsVisibilityConstrained(const UAvaPlayable* InPlayable) const
@@ -144,51 +275,33 @@ bool UAvaPlaybackServerTransition::CanStart(bool& bOutShouldDiscard) const
 {
 	using namespace UE::AvaPlaybackServerTransition::Private;
 
-	if (EnterPlaybackInstancesWeak.Num() < EnterInstanceIds.Num())
+	// Wait for any unresolved instances to be loaded.
+	if (!PendingEnterInstanceIds.IsEmpty()
+		|| !PendingPlayingInstanceIds.IsEmpty()
+		|| !PendingExitInstanceIds.IsEmpty())
 	{
 		bOutShouldDiscard = false;
 		return false;
 	}
 
-	for (const TWeakPtr<FAvaPlaybackInstance>& InstanceWeak : EnterPlaybackInstancesWeak)
+	if (!CanStartTransition(this, EnterPlaybackInstancesWeak, bOutShouldDiscard))
 	{
-		if (const TSharedPtr<FAvaPlaybackInstance> Instance = InstanceWeak.Pin())
-		{
-			if (const UAvaPlayable* Playable = GetPlayable(Instance.Get()))
-			{
-				const EAvaPlayableStatus PlayableStatus = Playable->GetPlayableStatus();
-
-				if (PlayableStatus == EAvaPlayableStatus::Unknown
-					|| PlayableStatus == EAvaPlayableStatus::Error)
-				{
-					// Discard the command
-					bOutShouldDiscard = true;
-					return false;
-				}
+		return false;
+	}
 
-				// todo: this might cause commands to become stale and fill the pending command list
-				if (PlayableStatus == EAvaPlayableStatus::Unloaded)
-				{
-					return false;
-				}
+	// Note: need to check the "non-entering" instances too in case the playback commands got delayed
+	// causing playables to also need loading/recovering.
+	
+	if (!CanStartTransition(this, PlayingPlaybackInstancesWeak, bOutShouldDiscard))
+	{
+		return false;
+	}
 
-				// Asset status must be visible to run the command.
-				// If not visible, the components are not yet added to the world.
-				if (PlayableStatus != EAvaPlayableStatus::Visible)
-				{
-					// Keep the command in the queue for next tick.
-					bOutShouldDiscard = false;
-					return false;
-				}
-			}
-		}
-		else
-		{
-			bOutShouldDiscard = true;
-			return false;
-		}
+	if (!CanStartTransition(this, ExitPlaybackInstancesWeak, bOutShouldDiscard))
+	{
+		return false;
 	}
-	
+
 	bOutShouldDiscard = true;
 	return true;
 }
@@ -245,7 +358,8 @@ void UAvaPlaybackServerTransition::Stop()
 		if (!PlaybackServer->RemovePlaybackInstanceTransition(TransitionId))
 		{
 			UE_LOG(LogAvaPlaybackServer, Error,
-				TEXT("Playback Transition {%s} Error: Was not found in server's active transitions. "), *GetPrettyTransitionInfo());
+				TEXT("%s Failed to remove Playback Transition {%s}. Reason: not found in server's active transitions."),
+				*UE::AvaPlayback::Utils::GetBriefFrameInfo(), *GetPrettyTransitionInfo());
 		}
 	}
 }
@@ -291,19 +405,19 @@ TSharedPtr<FAvaPlaybackInstance> UAvaPlaybackServerTransition::FindInstanceForPl
 		return nullptr;
 	}
 
-	TSharedPtr<FAvaPlaybackInstance> Instance =  Private::FindInstanceForPlayable(EnterPlaybackInstancesWeak, InPlayable);
+	TSharedPtr<FAvaPlaybackInstance> Instance = Private::FindInstanceForPlayable(EnterPlaybackInstancesWeak, InPlayable);
 	if (Instance)
 	{
 		return Instance;
 	}
 	
-	Instance =  Private::FindInstanceForPlayable(PlayingPlaybackInstancesWeak, InPlayable);
+	Instance = Private::FindInstanceForPlayable(PlayingPlaybackInstancesWeak, InPlayable);
 	if (Instance)
 	{
 		return Instance;
 	}
 
-	Instance =  Private::FindInstanceForPlayable(ExitPlaybackInstancesWeak, InPlayable);
+	Instance = Private::FindInstanceForPlayable(ExitPlaybackInstancesWeak, InPlayable);
 	if (Instance)
 	{
 		return Instance;
@@ -314,6 +428,7 @@ TSharedPtr<FAvaPlaybackInstance> UAvaPlaybackServerTransition::FindInstanceForPl
 void UAvaPlaybackServerTransition::OnTransitionEvent(UAvaPlayable* InPlayable, UAvaPlayableTransition* InTransition, EAvaPlayableTransitionEventFlags InTransitionFlags)
 {
 	using namespace UE::AvaPlaybackServerTransition::Private;
+	using namespace UE::AvaPlayback::Utils;
 	
 	// not this transition.
 	if (InTransition != PlayableTransition || PlayableTransition == nullptr)
@@ -338,8 +453,8 @@ void UAvaPlaybackServerTransition::OnTransitionEvent(UAvaPlayable* InPlayable, U
 			if (PlayableTransition->IsEnterPlayable(InPlayable))
 			{
 				UE_LOG(LogAvaPlaybackServer, Error,
-					TEXT("Playback Transition {%s} Error: An \"enter\" playable is being discarded for instance {%s}."),
-					*GetPrettyTransitionInfo(), *GetPrettyPlaybackInstanceInfo(Instance.Get()));
+					TEXT("%s Playback Transition {%s} Error: An \"enter\" playable is being discarded for instance {%s}."),
+					*GetBriefFrameInfo(), *GetPrettyTransitionInfo(), *GetPrettyPlaybackInstanceInfo(Instance.Get()));
 			}
 
 			// See UAvaRundownPagePlayer::Stop()
@@ -356,8 +471,8 @@ void UAvaPlaybackServerTransition::OnTransitionEvent(UAvaPlayable* InPlayable, U
 					if (!PlaybackServer->RemoveActivePlaybackInstance(Instance->GetInstanceId()))
 					{
 						UE_LOG(LogAvaPlaybackServer, Error,
-							TEXT("Playback Transition {%s} Error: \"exit\" instance {%s} was not found in server's active instances. "),
-							*GetPrettyTransitionInfo(), *GetPrettyPlaybackInstanceInfo(Instance.Get()));
+							TEXT("%s Playback Transition {%s} Error: \"exit\" instance {%s} was not found in server's active instances. "),
+							*GetBriefFrameInfo(), *GetPrettyTransitionInfo(), *GetPrettyPlaybackInstanceInfo(Instance.Get()));
 					}
 				}
 			}
@@ -390,7 +505,8 @@ void UAvaPlaybackServerTransition::OnPlayableCreated(UAvaPlaybackGraph* InPlayba
 void UAvaPlaybackServerTransition::MakePlayableTransition()
 {
 	using namespace UE::AvaPlaybackServerTransition::Private;
-
+	using namespace UE::AvaPlayback::Utils;
+	
 	FAvaPlayableTransitionBuilder TransitionBuilder;
 
 	auto AddInstancesToBuilder = [&TransitionBuilder, this](const TArray<TWeakPtr<FAvaPlaybackInstance>>& InPlaybackInstancesWeak, const TCHAR* InCategory, EAvaPlayableTransitionEntryRole InEntryRole)
@@ -413,8 +529,8 @@ void UAvaPlaybackServerTransition::MakePlayableTransition()
 				{
 					// If this happens, likely the playable is not yet loaded.
 					UE_LOG(LogAvaPlaybackServer, Error,
-						TEXT("Playback Transition {%s} Error: Failed to retrieve \"%s\" playable for instance {%s}."),
-						*GetPrettyTransitionInfo(), InCategory, *GetPrettyPlaybackInstanceInfo(Instance.Get()));
+						TEXT("%s Playback Transition {%s} Error: Failed to retrieve \"%s\" playable for instance {%s}."),
+						*GetBriefFrameInfo(), *GetPrettyTransitionInfo(), InCategory, *GetPrettyPlaybackInstanceInfo(Instance.Get()));
 				}
 			}
 			++ArrayIndex;
@@ -424,7 +540,7 @@ void UAvaPlaybackServerTransition::MakePlayableTransition()
 	AddInstancesToBuilder(EnterPlaybackInstancesWeak, TEXT("Enter"), EAvaPlayableTransitionEntryRole::Enter);
 	AddInstancesToBuilder(PlayingPlaybackInstancesWeak, TEXT("Playing"), EAvaPlayableTransitionEntryRole::Playing);
 	AddInstancesToBuilder(ExitPlaybackInstancesWeak, TEXT("Exit"), EAvaPlayableTransitionEntryRole::Exit);
-	PlayableTransition = TransitionBuilder.MakeTransition(this);
+	PlayableTransition = TransitionBuilder.MakeTransition(this, TransitionId);
 
 	if (PlayableTransition)
 	{
@@ -434,7 +550,8 @@ void UAvaPlaybackServerTransition::MakePlayableTransition()
 
 void UAvaPlaybackServerTransition::LogDetailedTransitionInfo() const
 {
-	UE_LOG(LogAvaPlaybackServer, Verbose, TEXT("Playback Transition {%s}:"), *GetPrettyTransitionInfo());
+	using namespace UE::AvaPlayback::Utils;
+	UE_LOG(LogAvaPlaybackServer, Verbose, TEXT("%s Playback Transition {%s}:"), *GetBriefFrameInfo(), *GetPrettyTransitionInfo());
 
 	auto LogInstances = [](const TArray<TWeakPtr<FAvaPlaybackInstance>>& InPlaybackInstancesWeak, const TCHAR* InCategory)
 	{
@@ -486,13 +603,14 @@ void UAvaPlaybackServerTransition::UpdateChannelName(const FAvaPlaybackInstance*
 	else
 	{
 		using namespace UE::AvaPlaybackServerTransition::Private;
+		using namespace UE::AvaPlayback::Utils;
 
 		// Validate the channel is the same.
 		if (ChannelName != InPlaybackInstance->GetChannelFName())
 		{
 			UE_LOG(LogAvaPlaybackServer, Error,
-				TEXT("Playback Transition {%s}: Adding Playback Instance {%s} in a different channel than previous playback instance (\"%s\")."),
-				*GetPrettyTransitionInfo(), *GetPrettyPlaybackInstanceInfo(InPlaybackInstance), *ChannelName.ToString());
+				TEXT("%s Playback Transition {%s}: Adding Playback Instance {%s} in a different channel than previous playback instance (\"%s\")."),
+				*GetBriefFrameInfo(), *GetPrettyTransitionInfo(), *GetPrettyPlaybackInstanceInfo(InPlaybackInstance), *ChannelName.ToString());
 		}
 	}
 }
\ No newline at end of file
diff --git a/Engine/Plugins/Experimental/Avalanche/Source/AvalancheMedia/Private/Playback/Transition/AvaPlaybackServerTransition.h b/Engine/Plugins/Experimental/Avalanche/Source/AvalancheMedia/Private/Playback/Transition/AvaPlaybackServerTransition.h
--- a/Engine/Plugins/Experimental/Avalanche/Source/AvalancheMedia/Private/Playback/Transition/AvaPlaybackServerTransition.h
+++ b/Engine/Plugins/Experimental/Avalanche/Source/AvalancheMedia/Private/Playback/Transition/AvaPlaybackServerTransition.h
@@ -33,13 +33,17 @@ public:
 	void SetClientName(const FString& InClientName) { ClientName = InClientName; }
 	void SetUnloadDiscardedInstances(bool bInUnloadDiscardedInstances) { bUnloadDiscardedInstances = bInUnloadDiscardedInstances; }
 	void SetTransitionFlags(EAvaPlayableTransitionFlags InTransitionFlags) { TransitionFlags = InTransitionFlags; }
-	void SetEnterInstanceIds(const TArray<FGuid>& InInstanceIds) { EnterInstanceIds = InInstanceIds; }
+	void AddPendingEnterInstanceIds(const TArray<FGuid>& InInstanceIds);
+	void AddPendingPlayingInstanceId(const FGuid& InInstanceId);
+	void AddPendingExitInstanceId(const FGuid& InInstanceId);
 	void SetEnterValues(const TArray<FAvaPlayableRemoteControlValues>& InEnterValues);
 	bool AddEnterInstance(const TSharedPtr<FAvaPlaybackInstance>& InPlaybackInstance);
 	bool AddPlayingInstance(const TSharedPtr<FAvaPlaybackInstance>& InPlaybackInstance);
 	bool AddExitInstance(const TSharedPtr<FAvaPlaybackInstance>& InPlaybackInstance);
-
+	
 	void TryResolveInstances(const FAvaPlaybackServer& InPlaybackServer);
+
+	bool ContainsInstance(const FGuid& InInstanceId) const;
 	
 	//~ Begin IAvaPlayableVisibilityConstraint
 	virtual bool IsVisibilityConstrained(const UAvaPlayable* InPlayable) const override;
@@ -79,11 +83,13 @@ protected:
 	
 	FString ClientName;
 	FName ChannelName;
-	FGuid TransitionId;
 	bool bUnloadDiscardedInstances = false;
 	EAvaPlayableTransitionFlags TransitionFlags = EAvaPlayableTransitionFlags::None;
-	
-	TArray<FGuid> EnterInstanceIds;
+
+	/** Instance Ids pending resolve. */
+	TArray<FGuid> PendingEnterInstanceIds;
+	TArray<FGuid> PendingPlayingInstanceIds;
+	TArray<FGuid> PendingExitInstanceIds;
 
 	TArray<TWeakPtr<FAvaPlaybackInstance>> EnterPlaybackInstancesWeak;
 	TArray<TWeakPtr<FAvaPlaybackInstance>> PlayingPlaybackInstancesWeak;
diff --git a/Engine/Plugins/Experimental/Avalanche/Source/AvalancheMedia/Private/Playback/Transition/AvaPlaybackTransition.h b/Engine/Plugins/Experimental/Avalanche/Source/AvalancheMedia/Private/Playback/Transition/AvaPlaybackTransition.h
--- a/Engine/Plugins/Experimental/Avalanche/Source/AvalancheMedia/Private/Playback/Transition/AvaPlaybackTransition.h
+++ b/Engine/Plugins/Experimental/Avalanche/Source/AvalancheMedia/Private/Playback/Transition/AvaPlaybackTransition.h
@@ -32,4 +32,9 @@ public:
 	//~ Begin IAvaPlayableVisibilityConstraint
 	virtual bool IsVisibilityConstrained(const UAvaPlayable* InPlayable) const override { return false; }
 	//~ End IAvaPlayableVisibilityConstraint
+
+	const FGuid& GetTransitionId() const { return TransitionId; }
+	
+protected:
+	FGuid TransitionId;
 };
\ No newline at end of file
diff --git a/Engine/Plugins/Experimental/Avalanche/Source/AvalancheMedia/Private/Rundown/AvaRundown.cpp b/Engine/Plugins/Experimental/Avalanche/Source/AvalancheMedia/Private/Rundown/AvaRundown.cpp
--- a/Engine/Plugins/Experimental/Avalanche/Source/AvalancheMedia/Private/Rundown/AvaRundown.cpp
+++ b/Engine/Plugins/Experimental/Avalanche/Source/AvalancheMedia/Private/Rundown/AvaRundown.cpp
@@ -8,6 +8,7 @@
 #include "Playable/AvaPlayable.h"
 #include "Playback/AvaPlaybackGraph.h"
 #include "Playback/AvaPlaybackManager.h"
+#include "Playback/AvaPlaybackUtils.h"
 #include "Rundown/AvaRundownManagedInstanceCache.h"
 #include "Rundown/AvaRundownPageLoadingManager.h"
 #include "Rundown/AvaRundownPagePlayer.h"
@@ -1033,11 +1034,13 @@ TArray<int32> UAvaRundown::PlayPages(const TArray<int32>& InPageIds, EAvaRundown
 
 bool UAvaRundown::RestorePlaySubPage(int32 InPageId, int32 InSubPageIndex, const FGuid& InExistingInstanceId, bool bInIsPreview, const FName& InPreviewChannelName)
 {
+	using namespace UE::AvaPlayback::Utils;
+	
 	auto LogError = [InPageId, InPreviewChannelName](const FString& InReason)
 	{
 		UE_LOG(LogAvaRundown, Error,
-			TEXT("Couldn't restore playback state of page %d on channel \"%s\": %s."),
-			InPageId, *InPreviewChannelName.ToString(), *InReason);
+			TEXT("%s Couldn't restore playback state of page %d on channel \"%s\": %s."),
+			*GetBriefFrameInfo(), InPageId, *InPreviewChannelName.ToString(), *InReason);
 	};
 
 	const FAvaRundownPage& Page = GetPage(InPageId);
@@ -1047,7 +1050,7 @@ bool UAvaRundown::RestorePlaySubPage(int32 InPageId, int32 InSubPageIndex, const
 		return false;
 	}
 
-	if (!IsChannelTypeCompatibleForRequest(Page, bInIsPreview, InPreviewChannelName, true))
+	if (!IsChannelTypeCompatibleForRequest(Page, bInIsPreview, InPreviewChannelName, /*bInLogFailureReason*/ true))
 	{
 		LogError(TEXT("Channel Type is not compatible"));
 		return false;
@@ -1070,6 +1073,7 @@ bool UAvaRundown::RestorePlaySubPage(int32 InPageId, int32 InSubPageIndex, const
 		{
 			return false;
 		}
+		UE_LOG(LogAvaRundown, Verbose, TEXT("%s Restored page player for page %d."), *GetBriefFrameInfo(), InPageId);
 	}
 	
 	if (const UAvaRundownPlaybackInstancePlayer* LoadedInstancePlayer = PagePlayer->LoadInstancePlayer(InSubPageIndex, InExistingInstanceId))
diff --git a/Engine/Plugins/Experimental/Avalanche/Source/AvalancheMedia/Private/Rundown/AvaRundownPagePlayer.cpp b/Engine/Plugins/Experimental/Avalanche/Source/AvalancheMedia/Private/Rundown/AvaRundownPagePlayer.cpp
--- a/Engine/Plugins/Experimental/Avalanche/Source/AvalancheMedia/Private/Rundown/AvaRundownPagePlayer.cpp
+++ b/Engine/Plugins/Experimental/Avalanche/Source/AvalancheMedia/Private/Rundown/AvaRundownPagePlayer.cpp
@@ -124,19 +124,16 @@ bool UAvaRundownPlaybackInstancePlayer::Continue(const FString& InChannelName)
 
 bool UAvaRundownPlaybackInstancePlayer::Stop()
 {
-	if (!Playback)
-	{
-		return false;
-	}
-	
 	const bool bUnload = !UAvaMediaSettings::Get().bKeepPagesLoaded;
+	bool bWasStopped = false; 
 	
-	if (Playback->IsPlaying())
+	if (Playback && Playback->IsPlaying())
 	{
 		// Propagate the unload options in case this object is playing remote.
 		const EAvaPlaybackStopOptions PlaybackStopOptions = bUnload ?
 			EAvaPlaybackStopOptions::Default | EAvaPlaybackStopOptions::Unload : EAvaPlaybackStopOptions::Default;
 		Playback->Stop(PlaybackStopOptions);
+		bWasStopped = true;
 	}
 
 	if (PlaybackInstance)
@@ -154,7 +151,7 @@ bool UAvaRundownPlaybackInstancePlayer::Stop()
 
 	Playback = nullptr;
 	PlaybackInstance.Reset();
-	return true;
+	return bWasStopped;
 }
 
 bool UAvaRundownPlaybackInstancePlayer::HasPlayable(const UAvaPlayable* InPlayable) const
@@ -169,12 +166,12 @@ UAvaPlayable* UAvaRundownPlaybackInstancePlayer::GetFirstPlayable() const
 
 UAvaRundownPagePlayer* UAvaRundownPlaybackInstancePlayer::GetPagePlayer() const
 {
-	return Cast<UAvaRundownPagePlayer>(GetOuter());
+	return ParentPagePlayer.Get();
 }
 
 void UAvaRundownPlaybackInstancePlayer::SetPagePlayer(UAvaRundownPagePlayer* InPagePlayer)
 {
-	LowLevelRename(GetFName(), InPagePlayer);
+	ParentPagePlayer = InPagePlayer;
 }
 
 UAvaRundownPagePlayer::UAvaRundownPagePlayer()
@@ -228,7 +225,7 @@ bool UAvaRundownPagePlayer::InitializeAndLoad(UAvaRundown* InRundown, const FAva
 
 UAvaRundownPlaybackInstancePlayer* UAvaRundownPagePlayer::LoadInstancePlayer(int32 InSubPageIndex, const FGuid& InInstanceId)
 {
-	const UAvaRundown* Rundown = RundownWeak.Get();
+	UAvaRundown* Rundown = RundownWeak.Get();
 	if (!Rundown)
 	{
 		UE_LOG(LogAvaRundown, Error, TEXT("UAvaRundownPagePlayer::LoadSubPage: Rundown is no longuer valid."));
@@ -386,11 +383,12 @@ UAvaRundownPlaybackInstancePlayer* UAvaRundownPagePlayer::FindInstancePlayerByAs
 	return nullptr;
 }
 
-UAvaRundownPlaybackInstancePlayer* UAvaRundownPagePlayer::CreateAndLoadInstancePlayer(const UAvaRundown* InRundown, const FAvaRundownPage& InPage, int32 InSubPageIndex, const FGuid& InInstanceId)
+UAvaRundownPlaybackInstancePlayer* UAvaRundownPagePlayer::CreateAndLoadInstancePlayer(UAvaRundown* InRundown, const FAvaRundownPage& InPage, int32 InSubPageIndex, const FGuid& InInstanceId)
 {
-	UAvaRundownPlaybackInstancePlayer* InstancePlayer = NewObject<UAvaRundownPlaybackInstancePlayer>(this);
+	UAvaRundownPlaybackInstancePlayer* InstancePlayer = NewObject<UAvaRundownPlaybackInstancePlayer>(InRundown);
 	if (InstancePlayer->Load(*this, InRundown, InPage, InSubPageIndex, InInstanceId))
 	{
+		InstancePlayer->SetPagePlayer(this);
 		InstancePlayers.Add(InstancePlayer);
 		return InstancePlayer;
 	}
diff --git a/Engine/Plugins/Experimental/Avalanche/Source/AvalancheMedia/Private/Rundown/AvaRundownPlaybackClientWatcher.cpp b/Engine/Plugins/Experimental/Avalanche/Source/AvalancheMedia/Private/Rundown/AvaRundownPlaybackClientWatcher.cpp
--- a/Engine/Plugins/Experimental/Avalanche/Source/AvalancheMedia/Private/Rundown/AvaRundownPlaybackClientWatcher.cpp
+++ b/Engine/Plugins/Experimental/Avalanche/Source/AvalancheMedia/Private/Rundown/AvaRundownPlaybackClientWatcher.cpp
@@ -3,6 +3,8 @@
 #include "Rundown/AvaRundownPlaybackClientWatcher.h"
 
 #include "Broadcast/AvaBroadcast.h"
+#include "IAvaMediaModule.h"
+#include "Playback/AvaPlaybackUtils.h"
 #include "Playback/IAvaPlaybackClient.h"
 #include "Rundown/AvaRundown.h"
 #include "Rundown/AvaRundownPagePlayer.h"
@@ -22,6 +24,8 @@ FAvaRundownPlaybackClientWatcher::~FAvaRundownPlaybackClientWatcher()
 
 void FAvaRundownPlaybackClientWatcher::TryRestorePlaySubPage(int InPageId, const UE::AvaPlaybackClient::Delegates::FPlaybackStatusChangedArgs& InEventArgs) const
 {
+	using namespace UE::AvaPlayback::Utils;
+	
 	// Restore page player and local playback proxies.
 	// (Need to specify the InstanceId from the server for everything to match.)
 	const FName ChannelFName(InEventArgs.ChannelName);
@@ -30,8 +34,8 @@ void FAvaRundownPlaybackClientWatcher::TryRestorePlaySubPage(int InPageId, const
 	if (UAvaBroadcast::Get().GetChannelIndex(ChannelFName) == INDEX_NONE)
 	{
 		UE_LOG(LogAvaRundown, Error,
-			TEXT("Received a playback object on channel \"%s\" which doesn't exist locally. Playback Server should be reset."),
-			*InEventArgs.ChannelName);
+			TEXT("%s Received a playback object on channel \"%s\" which doesn't exist locally. Playback Server should be reset."),
+			*GetBriefFrameInfo(), *InEventArgs.ChannelName);
 		return;
 	}
 	
@@ -44,7 +48,8 @@ void FAvaRundownPlaybackClientWatcher::TryRestorePlaySubPage(int InPageId, const
 	if (SubPageIndex == INDEX_NONE)
 	{
 		UE_LOG(LogAvaRundown, Error,
-			TEXT("Asset mismatch (expected (any of): \"%s\", received: \"%s\") for restoring page %d. Playback Server should be reset."),
+			TEXT("%s Asset mismatch (expected (any of): \"%s\", received: \"%s\") for restoring page %d. Playback Server should be reset."),
+			*GetBriefFrameInfo(),
 			*FString::JoinBy(PageToRestore.GetAssetPaths(Rundown), TEXT(","), [](const FSoftObjectPath& Path){ return Path.ToString();}),
 			*InEventArgs.AssetPath.ToString(), InPageId);
 		return;
@@ -52,7 +57,7 @@ void FAvaRundownPlaybackClientWatcher::TryRestorePlaySubPage(int InPageId, const
 
 	if (!Rundown->RestorePlaySubPage(InPageId, SubPageIndex, InEventArgs.InstanceId, bIsPreview, ChannelFName))
 	{
-		UE_LOG(LogAvaRundown, Error, TEXT("Failed to restore page %d. Playback Server should be reset."), InPageId);
+		UE_LOG(LogAvaRundown, Error, TEXT("%s Failed to restore page %d. Playback Server should be reset."), *GetBriefFrameInfo(), InPageId);
 	}
 }
 
@@ -73,32 +78,47 @@ void FAvaRundownPlaybackClientWatcher::HandlePlaybackStatusChanged(IAvaPlaybackC
 	// Try to determine if a playback has started or stopped.
 	const bool bWasRunning = IsAnyOf(InEventArgs.PrevStatus, RunningStates);
 	const bool bIsRunning = IsAnyOf(InEventArgs.NewStatus, RunningStates);
+
+	// TODO: Reconcile forked channels. Need to keep track of status per server. (Seems to work well enough for now, but may need to revisit)
 	
 	// If a playback instance is stopping, stop corresponding page (if any).
 	if (bWasRunning && !bIsRunning)
 	{
+		UE_LOG(LogAvaRundown, Verbose, 
+			TEXT("%s Playback Client Watcher: Detected asset stopping Id:%s from Server \"%s\"."),
+			*UE::AvaPlayback::Utils::GetBriefFrameInfo(), *InEventArgs.InstanceId.ToString(), *InEventArgs.ServerName);
+
 		for (UAvaRundownPagePlayer* PagePlayer : Rundown->PagePlayers)
 		{
 			// Search for a match with the event:				
 			if (PagePlayer && PagePlayer->ChannelName == InEventArgs.ChannelName)
 			{
-				PagePlayer->ForEachInstancePlayer([&InEventArgs](UAvaRundownPlaybackInstancePlayer* InInstancePlayer)
+				if (UAvaRundownPlaybackInstancePlayer* InstancePlayer = PagePlayer->FindInstancePlayerByInstanceId(InEventArgs.InstanceId))
 				{
-					if (InInstancePlayer
-						&& InInstancePlayer->SourceAssetPath == InEventArgs.AssetPath
-						&& InInstancePlayer->GetPlaybackInstanceId() != InEventArgs.InstanceId)
+					if (InstancePlayer->SourceAssetPath != InEventArgs.AssetPath)
 					{
-						InInstancePlayer->Stop();	
+						UE_LOG(LogAvaRundown, Error, TEXT("%s Playback Client Watcher: Instance Id:%s asset path mismatch in page player %d."),
+							*UE::AvaPlayback::Utils::GetBriefFrameInfo(), *InEventArgs.InstanceId.ToString(), PagePlayer->PageId);
+						continue;
 					}
-				});
+
+					UE_LOG(LogAvaRundown, Verbose, TEXT("%s Playback Client Watcher: Stopping Instance Id:%s in page player %d."),
+						*UE::AvaPlayback::Utils::GetBriefFrameInfo(), *InEventArgs.InstanceId.ToString(), PagePlayer->PageId);
+					
+					InstancePlayer->Stop();
+				}
 
 				// If we stopped all the instance players, stop the page (to broadcast events).
 				if (!PagePlayer->IsPlaying())
 				{
+					UE_LOG(LogAvaRundown, Verbose, TEXT("%s Playback Client Watcher: Stopping Page player %d, no more instances playing."),
+						*UE::AvaPlayback::Utils::GetBriefFrameInfo(), PagePlayer->PageId);
+
 					PagePlayer->Stop();	
 				}
 			}
 		}
+		
 		Rundown->RemoveStoppedPagePlayers();
 	}
 
@@ -106,17 +126,26 @@ void FAvaRundownPlaybackClientWatcher::HandlePlaybackStatusChanged(IAvaPlaybackC
 	if (bIsRunning)
 	{
 		// We need to figure out which page it is.
-		const FString* UserData = InPlaybackClient.GetRemotePlaybackUserData(InEventArgs.InstanceId, InEventArgs.AssetPath, InEventArgs.ChannelName);
+		const FString* RemoteUserData = InPlaybackClient.GetRemotePlaybackUserData(InEventArgs.InstanceId, InEventArgs.AssetPath, InEventArgs.ChannelName, InEventArgs.ServerName);
 
 		// We haven't received the user data for this playback. So we request it.
 		// This event will be received again with user data next time.
-		if (!UserData)
+		if (!RemoteUserData)
 		{
 			InPlaybackClient.RequestPlayback(InEventArgs.InstanceId, InEventArgs.AssetPath, InEventArgs.ChannelName, EAvaPlaybackAction::GetUserData);
 		}
 		else
 		{
-			const int32 PageId = UAvaRundownPagePlayer::GetPageIdFromInstanceUserData(*UserData);
+			const TSharedPtr<FAvaPlaybackInstance> LocalPlaybackInstance = IAvaMediaModule::Get().GetLocalPlaybackManager().FindPlaybackInstance(InEventArgs.InstanceId, InEventArgs.AssetPath, InEventArgs.ChannelName);
+			if (LocalPlaybackInstance && LocalPlaybackInstance->GetInstanceUserData() != *RemoteUserData)
+			{
+				UE_LOG(LogAvaRundown, Error, 
+					TEXT("%s Playback Client Watcher: Playback Instance Id:%s on server \"%s\": user data mismatch \"%s\", local user data: \"%s\"."),
+					*UE::AvaPlayback::Utils::GetBriefFrameInfo(),
+					*InEventArgs.InstanceId.ToString(), *InEventArgs.ServerName, *(*RemoteUserData), *LocalPlaybackInstance->GetInstanceUserData());
+			}
+			
+			const int32 PageId = UAvaRundownPagePlayer::GetPageIdFromInstanceUserData(*RemoteUserData);
 			if (PageId != FAvaRundownPage::InvalidPageId)
 			{
 				const UAvaRundownPagePlayer* PagePlayer = Rundown->FindPlayerForProgramPage(PageId);
diff --git a/Engine/Plugins/Experimental/Avalanche/Source/AvalancheMedia/Private/Rundown/Transition/AvaRundownPageTransition.cpp b/Engine/Plugins/Experimental/Avalanche/Source/AvalancheMedia/Private/Rundown/Transition/AvaRundownPageTransition.cpp
--- a/Engine/Plugins/Experimental/Avalanche/Source/AvalancheMedia/Private/Rundown/Transition/AvaRundownPageTransition.cpp
+++ b/Engine/Plugins/Experimental/Avalanche/Source/AvalancheMedia/Private/Rundown/Transition/AvaRundownPageTransition.cpp
@@ -43,14 +43,23 @@ namespace UE::AvaMedia::Rundown::PageTransition::Private
 		return nullptr;
 	}
 
-	FString GetPrettyPageInfo(const UAvaRundownPagePlayer* InPagePlayer)
+	FString GetPrettyInstancePlayerInfo(const UAvaRundownPlaybackInstancePlayer* InInstancePlayer)
 	{
-		FString Info;
-		Info += FString(TEXT("Channel: ")) + InPagePlayer->ChannelName;
+		if (InInstancePlayer)
+		{
+			return FString::Printf(TEXT("Id: %s, Asset: %s"),
+				*InInstancePlayer->GetPlaybackInstanceId().ToString(),
+				*InInstancePlayer->SourceAssetPath.ToString()); 
+		}
+		return FString(TEXT("Invalid"));
+	}
+
+	FString GetPrettyPagePlayerInfo(const UAvaRundownPagePlayer* InPagePlayer)
+	{
+		FString Info = FString::Printf(TEXT("Channel: %s"), *InPagePlayer->ChannelName);
 		InPagePlayer->ForEachInstancePlayer([&Info](const UAvaRundownPlaybackInstancePlayer* InInstancePlayer)
 		{
-			Info += FString(TEXT(", Instance Id: ")) + InInstancePlayer->GetPlaybackInstanceId().ToString();
-			Info += FString(TEXT(", Asset: ")) + InInstancePlayer->SourceAssetPath.ToString();
+			Info += FString::Printf(TEXT(", Instance {%s}"), *GetPrettyInstancePlayerInfo(InInstancePlayer));
 		});
 		return Info;
 	}
@@ -84,6 +93,13 @@ namespace UE::AvaMedia::Rundown::PageTransition::Private
 	}
 }
 
+UAvaRundownPageTransition* UAvaRundownPageTransition::MakeNew(UAvaRundown* InRundown)
+{
+	UAvaRundownPageTransition* NewTransition = NewObject<UAvaRundownPageTransition>(InRundown);
+	NewTransition->TransitionId = FGuid::NewGuid();
+	return NewTransition;
+}
+
 bool UAvaRundownPageTransition::AddEnterPage(UAvaRundownPagePlayer* InPagePlayer)
 {
 	using namespace UE::AvaMedia::Rundown::PageTransition::Private;
@@ -111,7 +127,7 @@ bool UAvaRundownPageTransition::AddEnterPage(UAvaRundownPagePlayer* InPagePlayer
 			{
 				UE_LOG(LogAvaRundown, Error,
 						TEXT("Page Transition \"%s\" Error: page %d can't be played with page %d because they are on the same layer."),
-						*GetFullName(), InPagePlayer->PageId, Player->PageId);
+						*GetInstanceName(), InPagePlayer->PageId, Player->PageId);
 				return false;
 			}
 		}
@@ -268,10 +284,23 @@ void UAvaRundownPageTransition::Stop()
 	{
 		Rundown->RemovePageTransition(this);
 		Rundown->RemoveStoppedPagePlayers();
+
+		// Make sure that there are no instance players
+		for (const TObjectPtr<UAvaRundownPagePlayer>& PagePlayer : Rundown->GetPagePlayers())
+		{
+			PagePlayer->ForEachInstancePlayer([this, PagePlayer](UAvaRundownPlaybackInstancePlayer* InInstancePlayer)
+			{
+				if (InstancesMarkedForDiscard.Contains(InInstancePlayer->GetPlaybackInstanceId()))
+				{
+					UE_LOG(LogAvaRundown, Error, TEXT("%s Page Transition \"%s\" has marked instance \"%s\" for discard but it is still playing in page %d"),
+						*UE::AvaPlayback::Utils::GetBriefFrameInfo(), *GetInstanceName(), *InInstancePlayer->GetPlaybackInstanceId().ToString(), PagePlayer->PageId);			
+				}
+			});
+		}
 	}
 	else
 	{
-		UE_LOG(LogAvaRundown, Error, TEXT("Page Transition \"%s\" Failed to remove transition: No rundown specified."), *GetFullName());
+		UE_LOG(LogAvaRundown, Error, TEXT("Page Transition \"%s\" Failed to remove transition: No rundown specified."), *GetInstanceName());
 	}
 }
 
@@ -375,8 +404,10 @@ void UAvaRundownPageTransition::OnTransitionEvent(UAvaPlayable* InPlayable, UAva
 
 	if (InstancePlayer && EnumHasAnyFlags(InTransitionFlags, EAvaPlayableTransitionEventFlags::MarkPlayableDiscard))
 	{
-		UE_LOG(LogAvaRundown, Verbose, TEXT("%s Instance Player Marked for Discard: Id:%s, Asset:\"%s\""),
-			*GetBriefFrameInfo(), *InstancePlayer->GetPlaybackInstanceId().ToString(), *InstancePlayer->SourceAssetPath.ToString());
+		InstancesMarkedForDiscard.Add(InPlayable->GetInstanceId());
+		
+		UE_LOG(LogAvaRundown, Verbose, TEXT("%s Instance Player Marked for Discard: Id:%s, Asset:\"%s\" -> Transition %s"),
+			*GetBriefFrameInfo(), *InstancePlayer->GetPlaybackInstanceId().ToString(), *InstancePlayer->SourceAssetPath.ToString(), *TransitionId.ToString());
 	}
 	
 	if (InstancePlayer && EnumHasAnyFlags(InTransitionFlags, EAvaPlayableTransitionEventFlags::StopPlayable))
@@ -385,12 +416,12 @@ void UAvaRundownPageTransition::OnTransitionEvent(UAvaPlayable* InPlayable, UAva
 		if (PlayableTransition->IsEnterPlayable(InPlayable))
 		{
 			UE_LOG(LogAvaRundown, Error,
-				TEXT("Page Transition \"%s\" Error: An \"enter\" playable is being discarded for page %d."),
-				*GetFullName(), (PagePlayer ? PagePlayer->PageId : -1));
+				TEXT("%s Page Transition \"%s\" Error: An \"enter\" playable is being discarded for page %d."),
+				*GetBriefFrameInfo(), *GetInstanceName(), (PagePlayer ? PagePlayer->PageId : -1));
 		}
 
-		UE_LOG(LogAvaRundown, Verbose, TEXT("%s Stopping Instance Player: Id:%s, Asset:\"%s\""),
-			*GetBriefFrameInfo(), *InstancePlayer->GetPlaybackInstanceId().ToString(), *InstancePlayer->SourceAssetPath.ToString());
+		UE_LOG(LogAvaRundown, Verbose, TEXT("%s Page Transition \"%s\" Stopping Instance Player: Id:%s, Asset:\"%s\""),
+			*GetBriefFrameInfo(), *GetInstanceName(), *InstancePlayer->GetPlaybackInstanceId().ToString(), *InstancePlayer->SourceAssetPath.ToString());
 
 		// With combo-templates, page players can be partially stopped.
 		InstancePlayer->Stop();
@@ -409,7 +440,7 @@ void UAvaRundownPageTransition::OnTransitionEvent(UAvaPlayable* InPlayable, UAva
 			}
 			else
 			{
-				UE_LOG(LogAvaRundown, Error, TEXT("Page Transition \"%s\" failed to remove stopped players: No rundown specified."), *GetFullName());
+				UE_LOG(LogAvaRundown, Error, TEXT("Page Transition \"%s\" failed to remove stopped players: No rundown specified."), *GetInstanceName());
 			}
 		}
 	}
@@ -479,8 +510,8 @@ void UAvaRundownPageTransition::AddPlayablesToBuilder(FAvaPlayableTransitionBuil
 		{
 			// If this happens, likely the playable is not yet loaded.
 			UE_LOG(LogAvaRundown, Error,
-				TEXT("Page Transition \"%s\" Error: Failed to retrieve \"%s\" playable for page %d."),
-				*GetFullName(), InCategory, InPlayer->PageId);
+				TEXT("%s Page Transition \"%s\" Error: Failed to retrieve \"%s\" playable for instance {%s} of page %d."),
+				*UE::AvaPlayback::Utils::GetBriefFrameInfo(), *GetInstanceName(), InCategory, *GetPrettyInstancePlayerInfo(InstancePlayer), InPlayer->PageId);
 		}
 	}
 }
@@ -491,7 +522,7 @@ void UAvaRundownPageTransition::MakePlayableTransition()
 	AddPlayersToBuilder(TransitionBuilder, EnterPlayersWeak, TEXT("Enter"), EAvaPlayableTransitionEntryRole::Enter);
 	AddPlayersToBuilder(TransitionBuilder, PlayingPlayersWeak, TEXT("Playing"), EAvaPlayableTransitionEntryRole::Playing);
 	AddPlayersToBuilder(TransitionBuilder, ExitPlayersWeak, TEXT("Exit"), EAvaPlayableTransitionEntryRole::Exit);
-	PlayableTransition = TransitionBuilder.MakeTransition(this);
+	PlayableTransition = TransitionBuilder.MakeTransition(this, TransitionId);
 
 	// For non-TL enter pages, we need to kick out the playing pages too. 
 	if (PlayableTransition && HasEnterPagesWithNoTransitionLogic())
@@ -500,6 +531,15 @@ void UAvaRundownPageTransition::MakePlayableTransition()
 	}
 }
 
+FString UAvaRundownPageTransition::GetInstanceName() const
+{
+	if (const UAvaRundown* Rundown = GetRundown())
+	{
+		return Rundown->GetName() + TEXT(":") + TransitionId.ToString();
+	}
+	return TransitionId.ToString();
+}
+
 void UAvaRundownPageTransition::LogDetailedTransitionInfo() const
 {
 	if (!UE_LOG_ACTIVE(LogAvaRundown, Verbose))
@@ -508,7 +548,7 @@ void UAvaRundownPageTransition::LogDetailedTransitionInfo() const
 	}
 	
 	using namespace UE::AvaPlayback::Utils;
-	UE_LOG(LogAvaRundown, Verbose, TEXT("%s Starting Page Transition \"%s\":"), *GetBriefFrameInfo(), *GetFullName());
+	UE_LOG(LogAvaRundown, Verbose, TEXT("%s Starting Page Transition \"%s\":"), *GetBriefFrameInfo(), *GetInstanceName());
 	
 	auto LogPlayers = [this](const TArray<TWeakObjectPtr<UAvaRundownPagePlayer>>& InPlayersWeak, const TCHAR* InCategory)
 	{
@@ -517,7 +557,7 @@ void UAvaRundownPageTransition::LogDetailedTransitionInfo() const
 		{
 			if (const UAvaRundownPagePlayer* Player = PlayerWeak.Get())
 			{
-				UE_LOG(LogAvaRundown, Verbose, TEXT("- %s Page: %d, %s."), InCategory, Player->PageId, *GetPrettyPageInfo(Player));
+				UE_LOG(LogAvaRundown, Verbose, TEXT("- %s Page: %d, %s."), InCategory, Player->PageId, *GetPrettyPagePlayerInfo(Player));
 			}
 		}
 	};
@@ -546,7 +586,7 @@ FString UAvaRundownPageTransition::GetBriefTransitionDescription() const
 	const FString PlayingPageList = MakePageIdList(PlayingPlayersWeak);
 	const FString ExitPageList = MakePageIdList(ExitPlayersWeak);
 	return FString::Printf(TEXT("Page Transition \"%s\": Enter Page(s): [%s], Playing Page(s): [%s], Exit Page(s): [%s]."),
-		*GetFullName(), *EnterPageList, *PlayingPageList, *ExitPageList);
+		*GetInstanceName(), *EnterPageList, *PlayingPageList, *ExitPageList);
 }
 
 void UAvaRundownPageTransition::RegisterToPlayableTransitionEvent()
@@ -639,12 +679,13 @@ void UAvaRundownPageTransition::UpdateChannelName(const UAvaRundownPagePlayer* I
 	else
 	{
 		using namespace UE::AvaMedia::Rundown::PageTransition::Private;
+		using namespace UE::AvaPlayback::Utils;
 
 		// Validate the channel is the same.
 		if (ChannelName != InPagePlayer->ChannelFName)
 		{
-			UE_LOG(LogAvaRundown, Error, TEXT("Page Transition \"%s\": Adding Page: %d, {%s} in a different channel than previous pages (\"%s\")."),
-				*GetFullName(), InPagePlayer->PageId, *GetPrettyPageInfo(InPagePlayer), *ChannelName.ToString());
+			UE_LOG(LogAvaRundown, Error, TEXT("%s Page Transition \"%s\": Adding Page: %d, {%s} in a different channel than previous pages (\"%s\")."),
+				*GetBriefFrameInfo(), *GetInstanceName(), InPagePlayer->PageId, *GetPrettyPagePlayerInfo(InPagePlayer), *ChannelName.ToString());
 		}
 	}
 }
diff --git a/Engine/Plugins/Experimental/Avalanche/Source/AvalancheMedia/Private/Rundown/Transition/AvaRundownPageTransition.h b/Engine/Plugins/Experimental/Avalanche/Source/AvalancheMedia/Private/Rundown/Transition/AvaRundownPageTransition.h
--- a/Engine/Plugins/Experimental/Avalanche/Source/AvalancheMedia/Private/Rundown/Transition/AvaRundownPageTransition.h
+++ b/Engine/Plugins/Experimental/Avalanche/Source/AvalancheMedia/Private/Rundown/Transition/AvaRundownPageTransition.h
@@ -27,6 +27,8 @@ class UAvaRundownPageTransition : public UAvaPlaybackTransition
 	GENERATED_BODY()
 	
 public:
+	static UAvaRundownPageTransition* MakeNew(UAvaRundown* InRundown);
+	
 	bool AddEnterPage(UAvaRundownPagePlayer* InPagePlayer);
 	bool AddPlayingPage(UAvaRundownPagePlayer* InPagePlayer);
 	bool AddExitPage(UAvaRundownPagePlayer* InPagePlayer);
@@ -67,6 +69,7 @@ protected:
 
 	void MakePlayableTransition();
 
+	FString GetInstanceName() const;
 	void LogDetailedTransitionInfo() const;
 	FString GetBriefTransitionDescription() const;
 
@@ -88,6 +91,8 @@ protected:
 
 	TSet<FAvaTagId> CachedTransitionLayers;
 	
+	TSet<FGuid> InstancesMarkedForDiscard;
+	
 	UPROPERTY(Transient)
 	TObjectPtr<UAvaPlayableTransition> PlayableTransition;
 };
\ No newline at end of file
diff --git a/Engine/Plugins/Experimental/Avalanche/Source/AvalancheMedia/Private/Rundown/Transition/AvaRundownPageTransitionBuilder.cpp b/Engine/Plugins/Experimental/Avalanche/Source/AvalancheMedia/Private/Rundown/Transition/AvaRundownPageTransitionBuilder.cpp
--- a/Engine/Plugins/Experimental/Avalanche/Source/AvalancheMedia/Private/Rundown/Transition/AvaRundownPageTransitionBuilder.cpp
+++ b/Engine/Plugins/Experimental/Avalanche/Source/AvalancheMedia/Private/Rundown/Transition/AvaRundownPageTransitionBuilder.cpp
@@ -60,7 +60,7 @@ UAvaRundownPageTransition* FAvaRundownPageTransitionBuilder::FindOrAddTransition
 		return PageTransition;
 	}
 		
-	UAvaRundownPageTransition* PageTransition = NewObject<UAvaRundownPageTransition>(Rundown);
+	UAvaRundownPageTransition* PageTransition = UAvaRundownPageTransition::MakeNew(Rundown);
 	if (PageTransition)
 	{
 		PageTransitions.Add(PageTransition);
diff --git a/Engine/Plugins/Experimental/Avalanche/Source/AvalancheMedia/Public/AvaMediaDefines.h b/Engine/Plugins/Experimental/Avalanche/Source/AvalancheMedia/Public/AvaMediaDefines.h
--- a/Engine/Plugins/Experimental/Avalanche/Source/AvalancheMedia/Public/AvaMediaDefines.h
+++ b/Engine/Plugins/Experimental/Avalanche/Source/AvalancheMedia/Public/AvaMediaDefines.h
@@ -134,12 +134,14 @@ enum class EAvaPlayableTransitionEventFlags : uint8
 	None = 0,
 	/** The transition is starting. */
 	Starting = 1 << 0,
+	/** The enter playable can be shown. */
+	ShowPlayable = 1 << 1,
 	/** The playable needs to be stopped. */
-	StopPlayable = 1 << 1,
+	StopPlayable = 1 << 2,
 	/** The playable needs to be discarded at the end of the transition. */
-	MarkPlayableDiscard = 1 << 2,
+	MarkPlayableDiscard = 1 << 3,
 	/** The transition is finished and can be cleaned up. */
-	Finished = 1 << 3,
+	Finished = 1 << 4,
 };
 ENUM_CLASS_FLAGS(EAvaPlayableTransitionEventFlags);
 
diff --git a/Engine/Plugins/Experimental/Avalanche/Source/AvalancheMedia/Public/AvaMediaSettings.h b/Engine/Plugins/Experimental/Avalanche/Source/AvalancheMedia/Public/AvaMediaSettings.h
--- a/Engine/Plugins/Experimental/Avalanche/Source/AvalancheMedia/Public/AvaMediaSettings.h
+++ b/Engine/Plugins/Experimental/Avalanche/Source/AvalancheMedia/Public/AvaMediaSettings.h
@@ -11,6 +11,7 @@
 #include "Math/MathFwd.h"
 #include "PixelFormat.h"
 #include "UObject/SoftObjectPtr.h"
+
 #include "AvaMediaSettings.generated.h"
 
 class UUserWidget;
@@ -166,10 +167,6 @@ public:
 	UPROPERTY(config, EditAnywhere, Category = "Playback Server")
 	FString PlaybackServerName;
 
-	/** Enable verbose logging for playback server. */
-	UPROPERTY(config, EditAnywhere, Category = "Playback Server")
-	bool bVerbosePlaybackServerLogging = false;
-
 	/**
 	 * Determines the verbosity level of the playback server's log replication.
 	 * The server is not going to replicate any log event that is below this log level.
@@ -181,6 +178,10 @@ public:
 	UPROPERTY(config, EditAnywhere, Category = "Playback Server")
 	float ServerPendingStatusRequestTimeout = 5.0f;
 
+	/** Defines the timeout, in seconds, after which pending playback commands are discarded. */
+	UPROPERTY(config, EditAnywhere, Category = "Playback Server")
+	float ServerPendingPlaybackCommandTimeout = 5.0f;
+
 	/** Settings for the local playback server process. See "Launch Local Server" in the broadcast editor toolbar. */
 	UPROPERTY(config, EditAnywhere, Category = "Playback Server")
 	FAvaMediaLocalPlaybackServerSettings LocalPlaybackServerSettings;
@@ -207,6 +208,9 @@ public:
 	UPROPERTY(config, EditAnywhere, Category = "Web Server")
 	uint32 HttpServerPort = 10123;
 
+	/** Default value of the synchronized events feature selection. */
+	static const FName SynchronizedEventsFeatureSelection_Default;
+
 private:
 	static UAvaMediaSettings* GetSingletonInstance();
 };
diff --git a/Engine/Plugins/Experimental/Avalanche/Source/AvalancheMedia/Public/ModularFeature/AvaMediaSynchronizedEvent.h b/Engine/Plugins/Experimental/Avalanche/Source/AvalancheMedia/Public/ModularFeature/AvaMediaSynchronizedEvent.h
new file mode 100644
--- /dev/null
+++ b/Engine/Plugins/Experimental/Avalanche/Source/AvalancheMedia/Public/ModularFeature/AvaMediaSynchronizedEvent.h
@@ -0,0 +1,29 @@
+// Copyright Epic Games, Inc. All Rights Reserved.
+
+#pragma once
+
+#include "Containers/UnrealString.h"
+#include "Templates/Function.h"
+
+/**
+ * @brief Base class for synchronized events
+ */
+struct FAvaMediaSynchronizedEvent
+{
+	FString Signature;
+	TUniqueFunction<void()> Function;
+
+	FAvaMediaSynchronizedEvent(const FString& InSignature, TUniqueFunction<void()> InFunction)
+		: Signature(InSignature)
+		, Function(MoveTemp(InFunction))
+	{
+	}
+	
+	FAvaMediaSynchronizedEvent(FString&& InSignature, TUniqueFunction<void()> InFunction)
+		: Signature(MoveTemp(InSignature))
+		, Function(MoveTemp(InFunction))
+	{
+	}
+
+	virtual ~FAvaMediaSynchronizedEvent() = default;
+};
diff --git a/Engine/Plugins/Experimental/Avalanche/Source/AvalancheMedia/Public/ModularFeature/IAvaMediaSynchronizedEventDispatcher.h b/Engine/Plugins/Experimental/Avalanche/Source/AvalancheMedia/Public/ModularFeature/IAvaMediaSynchronizedEventDispatcher.h
new file mode 100644
--- /dev/null
+++ b/Engine/Plugins/Experimental/Avalanche/Source/AvalancheMedia/Public/ModularFeature/IAvaMediaSynchronizedEventDispatcher.h
@@ -0,0 +1,54 @@
+// Copyright Epic Games, Inc. All Rights Reserved.
+
+#pragma once
+
+#include "Containers/UnrealString.h"
+#include "Templates/Function.h"
+
+enum class EAvaMediaSynchronizedEventState : uint8
+{
+	/**
+	 * The requested event signature is not tracked by the dispatcher.
+	 */
+	NotFound,
+	/**
+	 * Has been seen on other nodes, but has not been pushed locally yet.
+	 */
+	Tracked,
+	/**
+	 * Has been pushed locally and is waiting on other nodes. 
+	 */
+	Pending,
+	/**
+	 * Has been marked from all node and will be invoke on next dispatch call.
+	 * This state can only been seen on "late" dispatch implementation. Ready events
+	 * should be dispatched asap and not linger in the ready queue.
+	 */
+	Ready
+};
+
+/**
+ * Queue and Dispatcher for Synchronized events.
+ */
+class IAvaMediaSynchronizedEventDispatcher
+{
+public:
+	/**
+	 * Push a new punctual event in the queue.
+	 * @param InEventSignature unique signature for the event.
+	 * @param InFunction function to invoke when the event is signaled on all nodes.
+	 */
+	virtual bool PushEvent(FString&& InEventSignature, TUniqueFunction<void()> InFunction) = 0;
+
+	/**
+	 * Retrieve the given event state.
+	 */
+	virtual EAvaMediaSynchronizedEventState GetEventState(const FString& InEventSignature) const = 0;
+	
+	/**
+	 * Let the implementation update and dispatch the queued events. 
+	 */
+	virtual void DispatchEvents() = 0;
+	
+	virtual ~IAvaMediaSynchronizedEventDispatcher() = default;
+};
diff --git a/Engine/Plugins/Experimental/Avalanche/Source/AvalancheMedia/Public/ModularFeature/IAvaMediaSynchronizedEventsFeature.h b/Engine/Plugins/Experimental/Avalanche/Source/AvalancheMedia/Public/ModularFeature/IAvaMediaSynchronizedEventsFeature.h
new file mode 100644
--- /dev/null
+++ b/Engine/Plugins/Experimental/Avalanche/Source/AvalancheMedia/Public/ModularFeature/IAvaMediaSynchronizedEventsFeature.h
@@ -0,0 +1,63 @@
+// Copyright Epic Games, Inc. All Rights Reserved.
+
+#pragma once
+
+#include "Features/IModularFeature.h"
+#include "Templates/SharedPointer.h"
+#include "UObject/NameTypes.h"
+
+class IAvaMediaSynchronizedEventDispatcher;
+
+/**
+ * Interface for synchronized events implementation in a distributed environment such as
+ * the Display Cluster or forked channels.
+ * 
+ * For a given event signature across a distributed environment, i.e. the event itself
+ * is distributed (one instance per node), the event will be fired when all nodes raise the
+ * signal for that event. In other word, the event is invoked at the same time on all nodes only when
+ * it has been raised on all nodes.
+ */
+class IAvaMediaSynchronizedEventsFeature : public IModularFeature
+{
+public:
+	static FName GetModularFeatureName()
+	{
+		static const FName FeatureName = TEXT("AvaMediaSynchonizedEventsFeature");
+		return FeatureName;
+	}
+	
+	/** Returns underlying feature implementation name */
+	virtual FName GetName() const = 0;
+
+	/** Returns the localized feature implementation name for display. */ 
+	virtual FText GetDisplayName() const = 0;
+
+	/** Returns the feature implementation description. */ 
+	virtual FText GetDisplayDescription() const = 0;
+
+	/**
+	 * Returns the feature implementation priority for automatic selection of most appropriate default.
+	 * Implementations with higher priority will take precedence over lower ones.
+	 */
+	virtual int32 GetPriority() const = 0;
+
+	/** Returns a default implementation priority. */
+	static int32 GetDefaultPriority()
+	{
+		// This needs to be greater than 0 to allow factories to have both higher and lower priority than the default.
+		constexpr int32 DefaultImportPriority = 100;
+		return DefaultImportPriority;
+	}
+	
+	/**
+	 * Factory method to create an event dispatcher.
+	 * 
+	 * @param InSignature signature given to the dispatcher. This should match on all clustered/forked channel.
+	 * @remark Current implementation does not use the signature to automatically add scope to all pushed events.
+	 * @todo Investigate the convenience of using dispatcher signature to automatically add scope to all pushed events.
+	 */
+	virtual TSharedPtr<IAvaMediaSynchronizedEventDispatcher> CreateDispatcher(const FString& InSignature) = 0;
+
+protected:
+	virtual ~IAvaMediaSynchronizedEventsFeature() = default;
+};
\ No newline at end of file
diff --git a/Engine/Plugins/Experimental/Avalanche/Source/AvalancheMedia/Public/Playable/AvaPlayable.h b/Engine/Plugins/Experimental/Avalanche/Source/AvalancheMedia/Public/Playable/AvaPlayable.h
--- a/Engine/Plugins/Experimental/Avalanche/Source/AvalancheMedia/Public/Playable/AvaPlayable.h
+++ b/Engine/Plugins/Experimental/Avalanche/Source/AvalancheMedia/Public/Playable/AvaPlayable.h
@@ -9,8 +9,11 @@
 #include "UObject/Object.h"
 #include "UObject/ObjectPtr.h"
 #include "UObject/SoftObjectPtr.h"
+
 #include "AvaPlayable.generated.h"
 
+class FSceneView;
+class FSceneViewFamily;
 class IAvaSceneInterface;
 class UAvaPlayableGroup;
 class UAvaPlayableGroupManager;
@@ -117,6 +120,7 @@ public:
 	virtual bool IsRemoteProxy() const { return false; }
 	virtual bool GetShouldBeVisible() const { return true; }
 	virtual void SetShouldBeVisible(bool bInShouldBeVisible) {}
+	virtual void SetupView(FSceneViewFamily& InViewFamily, FSceneView& InView) {}
 	
 	/**
 	 * @brief Ensures the given asset is playing (visible) with the given parameters.
diff --git a/Engine/Plugins/Experimental/Avalanche/Source/AvalancheMedia/Public/Playable/AvaPlayableGroup.h b/Engine/Plugins/Experimental/Avalanche/Source/AvalancheMedia/Public/Playable/AvaPlayableGroup.h
--- a/Engine/Plugins/Experimental/Avalanche/Source/AvalancheMedia/Public/Playable/AvaPlayableGroup.h
+++ b/Engine/Plugins/Experimental/Avalanche/Source/AvalancheMedia/Public/Playable/AvaPlayableGroup.h
@@ -8,8 +8,11 @@
 #include "UObject/ObjectPtr.h"
 #include "UObject/SoftObjectPtr.h"
 #include "UObject/WeakInterfacePtr.h"
+
 #include "AvaPlayableGroup.generated.h"
 
+class FSceneView;
+class FSceneViewFamily;
 class UAvaPlayable;
 class UAvaPlayableGroupManager;
 class UAvaPlayableTransition;
@@ -86,6 +89,10 @@ public:
 
 	bool HasTransitions() const;
 
+	void PushSynchronizedEvent(FString&& InEventSignature, TUniqueFunction<void()> InFunction);
+
+	bool IsSynchronizedEventPushed(const FString& InEventSignature) const;
+
 	/**
 	 * Creates the game instance's world if it wasn't already.
 	 * @return true if the world was created. false if nothing was done.
@@ -164,6 +171,8 @@ public:
 
 	void RequestSetVisibility(UAvaPlayable* InPlayable, bool bInShouldBeVisible);
 
+	virtual void SetupView(FSceneViewFamily& InViewFamily, FSceneView& InView);
+
 protected:
 	bool IsVisibilityConstrained(const UAvaPlayable* InPlayable) const;
 
diff --git a/Engine/Plugins/Experimental/Avalanche/Source/AvalancheMedia/Public/Playable/Transition/AvaPlayableTransition.h b/Engine/Plugins/Experimental/Avalanche/Source/AvalancheMedia/Public/Playable/Transition/AvaPlayableTransition.h
--- a/Engine/Plugins/Experimental/Avalanche/Source/AvalancheMedia/Public/Playable/Transition/AvaPlayableTransition.h
+++ b/Engine/Plugins/Experimental/Avalanche/Source/AvalancheMedia/Public/Playable/Transition/AvaPlayableTransition.h
@@ -34,6 +34,7 @@ public:
 	virtual bool IsRunning() const { return false; }
 	virtual void Tick(double InDeltaSeconds) {}
 
+	void SetTransitionId(const FGuid& InTransitionId);
 	void SetTransitionFlags(EAvaPlayableTransitionFlags InFlags);
 	void SetEnterPlayables(TArray<TWeakObjectPtr<UAvaPlayable>>&& InPlayablesWeak);
 	void SetPlayingPlayables(TArray<TWeakObjectPtr<UAvaPlayable>>&& InPlayablesWeak);
@@ -57,6 +58,7 @@ protected:
 	UAvaPlayable* FindPlayable(const FGuid& InstanceId) const;
 	
 protected:
+	FGuid TransitionId;
 	EAvaPlayableTransitionFlags TransitionFlags = EAvaPlayableTransitionFlags::None;
 	
 	TArray<TSharedPtr<FAvaPlayableRemoteControlValues>> EnterPlayableValues;
@@ -84,7 +86,7 @@ public:
 	
 	bool AddPlayable(UAvaPlayable* InPlayable, EAvaPlayableTransitionEntryRole InPlayableRole)
 	{
-		switch(InPlayableRole)
+		switch (InPlayableRole)
 		{
 			case EAvaPlayableTransitionEntryRole::Enter:
 				return AddEnterPlayable(InPlayable);
@@ -96,7 +98,7 @@ public:
 		return false;
 	}
 
-	UAvaPlayableTransition* MakeTransition(UObject* InOuter);
+	UAvaPlayableTransition* MakeTransition(UObject* InOuter, const FGuid& InTransitionId = FGuid());
 
 private:
 	TArray<TSharedPtr<FAvaPlayableRemoteControlValues>> EnterPlayableValues;
diff --git a/Engine/Plugins/Experimental/Avalanche/Source/AvalancheMedia/Public/Playback/AvaPlaybackClientDelegates.h b/Engine/Plugins/Experimental/Avalanche/Source/AvalancheMedia/Public/Playback/AvaPlaybackClientDelegates.h
--- a/Engine/Plugins/Experimental/Avalanche/Source/AvalancheMedia/Public/Playback/AvaPlaybackClientDelegates.h
+++ b/Engine/Plugins/Experimental/Avalanche/Source/AvalancheMedia/Public/Playback/AvaPlaybackClientDelegates.h
@@ -42,6 +42,7 @@ namespace UE::AvaPlaybackClient::Delegates
 		const FString& ServerName;
 		const FString& SequenceName;
 		EAvaPlayableSequenceEventType EventType;
+		int32 ServerFrameNumber;
 	};
 
 	struct FPlaybackTransitionEventArgs
@@ -51,6 +52,7 @@ namespace UE::AvaPlaybackClient::Delegates
 		const FString& ChannelName;
 		const FString& ServerName;
 		EAvaPlayableTransitionEventFlags EventFlags;
+		int32 ServerFrameNumber;
 	};
 
 	enum class EConnectionEvent : uint8
diff --git a/Engine/Plugins/Experimental/Avalanche/Source/AvalancheMedia/Public/Playback/IAvaPlaybackClient.h b/Engine/Plugins/Experimental/Avalanche/Source/AvalancheMedia/Public/Playback/IAvaPlaybackClient.h
--- a/Engine/Plugins/Experimental/Avalanche/Source/AvalancheMedia/Public/Playback/IAvaPlaybackClient.h
+++ b/Engine/Plugins/Experimental/Avalanche/Source/AvalancheMedia/Public/Playback/IAvaPlaybackClient.h
@@ -68,6 +68,11 @@ public:
 	 * commands to server even if they are not live yet (but connected).
 	 */
 	virtual bool HasAnyServerOnlineForChannel(const FName& InChannelName) const = 0;
+
+	/**
+	 * Returns a list of all connected servers having outputs in the given channel.
+	 */
+	virtual TArray<FString> GetOnlineServersForChannel(const FName& InChannelName) const = 0;
 	
 	/**
 	 * Returns the status of the playback asset on the given channel on the remote server.
diff --git a/Engine/Plugins/Experimental/Avalanche/Source/AvalancheMedia/Public/Rundown/AvaRundownPagePlayer.h b/Engine/Plugins/Experimental/Avalanche/Source/AvalancheMedia/Public/Rundown/AvaRundownPagePlayer.h
--- a/Engine/Plugins/Experimental/Avalanche/Source/AvalancheMedia/Public/Rundown/AvaRundownPagePlayer.h
+++ b/Engine/Plugins/Experimental/Avalanche/Source/AvalancheMedia/Public/Rundown/AvaRundownPagePlayer.h
@@ -60,6 +60,8 @@ public:
 	TObjectPtr<UAvaPlaybackGraph> Playback;
 
 	TSharedPtr<FAvaPlaybackInstance> PlaybackInstance;
+
+	TWeakObjectPtr<UAvaRundownPagePlayer> ParentPagePlayer;
 };
 
 UCLASS()
@@ -155,7 +157,7 @@ public:
 	UAvaRundownPlaybackInstancePlayer* FindInstancePlayerByAssetPath(const FSoftObjectPath& InAssetPath) const;
 	
 protected:
-	UAvaRundownPlaybackInstancePlayer* CreateAndLoadInstancePlayer(const UAvaRundown* InRundown, const FAvaRundownPage& InPage, int32 InSubPageIndex, const FGuid& InInstanceId);
+	UAvaRundownPlaybackInstancePlayer* CreateAndLoadInstancePlayer(UAvaRundown* InRundown, const FAvaRundownPage& InPage, int32 InSubPageIndex, const FGuid& InInstanceId);
 
 	void RemoveInstancePlayer(UAvaRundownPlaybackInstancePlayer* InInstancePlayer);
 	
diff --git a/Engine/Plugins/Experimental/Avalanche/Source/AvalancheMediaEditor/Private/AvaMediaEditorModule.cpp b/Engine/Plugins/Experimental/Avalanche/Source/AvalancheMediaEditor/Private/AvaMediaEditorModule.cpp
--- a/Engine/Plugins/Experimental/Avalanche/Source/AvalancheMediaEditor/Private/AvaMediaEditorModule.cpp
+++ b/Engine/Plugins/Experimental/Avalanche/Source/AvalancheMediaEditor/Private/AvaMediaEditorModule.cpp
@@ -5,9 +5,11 @@
 #include "AvaMediaEditorStyle.h"
 #include "Broadcast/AvaBroadcastEditor.h"
 #include "Broadcast/OutputDevices/AvaBroadcastMediaIOOutputConfigurationCustomization.h"
+#include "Customization/AvaSynchronizedEventsFeatureSelectionCustomization.h"
 #include "Editor.h"
 #include "Engine/GameViewportClient.h"
 #include "Framework/Application/SlateApplication.h"
+#include "Framework/AvaInstanceSettings.h"
 #include "IAvaMediaModule.h"
 #include "LevelEditor.h"
 #include "LevelEditorViewport.h"
@@ -269,6 +271,7 @@ void FAvaMediaEditorModule::RegisterCustomizations() const
 	PropertyModule.RegisterCustomPropertyTypeLayout(FMediaIOOutputConfiguration::StaticStruct()->GetFName(), FOnGetPropertyTypeCustomizationInstance::CreateStatic(&FAvaBroadcastMediaIOOutputConfigurationCustomization::MakeInstance));
 	PropertyModule.RegisterCustomPropertyTypeLayout(FAvaRundownMacroCommand::StaticStruct()->GetFName(), FOnGetPropertyTypeCustomizationInstance::CreateStatic(&FAvaRundownMacroCommandCustomization::MakeInstance));
 	PropertyModule.RegisterCustomPropertyTypeLayout(FAvaRundownMacroKeyBinding::StaticStruct()->GetFName(), FOnGetPropertyTypeCustomizationInstance::CreateStatic(&FAvaRundownMacroKeyBindingCustomization::MakeInstance));
+	PropertyModule.RegisterCustomPropertyTypeLayout(FAvaSynchronizedEventsFeatureSelection::StaticStruct()->GetFName(), FOnGetPropertyTypeCustomizationInstance::CreateStatic(&FAvaSynchronizedEventsFeatureSelectionCustomization::MakeInstance));
 }
 
 void FAvaMediaEditorModule::UnregisterCustomizations() const
@@ -277,6 +280,7 @@ void FAvaMediaEditorModule::UnregisterCustomizations() const
 	PropertyModule.UnregisterCustomPropertyTypeLayout(FMediaIOOutputConfiguration::StaticStruct()->GetFName());
 	PropertyModule.UnregisterCustomPropertyTypeLayout(FAvaRundownMacroCommand::StaticStruct()->GetFName());
 	PropertyModule.UnregisterCustomPropertyTypeLayout(FAvaRundownMacroKeyBinding::StaticStruct()->GetFName());
+	PropertyModule.UnregisterCustomPropertyTypeLayout(FAvaSynchronizedEventsFeatureSelection::StaticStruct()->GetFName());
 }
 
 void FAvaMediaEditorModule::PostEngineInit()
diff --git a/Engine/Plugins/Experimental/Avalanche/Source/AvalancheMediaEditor/Private/Customization/AvaSynchronizedEventsFeatureSelectionCustomization.cpp b/Engine/Plugins/Experimental/Avalanche/Source/AvalancheMediaEditor/Private/Customization/AvaSynchronizedEventsFeatureSelectionCustomization.cpp
new file mode 100644
--- /dev/null
+++ b/Engine/Plugins/Experimental/Avalanche/Source/AvalancheMediaEditor/Private/Customization/AvaSynchronizedEventsFeatureSelectionCustomization.cpp
@@ -0,0 +1,36 @@
+// Copyright Epic Games, Inc. All Rights Reserved.
+
+#include "AvaSynchronizedEventsFeatureSelectionCustomization.h"
+
+#include "DetailWidgetRow.h"
+#include "Framework/AvaInstanceSettings.h"
+#include "IDetailChildrenBuilder.h"
+#include "SAvaSynchronizedEventsImplementationSelector.h"
+
+TSharedRef<IPropertyTypeCustomization> FAvaSynchronizedEventsFeatureSelectionCustomization::MakeInstance()
+{
+	return MakeShared<FAvaSynchronizedEventsFeatureSelectionCustomization>();
+}
+
+void FAvaSynchronizedEventsFeatureSelectionCustomization::CustomizeHeader(TSharedRef<IPropertyHandle> InPropertyHandle, FDetailWidgetRow& InHeaderRow, IPropertyTypeCustomizationUtils& InPropertyTypeCustomizationUtils)
+{
+
+}
+
+void FAvaSynchronizedEventsFeatureSelectionCustomization::CustomizeChildren(TSharedRef<IPropertyHandle> InStructPropertyHandle, IDetailChildrenBuilder& InStructBuilder, IPropertyTypeCustomizationUtils& InStructCustomizationUtils)
+{
+	const TSharedPtr<IPropertyHandle> ImplementationHandle = InStructPropertyHandle->GetChildHandle(GET_MEMBER_NAME_CHECKED(FAvaSynchronizedEventsFeatureSelection, Implementation));
+	
+	if (ImplementationHandle.IsValid())
+	{
+		InStructBuilder.AddProperty(ImplementationHandle.ToSharedRef()).CustomWidget()
+			.NameContent()
+			[
+				InStructPropertyHandle->CreatePropertyNameWidget()
+			]
+			.ValueContent()
+			[
+				SNew(SAvaSynchronizedEventsImplementationSelector, ImplementationHandle.ToSharedRef())
+			];
+	}
+}
\ No newline at end of file
diff --git a/Engine/Plugins/Experimental/Avalanche/Source/AvalancheMediaEditor/Private/Customization/AvaSynchronizedEventsFeatureSelectionCustomization.h b/Engine/Plugins/Experimental/Avalanche/Source/AvalancheMediaEditor/Private/Customization/AvaSynchronizedEventsFeatureSelectionCustomization.h
new file mode 100644
--- /dev/null
+++ b/Engine/Plugins/Experimental/Avalanche/Source/AvalancheMediaEditor/Private/Customization/AvaSynchronizedEventsFeatureSelectionCustomization.h
@@ -0,0 +1,19 @@
+// Copyright Epic Games, Inc. All Rights Reserved.
+
+#pragma once
+
+#include "IPropertyTypeCustomization.h"
+#include "PropertyHandle.h"
+
+class IPropertyHandle;
+
+class FAvaSynchronizedEventsFeatureSelectionCustomization : public IPropertyTypeCustomization
+{
+public:
+	static TSharedRef<IPropertyTypeCustomization> MakeInstance();
+	
+	//~ Begin IPropertyTypeCustomization
+	virtual void CustomizeHeader(TSharedRef<IPropertyHandle> InPropertyHandle, FDetailWidgetRow& InHeaderRow, IPropertyTypeCustomizationUtils& InPropertyTypeCustomizationUtils) override;
+	virtual void CustomizeChildren(TSharedRef<IPropertyHandle> InPropertyHandle, IDetailChildrenBuilder& InStructBuilder, IPropertyTypeCustomizationUtils& InPropertyTypeCustomizationUtils) override;
+	//~ End IPropertyTypeCustomization
+};
diff --git a/Engine/Plugins/Experimental/Avalanche/Source/AvalancheMediaEditor/Private/Customization/SAvaSynchronizedEventsImplementationSelector.cpp b/Engine/Plugins/Experimental/Avalanche/Source/AvalancheMediaEditor/Private/Customization/SAvaSynchronizedEventsImplementationSelector.cpp
new file mode 100644
--- /dev/null
+++ b/Engine/Plugins/Experimental/Avalanche/Source/AvalancheMediaEditor/Private/Customization/SAvaSynchronizedEventsImplementationSelector.cpp
@@ -0,0 +1,104 @@
+// Copyright Epic Games, Inc. All Rights Reserved.
+
+#include "SAvaSynchronizedEventsImplementationSelector.h"
+
+#include "AvaMediaSettings.h"
+#include "DetailLayoutBuilder.h"
+#include "ModularFeature/AvaMediaSynchronizedEventsFeature.h"
+#include "Widgets/Input/SComboBox.h"
+
+void SAvaSynchronizedEventsImplementationSelector::Construct(const FArguments& InArgs, TSharedPtr<IPropertyHandle> InPropertyHandle)
+{
+	UpdateItems();
+	PropertyHandle = InPropertyHandle;
+	
+	ChildSlot
+	[
+		SAssignNew(Combo, SComboBox<FName>)
+		.InitiallySelectedItem(GetItemFromProperty())
+		.OptionsSource(&Items)
+		.OnGenerateWidget(this, &SAvaSynchronizedEventsImplementationSelector::GenerateWidget)
+		.OnSelectionChanged(this, &SAvaSynchronizedEventsImplementationSelector::HandleSelectionChanged)
+		.OnComboBoxOpening(this, &SAvaSynchronizedEventsImplementationSelector::OnComboBoxOpening)
+		[
+			SNew(STextBlock)
+			.Text(this, &SAvaSynchronizedEventsImplementationSelector::GetDisplayTextFromProperty)
+			.Font(IDetailLayoutBuilder::GetDetailFont())
+		]
+	];
+}
+
+TSharedRef<SWidget> SAvaSynchronizedEventsImplementationSelector::GenerateWidget(FName InItem)
+{
+	return SNew(STextBlock)
+		.Text(GetDisplayTextFromItem(InItem))
+		.ToolTipText(GetDisplayDescriptionFromItem(InItem))
+		.Font(IDetailLayoutBuilder::GetDetailFont());
+}
+
+void SAvaSynchronizedEventsImplementationSelector::HandleSelectionChanged(FName InProposedSelection, ESelectInfo::Type InSelectInfo)
+{
+	if (PropertyHandle)
+	{
+		PropertyHandle->SetValue(InProposedSelection.ToString());
+	}
+}
+
+FText SAvaSynchronizedEventsImplementationSelector::GetDisplayTextFromProperty() const
+{
+	return GetDisplayTextFromItem(GetItemFromProperty());
+}
+
+void SAvaSynchronizedEventsImplementationSelector::OnComboBoxOpening()
+{
+	check(Combo.IsValid());
+	Combo->SetSelectedItem(GetItemFromProperty());
+}
+
+FName SAvaSynchronizedEventsImplementationSelector::GetItemFromProperty() const
+{
+	if (PropertyHandle)
+	{
+		FString Value;
+		PropertyHandle->GetValue(Value);
+		return FName(*Value);
+	}
+	return FName();
+}
+
+FText SAvaSynchronizedEventsImplementationSelector::GetDisplayTextFromItem(FName InItem) const
+{
+	if (const IAvaMediaSynchronizedEventsFeature* Implementation = FAvaMediaSynchronizedEventsFeature::FindImplementation(InItem))
+	{
+		if (InItem == UAvaMediaSettings::SynchronizedEventsFeatureSelection_Default)
+		{
+			return FText::Format(INVTEXT("{0} ({1})"), FText::FromName(InItem), Implementation->GetDisplayName());
+		}
+		else
+		{
+			return Implementation->GetDisplayName();
+		}
+	}
+	return FText::FromName(InItem);
+}
+
+FText SAvaSynchronizedEventsImplementationSelector::GetDisplayDescriptionFromItem(FName InItem) const
+{
+	if (const IAvaMediaSynchronizedEventsFeature* Implementation = FAvaMediaSynchronizedEventsFeature::FindImplementation(InItem))
+	{
+		return Implementation->GetDisplayDescription();
+	}
+	return FText::FromName(InItem);
+}
+
+void SAvaSynchronizedEventsImplementationSelector::UpdateItems()
+{
+	Items.Reset();
+
+	FAvaMediaSynchronizedEventsFeature::EnumerateImplementations([this](const IAvaMediaSynchronizedEventsFeature* InFeature)
+	{
+		Items.Add(InFeature->GetName());
+	});
+
+	Items.Add(UAvaMediaSettings::SynchronizedEventsFeatureSelection_Default);
+}
diff --git a/Engine/Plugins/Experimental/Avalanche/Source/AvalancheMediaEditor/Private/Customization/SAvaSynchronizedEventsImplementationSelector.h b/Engine/Plugins/Experimental/Avalanche/Source/AvalancheMediaEditor/Private/Customization/SAvaSynchronizedEventsImplementationSelector.h
new file mode 100644
--- /dev/null
+++ b/Engine/Plugins/Experimental/Avalanche/Source/AvalancheMediaEditor/Private/Customization/SAvaSynchronizedEventsImplementationSelector.h
@@ -0,0 +1,40 @@
+// Copyright Epic Games, Inc. All Rights Reserved.
+
+#pragma once
+
+#include "Widgets/SCompoundWidget.h"
+
+class IPropertyHandle;
+template <typename OptionType> class SComboBox;
+
+class SAvaSynchronizedEventsImplementationSelector : public SCompoundWidget
+{
+public:
+	SLATE_BEGIN_ARGS(SAvaSynchronizedEventsImplementationSelector){}
+	SLATE_END_ARGS()
+
+	/** Constructs this widget with InArgs */
+	void Construct(const FArguments& InArgs, TSharedPtr<IPropertyHandle> InPropertyHandle);
+	
+	TSharedRef<SWidget> GenerateWidget(FName InItem);
+	
+	void HandleSelectionChanged(FName InProposedSelection, ESelectInfo::Type InSelectInfo);
+	
+	FText GetDisplayTextFromProperty() const;
+	
+	void OnComboBoxOpening();
+	
+protected:
+	FName GetItemFromProperty() const;
+	
+	FText GetDisplayTextFromItem(FName InItem) const;
+	FText GetDisplayDescriptionFromItem(FName InItem) const;
+	
+	void UpdateItems();
+	
+	TSharedPtr<IPropertyHandle> PropertyHandle;
+	
+	TSharedPtr<SComboBox<FName>> Combo;
+
+	TArray<FName> Items;
+};
diff --git a/Engine/Plugins/Experimental/AvalancheDisplayCluster/AvalancheDisplayCluster.uplugin b/Engine/Plugins/Experimental/AvalancheDisplayCluster/AvalancheDisplayCluster.uplugin
new file mode 100644
--- /dev/null
+++ b/Engine/Plugins/Experimental/AvalancheDisplayCluster/AvalancheDisplayCluster.uplugin
@@ -0,0 +1,40 @@
+{
+	"FileVersion": 3,
+	"Version": 1,
+	"VersionName": "1.0",
+	"FriendlyName": "Motion Design For nDisplay",
+	"Description": "Motion Design Synchronization extensions for nDisplay clustering",
+	"Category": "Experimental",
+	"CreatedBy": "Epic Games, Inc.",
+	"CreatedByURL": "https://epicgames.com",
+	"DocsURL": "",
+	"MarketplaceURL": "",
+	"SupportURL": "",
+	"CanContainContent": true,
+	"IsBetaVersion": false,
+	"IsExperimentalVersion": true,
+	"Installed": false,
+	"EnabledByDefault": false,
+	"SupportedTargetPlatforms": [
+		"Win64",
+		"Linux",
+		"Mac"
+	],
+	"Modules": [
+		{
+			"Name": "AvalancheDisplayCluster",
+			"Type": "Runtime",
+			"LoadingPhase": "Default"
+		}
+	],
+	"Plugins": [
+		{
+			"Name": "Avalanche",
+			"Enabled": true
+		},
+		{
+			"Name": "nDisplay",
+			"Enabled": true
+		}
+	]
+}
\ No newline at end of file
diff --git a/Engine/Plugins/Experimental/AvalancheDisplayCluster/Source/AvalancheDisplayCluster/AvalancheDisplayCluster.Build.cs b/Engine/Plugins/Experimental/AvalancheDisplayCluster/Source/AvalancheDisplayCluster/AvalancheDisplayCluster.Build.cs
new file mode 100644
--- /dev/null
+++ b/Engine/Plugins/Experimental/AvalancheDisplayCluster/Source/AvalancheDisplayCluster/AvalancheDisplayCluster.Build.cs
@@ -0,0 +1,26 @@
+// Copyright Epic Games, Inc. All Rights Reserved.
+
+using UnrealBuildTool;
+
+public class AvalancheDisplayCluster : ModuleRules
+{
+	public AvalancheDisplayCluster(ReadOnlyTargetRules Target) : base(Target)
+	{
+		PCHUsage = ModuleRules.PCHUsageMode.UseExplicitOrSharedPCHs;
+
+		PublicDependencyModuleNames.AddRange(
+			new string[]
+			{
+				"Core"
+			});
+		
+		PrivateDependencyModuleNames.AddRange(
+			new string[]
+			{
+				"AvalancheMedia",
+				"DisplayCluster",
+				"Engine"
+			}
+		);
+	}
+}
diff --git a/Engine/Plugins/Experimental/AvalancheDisplayCluster/Source/AvalancheDisplayCluster/Private/AvaDisplayClusterLog.h b/Engine/Plugins/Experimental/AvalancheDisplayCluster/Source/AvalancheDisplayCluster/Private/AvaDisplayClusterLog.h
new file mode 100644
--- /dev/null
+++ b/Engine/Plugins/Experimental/AvalancheDisplayCluster/Source/AvalancheDisplayCluster/Private/AvaDisplayClusterLog.h
@@ -0,0 +1,7 @@
+// Copyright Epic Games, Inc. All Rights Reserved.
+
+#pragma once
+
+#include "Logging/LogMacros.h"
+
+DECLARE_LOG_CATEGORY_EXTERN(LogAvaDisplayCluster, Log, All);
\ No newline at end of file
diff --git a/Engine/Plugins/Experimental/AvalancheDisplayCluster/Source/AvalancheDisplayCluster/Private/AvaDisplayClusterModule.cpp b/Engine/Plugins/Experimental/AvalancheDisplayCluster/Source/AvalancheDisplayCluster/Private/AvaDisplayClusterModule.cpp
new file mode 100644
--- /dev/null
+++ b/Engine/Plugins/Experimental/AvalancheDisplayCluster/Source/AvalancheDisplayCluster/Private/AvaDisplayClusterModule.cpp
@@ -0,0 +1,24 @@
+// Copyright Epic Games, Inc. All Rights Reserved.
+
+#include "AvaDisplayClusterModule.h"
+
+#include "AvaDisplayClusterLog.h"
+#include "Features/IModularFeatures.h"
+#include "Modules/ModuleManager.h"
+#include "SynchronizedEvents/AvaDisplayClusterSynchronizedEventsFeature.h"
+
+DEFINE_LOG_CATEGORY(LogAvaDisplayCluster);
+
+void FAvaDisplayClusterModule::StartupModule()
+{
+	SyncEventsFeature = MakeUnique<FAvaDisplayClusterSynchronizedEventsFeature>();
+	IModularFeatures::Get().RegisterModularFeature(IAvaMediaSynchronizedEventsFeature::GetModularFeatureName(), SyncEventsFeature.Get());
+}
+
+void FAvaDisplayClusterModule::ShutdownModule()
+{
+	IModularFeatures::Get().UnregisterModularFeature(IAvaMediaSynchronizedEventsFeature::GetModularFeatureName(), SyncEventsFeature.Get());
+	SyncEventsFeature.Reset();
+}
+
+IMPLEMENT_MODULE(FAvaDisplayClusterModule, AvalancheDisplayCluster)
diff --git a/Engine/Plugins/Experimental/AvalancheDisplayCluster/Source/AvalancheDisplayCluster/Private/AvaDisplayClusterModule.h b/Engine/Plugins/Experimental/AvalancheDisplayCluster/Source/AvalancheDisplayCluster/Private/AvaDisplayClusterModule.h
new file mode 100644
--- /dev/null
+++ b/Engine/Plugins/Experimental/AvalancheDisplayCluster/Source/AvalancheDisplayCluster/Private/AvaDisplayClusterModule.h
@@ -0,0 +1,21 @@
+// Copyright Epic Games, Inc. All Rights Reserved.
+
+#pragma once
+
+#include "Containers/UnrealString.h"
+#include "Modules/ModuleInterface.h"
+#include "Templates/UniquePtr.h"
+
+class FAvaDisplayClusterSynchronizedEventsFeature;
+
+class FAvaDisplayClusterModule : public IModuleInterface
+{
+public:
+	//~ Begin IModuleInterface
+	virtual void StartupModule() override;
+	virtual void ShutdownModule() override;
+	//~ End IModuleInterface
+
+private:
+	TUniquePtr<FAvaDisplayClusterSynchronizedEventsFeature> SyncEventsFeature;
+};
diff --git a/Engine/Plugins/Experimental/AvalancheDisplayCluster/Source/AvalancheDisplayCluster/Private/SynchronizedEvents/AvaDisplayClusterSyncEventsLog.h b/Engine/Plugins/Experimental/AvalancheDisplayCluster/Source/AvalancheDisplayCluster/Private/SynchronizedEvents/AvaDisplayClusterSyncEventsLog.h
new file mode 100644
--- /dev/null
+++ b/Engine/Plugins/Experimental/AvalancheDisplayCluster/Source/AvalancheDisplayCluster/Private/SynchronizedEvents/AvaDisplayClusterSyncEventsLog.h
@@ -0,0 +1,7 @@
+// Copyright Epic Games, Inc. All Rights Reserved.
+
+#pragma once
+
+#include "Logging/LogMacros.h"
+
+DECLARE_LOG_CATEGORY_EXTERN(LogAvaDisplayClusterSyncEvents, Log, All);
\ No newline at end of file
diff --git a/Engine/Plugins/Experimental/AvalancheDisplayCluster/Source/AvalancheDisplayCluster/Private/SynchronizedEvents/AvaDisplayClusterSynchronizedEventsDispatcher.cpp b/Engine/Plugins/Experimental/AvalancheDisplayCluster/Source/AvalancheDisplayCluster/Private/SynchronizedEvents/AvaDisplayClusterSynchronizedEventsDispatcher.cpp
new file mode 100644
--- /dev/null
+++ b/Engine/Plugins/Experimental/AvalancheDisplayCluster/Source/AvalancheDisplayCluster/Private/SynchronizedEvents/AvaDisplayClusterSynchronizedEventsDispatcher.cpp
@@ -0,0 +1,327 @@
+// Copyright Epic Games, Inc. All Rights Reserved.
+
+#include "AvaDisplayClusterSynchronizedEventsDispatcher.h"
+
+#include "AvaDisplayClusterSyncEventsLog.h"
+#include "Cluster/DisplayClusterClusterEvent.h"
+#include "Cluster/IDisplayClusterClusterManager.h"
+#include "HAL/IConsoleManager.h"
+#include "IDisplayCluster.h"
+#include "Playback/AvaPlaybackUtils.h"
+#include "Serialization/ArrayWriter.h"
+
+namespace UE::AvaDisplayCluster::Private
+{
+	// Allow events to be dispatched as early as possible (at most 1 frame earlier).
+	TAutoConsoleVariable<bool> CVarSyncEarlyDispatch(
+		TEXT("AvaDisplayCluster.Sync.EarlyDispatch")
+		, false
+		, TEXT("If true, will dispatch events as soon as ready. if false, ready events are all batched on the next tick."), ECVF_Cheat);
+
+	static TAutoConsoleVariable<float> CVarSyncDispatchTimeout(
+		TEXT("AvaDisplayCluster.Sync.DispatchTimeout"), 5000.0f,
+		TEXT("Delay after witch the event is dispatched even if not signaled by all the other nodes. Units: milliseconds"),
+		ECVF_Default);
+
+	static TAutoConsoleVariable<float> CVarSyncRepeatTimeout(
+		TEXT("AvaDisplayCluster.Sync.RepeatTimeout"), 200.0f,
+		TEXT("Delay after witch the event is re-emitted in case cluster events don't get through for some reason. Units: milliseconds"),
+		ECVF_Default);
+
+	static TAutoConsoleVariable<float> CVarSyncTrackingTimeout(
+		TEXT("AvaDisplayCluster.Sync.TrackingTimeout"), 5000.0f,
+		TEXT("Delay after witch the tracked events are discarded. Units: milliseconds"),
+		ECVF_Default);
+}
+
+FAvaDisplayClusterSynchronizedEventDispatcher::FAvaDisplayClusterSynchronizedEventDispatcher(const FString& InSignature)
+	: Signature(InSignature)
+{
+	if (const IDisplayClusterClusterManager* ClusterManager = IDisplayCluster::Get().GetClusterMgr())
+	{
+		NodeId = ClusterManager->GetNodeId();
+		ClusterManager->GetNodeIds(AllNodeIds);
+	}
+}
+
+FString FAvaDisplayClusterSynchronizedEventDispatcher::GetFrameInfo() const
+{
+	return FString::Printf(TEXT("%s (disp:%s)"), *UE::AvaPlayback::Utils::GetBriefFrameInfo(), *Signature);
+}
+
+void FAvaDisplayClusterSynchronizedEventDispatcher::OnClusterEventReceived(const FAvaDisplayClusterClusterEventPayload& InPayload)
+{
+	using namespace UE::AvaDisplayCluster::Private;
+
+	UE_LOG(LogAvaDisplayClusterSyncEvents, Verbose, TEXT("%s Received cluster event \"%s\" (count:%d) from \"%s\"."),
+		*GetFrameInfo(), *InPayload.Signature, InPayload.EmitCount, *InPayload.NodeId);
+
+	// Todo: Check for re-emission (if necessary)
+	//if (Payload.NodeId != NodeId && Payload.EmitCount > 1)
+	//{
+		// If it is either pending, ready or invoked, we should emit it again.
+		// It means the other node didn't receive previous signal.
+		// This will likely cause desyncs but will lead to faster recovery then waiting for a full timeout.
+	//}
+
+	// TODO: Investigate why this happens...
+	if (ReadyEvents.Contains(InPayload.Signature))
+	{
+		UE_LOG(LogAvaDisplayClusterSyncEvents, Warning, TEXT("%s Event \"%s\" is already ready on \"%s\"."),
+			*GetFrameInfo(), *InPayload.Signature, *NodeId);
+		
+		if (PendingEvents.Contains(InPayload.Signature))
+		{
+			PendingEvents.Remove(InPayload.Signature);
+			UE_LOG(LogAvaDisplayClusterSyncEvents, Warning, TEXT("%s Event \"%s\" was pending on \"%s\". Removing..."),
+				*GetFrameInfo(), *InPayload.Signature, *NodeId);
+		}
+		
+		if (TrackedEvents.Contains(InPayload.Signature))
+		{
+			TrackedEvents.Remove(InPayload.Signature);
+			UE_LOG(LogAvaDisplayClusterSyncEvents, Warning, TEXT("%s Event \"%s\" was tracked on \"%s\". Removing..."),
+				*GetFrameInfo(), *InPayload.Signature, *NodeId);
+		}
+		return;
+	}
+
+	if (TUniquePtr<FAvaDisplayClusterSynchronizedEvent>* FoundEvent = PendingEvents.Find(InPayload.Signature))
+	{
+		(*FoundEvent)->NodeInfo.Mark(InPayload.NodeId);
+
+		if ((*FoundEvent)->NodeInfo.IsAllMarked(AllNodeIds))
+		{
+			UE_LOG(LogAvaDisplayClusterSyncEvents, Verbose, TEXT("%s Event \"%s\" is ready on \"%s\"."), *GetFrameInfo(), *InPayload.Signature, *NodeId);
+
+			if (CVarSyncEarlyDispatch.GetValueOnAnyThread())
+			{
+				DispatchEvent((*FoundEvent).Get(), FAvaDisplayClusterTimeStamp::Now());
+			}
+			else
+			{
+				ReadyEvents.Add(InPayload.Signature, MoveTemp(*FoundEvent));
+			}
+			
+			PendingEvents.Remove(InPayload.Signature);
+		}
+	}
+	else
+	{
+		// Tracking an event from the same node is considered an error state. It will just expire after a while.
+		// If this happens, it is a symptom of another problem that will need to be investigated.
+		// One cause is having a dispatch timeout shorter than the repeat timeout causing local re-emission to end up tracked.
+		if (NodeId == InPayload.NodeId)
+		{
+			UE_LOG(LogAvaDisplayClusterSyncEvents, Error, TEXT("%s Tracking Event \"%s\" on \"%s\" (same node)."),
+				*GetFrameInfo(), *InPayload.Signature, *NodeId);
+		}
+
+		// Corresponding local event hasn't been pushed yet, we need to add a tracked event
+		// to mark the nodeId. It will be transferred to the real event when it is locally pushed.
+		if (const TUniquePtr<FAvaDisplayClusterTrackedClusterEvent>* FoundTrackedEvent = TrackedEvents.Find(InPayload.Signature))
+		{
+			(*FoundTrackedEvent)->NodeInfo.Mark(InPayload.NodeId);
+			(*FoundTrackedEvent)->ReceivedTimeStamp = FAvaDisplayClusterTimeStamp::Now();
+
+			UE_LOG(LogAvaDisplayClusterSyncEvents, Verbose, TEXT("%s Tracked Event \"%s\" touched (on \"%s\"): Marking \"%s\"."),
+				*GetFrameInfo(), *InPayload.Signature, *NodeId, *InPayload.NodeId);
+		}
+		else
+		{
+			TUniquePtr<FAvaDisplayClusterTrackedClusterEvent> TrackedEvent = MakeUnique<FAvaDisplayClusterTrackedClusterEvent>();
+			TrackedEvent->NodeInfo.Mark(InPayload.NodeId);
+			TrackedEvent->ReceivedTimeStamp = FAvaDisplayClusterTimeStamp::Now();
+			TrackedEvents.Add(InPayload.Signature, MoveTemp(TrackedEvent));
+		
+			UE_LOG(LogAvaDisplayClusterSyncEvents, Verbose, TEXT("%s Tracked Event \"%s\" created (on \"%s\"): Marking for \"%s\"."),
+				*GetFrameInfo(), *InPayload.Signature, *NodeId, *InPayload.NodeId);
+		}
+	}
+}
+
+void FAvaDisplayClusterSynchronizedEventDispatcher::EmitClusterEvent(FAvaDisplayClusterSynchronizedEvent& InEvent, const FAvaDisplayClusterTimeStamp& InNow, IDisplayClusterClusterManager* InClusterManager)
+{
+	using namespace UE::AvaDisplayCluster::Private;
+	if (InClusterManager)
+	{
+		++InEvent.EmitCount;
+
+		FArrayWriter ArrayWriter;
+		FAvaDisplayClusterClusterEventPayload::Serialize(ArrayWriter, Signature, InEvent.Signature, NodeId, InEvent.EmitCount);
+			
+		FDisplayClusterClusterEventBinary ClusterEvent;
+		ClusterEvent.bIsSystemEvent = true;
+		ClusterEvent.bShouldDiscardOnRepeat = false;	// Multiple events with this Id are emitted.
+		ClusterEvent.EventId = SynchronizedEventsClusterEventId;
+		ClusterEvent.EventData = MoveTemp(ArrayWriter);
+		constexpr bool bEmitFromPrimaryOnly = false;	// All nodes emit this event.
+		InClusterManager->EmitClusterEventBinary(ClusterEvent, bEmitFromPrimaryOnly);
+
+		InEvent.LastEmitTimeStamp = InNow;
+
+		UE_LOG(LogAvaDisplayClusterSyncEvents, Verbose, TEXT("%s Emitting cluster event \"%s\" (count:%d) from \"%s\"."),
+			*GetFrameInfo(), *InEvent.Signature, InEvent.EmitCount, *NodeId);
+	}
+}
+
+void FAvaDisplayClusterSynchronizedEventDispatcher::DispatchEvent(const FAvaDisplayClusterSynchronizedEvent* InEvent, const FAvaDisplayClusterTimeStamp& InNow ) const
+{
+	check(IsInGameThread());
+	using namespace UE::AvaDisplayCluster::Private;
+	if (!InEvent)
+	{
+		return;
+	}
+		
+	if (InEvent->Function)
+	{
+		UE_LOG(LogAvaDisplayClusterSyncEvents, Verbose, TEXT("%s Invoking Event \"%s\" on \"%s\", wait time: %.2f ms (%d frames)."),
+			*GetFrameInfo(), *InEvent->Signature, *NodeId, InEvent->PushTimeStamp.GetWaitTimeInMs(InNow), InEvent->PushTimeStamp.GetWaitTimeInFrames(InNow));
+			
+		InEvent->Function();
+	}
+	else
+	{
+		UE_LOG(LogAvaDisplayClusterSyncEvents, Error, TEXT("%s Unbound Event \"%s\" on \"%s\"."), *GetFrameInfo(), *InEvent->Signature, *NodeId);
+	}
+}
+
+bool FAvaDisplayClusterSynchronizedEventDispatcher::PushEvent(FString&& InEventSignature, TUniqueFunction<void()> InFunction)
+{
+	check(IsInGameThread());
+	using namespace UE::AvaDisplayCluster::Private;
+
+	TUniquePtr<FAvaDisplayClusterSynchronizedEvent> Event
+		= MakeUnique<FAvaDisplayClusterSynchronizedEvent>(MoveTemp(InEventSignature), MoveTemp(InFunction));
+		
+	Event->PushTimeStamp = FAvaDisplayClusterTimeStamp::Now();
+
+	IDisplayClusterClusterManager* ClusterManager = IDisplayCluster::Get().GetClusterMgr();
+	if (ClusterManager && IDisplayCluster::Get().GetOperationMode() == EDisplayClusterOperationMode::Cluster)
+	{
+		if (PendingEvents.Contains(Event->Signature))
+		{
+			UE_LOG(LogAvaDisplayClusterSyncEvents, Warning, TEXT("%s Event \"%s\" is already pushed on \"%s\"."), *GetFrameInfo(), *Event->Signature, *NodeId);
+			return false; // Don't stomp events. If an event is pushed more than once, only keep one.
+		}
+
+		// Check if we have a tracked event, if so, transfer marked node ids.
+		if (const TUniquePtr<FAvaDisplayClusterTrackedClusterEvent>* FoundEvent = TrackedEvents.Find(Event->Signature))
+		{
+			Event->NodeInfo.NodeIds = MoveTemp((*FoundEvent)->NodeInfo.NodeIds);
+			TrackedEvents.Remove(Event->Signature);
+		}
+
+		const TUniquePtr<FAvaDisplayClusterSynchronizedEvent>& AddedEvent = PendingEvents.Add(Event->Signature, MoveTemp(Event));
+		EmitClusterEvent(*AddedEvent, AddedEvent->PushTimeStamp, ClusterManager);
+	}
+	else
+	{
+		if (CVarSyncEarlyDispatch.GetValueOnAnyThread())
+		{
+			DispatchEvent(Event.Get(), Event->PushTimeStamp);	
+		}
+		else
+		{
+			if (ReadyEvents.Contains(Event->Signature))
+			{
+				UE_LOG(LogAvaDisplayClusterSyncEvents, Warning, TEXT("%s Event \"%s\" is already \"ready\" on \"%s\". Previous event will be discarded."), *GetFrameInfo(), *Event->Signature, *NodeId);
+			}
+			ReadyEvents.Add(Event->Signature, MoveTemp(Event));
+		}
+	}
+	return true;
+}
+
+EAvaMediaSynchronizedEventState FAvaDisplayClusterSynchronizedEventDispatcher::GetEventState(const FString& InEventSignature) const
+{
+	if (PendingEvents.Contains(InEventSignature))
+	{
+		return EAvaMediaSynchronizedEventState::Pending;
+	}
+	if (ReadyEvents.Contains(InEventSignature))
+	{
+		return EAvaMediaSynchronizedEventState::Ready;
+	}
+	if (TrackedEvents.Contains(InEventSignature))
+	{
+		return EAvaMediaSynchronizedEventState::Tracked;
+	}
+	return EAvaMediaSynchronizedEventState::NotFound;
+}
+
+void FAvaDisplayClusterSynchronizedEventDispatcher::DispatchEvents()
+{
+	check(IsInGameThread());
+	using namespace UE::AvaDisplayCluster::Private;
+	const FAvaDisplayClusterTimeStamp Now = FAvaDisplayClusterTimeStamp::Now();
+
+	IDisplayClusterClusterManager* ClusterManager = IDisplayCluster::Get().GetClusterMgr();
+	
+	// Refresh the list of nodes (only once a frame) in case it has changed.
+	// There doesn't seem to be an event for this.
+	if (ClusterManager)
+	{
+		ClusterManager->GetNodeIds(AllNodeIds);
+	}
+
+	const double RepeatTimeoutInMs = CVarSyncRepeatTimeout.GetValueOnAnyThread();
+	const double DispatchTimeoutInMs = CVarSyncDispatchTimeout.GetValueOnAnyThread();
+
+	for (TMap<FString, TUniquePtr<FAvaDisplayClusterSynchronizedEvent>>::TIterator EventIt(PendingEvents); EventIt; ++EventIt)
+	{
+		if (EventIt->Value->NodeInfo.IsAllMarked(AllNodeIds))
+		{
+			DispatchEvent(EventIt->Value.Get(), Now);
+			EventIt.RemoveCurrent();
+			continue;
+		}
+
+		// Locally dispatch pending events that timed out.
+		if (EventIt->Value->PushTimeStamp.GetWaitTimeInMs(Now) > DispatchTimeoutInMs)
+		{
+			UE_LOG(LogAvaDisplayClusterSyncEvents, Warning, TEXT("%s Pending Event \"%s\" has timed out on \"%s\". Wait time: %.2f ms"),
+				*GetFrameInfo(), *EventIt->Value->Signature, *NodeId, EventIt->Value->PushTimeStamp.GetWaitTimeInMs(Now));
+
+			DispatchEvent(EventIt->Value.Get(), Now);
+			EventIt.RemoveCurrent();
+			
+			// Todo: Investigate idea. if an event times out, could keep track of the node that didn't fire and
+			// "mute it" (don't wait for it anymore) until it returns.
+
+			continue;
+		}
+
+		// Failsafe: periodically re-emit pending events in case cluster event didn't get through.
+		// (Note: make sure to not re-emit a timed out event. time out test first above.)
+		if (EventIt->Value->LastEmitTimeStamp.GetWaitTimeInMs(Now) > RepeatTimeoutInMs)
+		{
+			UE_LOG(LogAvaDisplayClusterSyncEvents, Verbose, TEXT("%s Re-emitting event \"%s\" on \"%s\"."),
+				*GetFrameInfo(), *EventIt->Value->Signature, *NodeId);
+
+			EmitClusterEvent(*EventIt->Value, Now, ClusterManager);
+		}
+	}
+
+	// Tracking timeout.
+	const double TrackingTimeoutInMs = CVarSyncTrackingTimeout.GetValueOnAnyThread();
+	
+	for (TMap<FString, TUniquePtr<FAvaDisplayClusterTrackedClusterEvent>>::TIterator TrackedEventIt(TrackedEvents); TrackedEventIt; ++TrackedEventIt)
+	{
+		if (TrackedEventIt->Value->ReceivedTimeStamp.GetWaitTimeInMs(Now) > TrackingTimeoutInMs)
+		{
+			UE_LOG(LogAvaDisplayClusterSyncEvents, Warning, TEXT("%s Tracked Event \"%s\" has timed out on \"%s\". Wait time: %.2f ms"),
+				*GetFrameInfo(), *TrackedEventIt->Key, *NodeId, TrackedEventIt->Value->ReceivedTimeStamp.GetWaitTimeInMs(Now));
+
+			TrackedEventIt.RemoveCurrent();
+		}
+	}
+
+	for (const TPair<FString, TUniquePtr<FAvaDisplayClusterSynchronizedEvent>>& Event : ReadyEvents)
+	{
+		DispatchEvent(Event.Value.Get(), Now);
+	}
+	ReadyEvents.Reset();
+}
+
diff --git a/Engine/Plugins/Experimental/AvalancheDisplayCluster/Source/AvalancheDisplayCluster/Private/SynchronizedEvents/AvaDisplayClusterSynchronizedEventsDispatcher.h b/Engine/Plugins/Experimental/AvalancheDisplayCluster/Source/AvalancheDisplayCluster/Private/SynchronizedEvents/AvaDisplayClusterSynchronizedEventsDispatcher.h
new file mode 100644
--- /dev/null
+++ b/Engine/Plugins/Experimental/AvalancheDisplayCluster/Source/AvalancheDisplayCluster/Private/SynchronizedEvents/AvaDisplayClusterSynchronizedEventsDispatcher.h
@@ -0,0 +1,123 @@
+// Copyright Epic Games, Inc. All Rights Reserved.
+
+#pragma once
+
+#include "Containers/Map.h"
+#include "Containers/Set.h"
+#include "ModularFeature/AvaMediaSynchronizedEvent.h"
+#include "ModularFeature/IAvaMediaSynchronizedEventDispatcher.h"
+#include "Utils/AvaDisplayClusterTimeStamp.h"
+
+class IDisplayClusterClusterManager;
+
+struct FAvaDisplayClusterClusterEventPayload
+{
+	/** Dispatcher Signature */
+	FString Dispatcher;
+	/** Event Signature */
+	FString Signature;
+	FString NodeId;
+	int32 EmitCount = 0;
+		
+	static void Serialize(FArchive& InArchive, FString& InOutDispatcher, FString& InOutSignature, FString& InOutNodeId, int32& InOutEmitCount)
+	{
+		InArchive << InOutDispatcher;
+		InArchive << InOutSignature;
+		InArchive << InOutNodeId;
+		InArchive << InOutEmitCount;
+	}
+
+	void Serialize(FArchive& InArchive)
+	{
+		Serialize(InArchive, Dispatcher, Signature, NodeId, EmitCount);
+	}
+};
+
+struct FAvaDisplayClusterNodeInfo
+{
+	void Mark(const FString& InNodeId)
+	{
+		NodeIds.Add(InNodeId);
+	}
+
+	bool IsAllMarked(const TArray<FString>& InNodeIds) const
+	{
+		return NodeIds.Num() == InNodeIds.Num();
+	}
+
+	TSet<FString> NodeIds;
+};
+
+struct FAvaDisplayClusterTrackedClusterEvent
+{
+	FAvaDisplayClusterNodeInfo NodeInfo;
+	FAvaDisplayClusterTimeStamp ReceivedTimeStamp;
+};
+
+struct FAvaDisplayClusterSynchronizedEvent : public FAvaMediaSynchronizedEvent
+{
+	FAvaDisplayClusterSynchronizedEvent(const FString& InSignature, TUniqueFunction<void()> InFunction)
+		: FAvaMediaSynchronizedEvent(InSignature, MoveTemp(InFunction))
+	{
+	}
+
+	FAvaDisplayClusterSynchronizedEvent(FString&& InSignature, TUniqueFunction<void()> InFunction)
+		: FAvaMediaSynchronizedEvent(MoveTemp(InSignature), MoveTemp(InFunction))
+	{
+	}
+
+	/** Tracks which node is marked. */
+	FAvaDisplayClusterNodeInfo NodeInfo;
+
+	/** Tracks time at which the event was pushed. */
+	FAvaDisplayClusterTimeStamp PushTimeStamp;
+
+	/** Tracks last time the event was emitted. */
+	FAvaDisplayClusterTimeStamp LastEmitTimeStamp;
+
+	/** Tracks number of times the event has been emitted. */
+	int32 EmitCount = 0;
+};
+
+class FAvaDisplayClusterSynchronizedEventDispatcher : public IAvaMediaSynchronizedEventDispatcher
+{
+public:
+	FAvaDisplayClusterSynchronizedEventDispatcher(const FString& InSignature);
+
+	virtual ~FAvaDisplayClusterSynchronizedEventDispatcher() override = default;
+
+	// Formatted frame info: time stamp and dispatcher's signature.
+	FString GetFrameInfo() const;
+
+	void OnClusterEventReceived(const FAvaDisplayClusterClusterEventPayload& InPayload);
+
+	void EmitClusterEvent(FAvaDisplayClusterSynchronizedEvent& InEvent, const FAvaDisplayClusterTimeStamp& InNow, IDisplayClusterClusterManager* InClusterManager);
+
+	void DispatchEvent(const FAvaDisplayClusterSynchronizedEvent* InEvent, const FAvaDisplayClusterTimeStamp& InNow ) const;
+	
+	//~ Begin IAvaMediaSynchronizedEventDispatcher
+	virtual bool PushEvent(FString&& InEventSignature, TUniqueFunction<void()> InFunction) override;
+	virtual EAvaMediaSynchronizedEventState GetEventState(const FString& InEventSignature) const override;
+	virtual void DispatchEvents() override;
+	//~ End IAvaMediaSynchronizedEventDispatcher 
+
+	/** Dispatcher's signature. (used for debugging purpose atm) */
+	FString Signature;
+	
+	/** Id of the current Cluster Node. */
+	FString NodeId;
+
+	/** All nodes in the cluster. */
+	TArray<FString> AllNodeIds;
+
+	/** Pending events waiting on other node's signal. */
+	TMap<FString, TUniquePtr<FAvaDisplayClusterSynchronizedEvent>> PendingEvents;
+
+	/** Tracked Events that are received as cluster event prior to getting submitted locally. */
+	TMap<FString, TUniquePtr<FAvaDisplayClusterTrackedClusterEvent>> TrackedEvents;
+	
+	/** Events Ready to be locally dispatched. */
+	TMap<FString, TUniquePtr<FAvaDisplayClusterSynchronizedEvent>> ReadyEvents;
+	
+	static constexpr int32 SynchronizedEventsClusterEventId = 0xABCDEF01;
+};
diff --git a/Engine/Plugins/Experimental/AvalancheDisplayCluster/Source/AvalancheDisplayCluster/Private/SynchronizedEvents/AvaDisplayClusterSynchronizedEventsFeature.cpp b/Engine/Plugins/Experimental/AvalancheDisplayCluster/Source/AvalancheDisplayCluster/Private/SynchronizedEvents/AvaDisplayClusterSynchronizedEventsFeature.cpp
new file mode 100644
--- /dev/null
+++ b/Engine/Plugins/Experimental/AvalancheDisplayCluster/Source/AvalancheDisplayCluster/Private/SynchronizedEvents/AvaDisplayClusterSynchronizedEventsFeature.cpp
@@ -0,0 +1,172 @@
+// Copyright Epic Games, Inc. All Rights Reserved.
+
+#include "AvaDisplayClusterSynchronizedEventsFeature.h"
+
+#include "AvaDisplayClusterSyncEventsLog.h"
+#include "AvaDisplayClusterSynchronizedEventsDispatcher.h"
+#include "Cluster/DisplayClusterClusterEvent.h"
+//#include "Cluster/IDisplayClusterClusterManager.h"
+#include "IDisplayCluster.h"
+#include "IDisplayClusterCallbacks.h"
+#include "Misc/CoreDelegates.h"
+#include "Playback/AvaPlaybackUtils.h"
+#include "Serialization/MemoryReader.h"
+
+DEFINE_LOG_CATEGORY(LogAvaDisplayClusterSyncEvents);
+
+#define LOCTEXT_NAMESPACE "AvaDisplayClusterSynchronizedEvent"
+
+FAvaDisplayClusterSynchronizedEventsFeature::FAvaDisplayClusterSynchronizedEventsFeature()
+{
+	if (IDisplayCluster* DisplayCluster = FModuleManager::LoadModulePtr<IDisplayCluster>(IDisplayCluster::ModuleName))
+	{
+		DisplayCluster->GetCallbacks().OnDisplayClusterStartSession().AddRaw(this, &FAvaDisplayClusterSynchronizedEventsFeature::OnDisplayClusterStartSession);
+	}
+	FCoreDelegates::OnEndFrame.AddRaw(this, &FAvaDisplayClusterSynchronizedEventsFeature::OnEndFrame);
+}
+
+FAvaDisplayClusterSynchronizedEventsFeature::~FAvaDisplayClusterSynchronizedEventsFeature()
+{
+	FCoreDelegates::OnEndFrame.RemoveAll(this);
+
+	if (IDisplayCluster::IsAvailable())
+	{
+		IDisplayCluster::Get().GetCallbacks().OnDisplayClusterStartSession().RemoveAll(this);
+		
+		if (IDisplayClusterClusterManager* ClusterManager = IDisplayCluster::Get().GetClusterMgr())
+		{
+			if (BinaryListener.IsBound())
+			{
+				ClusterManager->RemoveClusterEventBinaryListener(BinaryListener);
+			}
+		}
+	}
+}
+
+FName FAvaDisplayClusterSynchronizedEventsFeature::GetName() const
+{
+	static const FName ImplementationName(TEXT("DisplayClusterSync"));
+	return ImplementationName;
+}
+
+FText FAvaDisplayClusterSynchronizedEventsFeature::GetDisplayName() const
+{
+	static const FText DisplayName = LOCTEXT("AvaDisplayClusterSyncDisplayName", "nDisplay Sync");
+	return DisplayName;
+}
+
+FText FAvaDisplayClusterSynchronizedEventsFeature::GetDisplayDescription() const
+{
+	static const FText DisplayDescription = LOCTEXT("AvaDisplayClusterSyncDescription",
+		 "Synchronize events over nDisplay using Cluster Events.");
+	return DisplayDescription;
+}
+
+int32 FAvaDisplayClusterSynchronizedEventsFeature::GetPriority() const
+{
+	// It is likely cluster synchronization will take priority over other implementations.
+	return GetDefaultPriority() + 10;
+}
+
+TSharedPtr<IAvaMediaSynchronizedEventDispatcher> FAvaDisplayClusterSynchronizedEventsFeature::CreateDispatcher(const FString& InSignature)
+{
+	TSharedPtr<FAvaDisplayClusterSynchronizedEventDispatcher> Dispatcher = MakeShared<FAvaDisplayClusterSynchronizedEventDispatcher>(InSignature);
+
+	// Issue an error if the dispatcher already exists.
+	const TWeakPtr<FAvaDisplayClusterSynchronizedEventDispatcher>* FoundDispatcherWeak = DispatchersWeak.Find(InSignature);
+	if (FoundDispatcherWeak && FoundDispatcherWeak->IsValid())
+	{
+		UE_LOG(LogAvaDisplayClusterSyncEvents, Error, TEXT("Dispatcher \"%s\" already exist. Previous dispatcher will no longuer received events."), *InSignature);
+	}
+	
+	// Keep track of our dispatchers for event routing.
+	DispatchersWeak.Add(InSignature, Dispatcher.ToWeakPtr());
+
+	// Carry over the tracked events prior to the dispatcher creation.
+	//
+	// Expectation is that the creation of the dispatchers may be out of sync on different nodes because it is
+	// triggered by means that may have a few frames of delay. But we don't expect corresponding dispatchers
+	// to be create with a large desync (of seconds), for now (tracker has a timeout of a few seconds).
+	//
+	const TSharedPtr<FAvaDisplayClusterSynchronizedEventDispatcher>* FoundTrackingDispatcher = TrackingDispatchers.Find(InSignature);
+	if (FoundTrackingDispatcher && FoundTrackingDispatcher->IsValid())
+	{
+		Dispatcher->TrackedEvents = MoveTemp((*FoundTrackingDispatcher)->TrackedEvents);
+	}
+	TrackingDispatchers.Remove(InSignature);
+	
+	return Dispatcher;
+}
+
+void FAvaDisplayClusterSynchronizedEventsFeature::OnDisplayClusterStartSession()
+{
+	IDisplayClusterClusterManager* ClusterManager = IDisplayCluster::Get().GetClusterMgr();
+	if (!ClusterManager)
+	{
+		UE_LOG(LogAvaDisplayClusterSyncEvents, Error, TEXT("Display Cluster Manager is not available. Synchronized events will be disabled."));
+		return;
+	}
+	
+	BinaryListener = FOnClusterEventBinaryListener::CreateRaw(this, &FAvaDisplayClusterSynchronizedEventsFeature::OnBinaryClusterEventReceived);
+	ClusterManager->AddClusterEventBinaryListener(BinaryListener);
+}
+
+void FAvaDisplayClusterSynchronizedEventsFeature::OnEndFrame()
+{
+	// Timeout tracked events - We may need to revisit this for edge cases (ex: long cluster init).
+	for (const TPair<FString, TSharedPtr<FAvaDisplayClusterSynchronizedEventDispatcher>>& TrackingDispatcher : TrackingDispatchers)
+	{
+		TrackingDispatcher.Value->DispatchEvents();
+	}
+}
+
+void FAvaDisplayClusterSynchronizedEventsFeature::OnBinaryClusterEventReceived(const FDisplayClusterClusterEventBinary& InClusterEvent)
+{
+	if (InClusterEvent.EventId != FAvaDisplayClusterSynchronizedEventDispatcher::SynchronizedEventsClusterEventId)
+	{
+		return;
+	}
+		
+	FMemoryReader Reader(InClusterEvent.EventData);
+	FAvaDisplayClusterClusterEventPayload Payload;
+	Payload.Serialize(Reader);
+	
+	if (const TWeakPtr<FAvaDisplayClusterSynchronizedEventDispatcher>* FoundDispatcherWeak = DispatchersWeak.Find(Payload.Dispatcher))
+	{
+		if (const TSharedPtr<FAvaDisplayClusterSynchronizedEventDispatcher> FoundDispatcher = FoundDispatcherWeak->Pin())
+		{
+			FoundDispatcher->OnClusterEventReceived(Payload);
+			return;	// Handled.
+		}
+		else
+		{
+			DispatchersWeak.Remove(Payload.Dispatcher);	// Remove stale map entries.
+		}
+	}
+
+	// If no existing dispatcher has handled this event, we need to keep track of it in case dispatcher's creation is
+	// out of sync between the different nodes.
+	if (const TSharedPtr<FAvaDisplayClusterSynchronizedEventDispatcher>* FoundTrackingDispatcher = TrackingDispatchers.Find(Payload.Dispatcher))
+	{
+		if (*FoundTrackingDispatcher)
+		{
+			(*FoundTrackingDispatcher)->OnClusterEventReceived(Payload);
+			return; // Handled.
+		}
+		else
+		{
+			UE_LOG(LogAvaDisplayClusterSyncEvents, Warning, TEXT("Stale Tracking Dispatcher Found."));
+			TrackingDispatchers.Remove(Payload.Dispatcher);	// Remove stale map entries.
+		}
+	}
+	
+	const TSharedPtr<FAvaDisplayClusterSynchronizedEventDispatcher> TrackingDispatcher = MakeShared<FAvaDisplayClusterSynchronizedEventDispatcher>(Payload.Signature);
+	TrackingDispatchers.Add(Payload.Signature, TrackingDispatcher);
+
+	UE_LOG(LogAvaDisplayClusterSyncEvents, Verbose, TEXT("%s Tracking Dispatcher \"%s\" created on \"%s\"."),
+		*UE::AvaPlayback::Utils::GetBriefFrameInfo(),  *TrackingDispatcher->Signature, *TrackingDispatcher->NodeId);
+	
+	TrackingDispatcher->OnClusterEventReceived(Payload);
+}
+
+#undef LOCTEXT_NAMESPACE
\ No newline at end of file
diff --git a/Engine/Plugins/Experimental/AvalancheDisplayCluster/Source/AvalancheDisplayCluster/Private/SynchronizedEvents/AvaDisplayClusterSynchronizedEventsFeature.h b/Engine/Plugins/Experimental/AvalancheDisplayCluster/Source/AvalancheDisplayCluster/Private/SynchronizedEvents/AvaDisplayClusterSynchronizedEventsFeature.h
new file mode 100644
--- /dev/null
+++ b/Engine/Plugins/Experimental/AvalancheDisplayCluster/Source/AvalancheDisplayCluster/Private/SynchronizedEvents/AvaDisplayClusterSynchronizedEventsFeature.h
@@ -0,0 +1,39 @@
+// Copyright Epic Games, Inc. All Rights Reserved.
+
+#pragma once
+
+#include "Cluster/IDisplayClusterClusterManager.h"
+#include "ModularFeature/IAvaMediaSynchronizedEventsFeature.h"
+
+class FAvaDisplayClusterSynchronizedEventDispatcher;
+struct FDisplayClusterClusterEventBinary;
+
+class FAvaDisplayClusterSynchronizedEventsFeature : public IAvaMediaSynchronizedEventsFeature
+{
+public:
+	FAvaDisplayClusterSynchronizedEventsFeature();
+	virtual ~FAvaDisplayClusterSynchronizedEventsFeature() override;
+
+	//~ Begin IAvaMediaSynchronizedEventsFeature
+	virtual FName GetName() const override;
+	virtual FText GetDisplayName() const override;
+	virtual FText GetDisplayDescription() const override;
+	virtual int32 GetPriority() const override;
+	virtual TSharedPtr<IAvaMediaSynchronizedEventDispatcher> CreateDispatcher(const FString& InSignature) override;
+	//~ End IAvaMediaSynchronizedEventsFeature
+
+private:
+	// Event Handlers
+	void OnDisplayClusterStartSession();
+	void OnEndFrame();
+	void OnBinaryClusterEventReceived(const FDisplayClusterClusterEventBinary& InClusterEvent);
+	
+private:
+	/** Keep track of the created dispatchers for routing the cluster events. */
+	TMap<FString, TWeakPtr<FAvaDisplayClusterSynchronizedEventDispatcher>> DispatchersWeak;
+
+	/** Dispatchers only used for tracking cluster events when there is no active dispatcher created. */
+	TMap<FString, TSharedPtr<FAvaDisplayClusterSynchronizedEventDispatcher>> TrackingDispatchers;
+
+	FOnClusterEventBinaryListener BinaryListener;
+};
\ No newline at end of file
diff --git a/Engine/Plugins/Experimental/AvalancheDisplayCluster/Source/AvalancheDisplayCluster/Private/Utils/AvaDisplayClusterTimeStamp.cpp b/Engine/Plugins/Experimental/AvalancheDisplayCluster/Source/AvalancheDisplayCluster/Private/Utils/AvaDisplayClusterTimeStamp.cpp
new file mode 100644
--- /dev/null
+++ b/Engine/Plugins/Experimental/AvalancheDisplayCluster/Source/AvalancheDisplayCluster/Private/Utils/AvaDisplayClusterTimeStamp.cpp
@@ -0,0 +1,10 @@
+// Copyright Epic Games, Inc. All Rights Reserved.
+
+#include "AvaDisplayClusterTimeStamp.h"
+
+#include "CoreGlobals.h"
+
+FAvaDisplayClusterTimeStamp FAvaDisplayClusterTimeStamp::Now()
+{
+	return { FDateTime::Now(), GFrameNumber};
+}
diff --git a/Engine/Plugins/Experimental/AvalancheDisplayCluster/Source/AvalancheDisplayCluster/Private/Utils/AvaDisplayClusterTimeStamp.h b/Engine/Plugins/Experimental/AvalancheDisplayCluster/Source/AvalancheDisplayCluster/Private/Utils/AvaDisplayClusterTimeStamp.h
new file mode 100644
--- /dev/null
+++ b/Engine/Plugins/Experimental/AvalancheDisplayCluster/Source/AvalancheDisplayCluster/Private/Utils/AvaDisplayClusterTimeStamp.h
@@ -0,0 +1,30 @@
+// Copyright Epic Games, Inc. All Rights Reserved.
+
+#pragma once
+
+#include "Misc/DateTime.h"
+
+/**
+ * @brief TimeStamp utility for scheduling events, timeout, repeating, etc.
+ */
+struct FAvaDisplayClusterTimeStamp
+{
+	static FAvaDisplayClusterTimeStamp Now();
+
+	double GetWaitTimeInMs(const FAvaDisplayClusterTimeStamp& InNow) const
+	{
+		const FTimespan wait = InNow.Time - Time;
+		return wait.GetTotalMilliseconds();
+	}
+
+	uint32 GetWaitTimeInFrames(const FAvaDisplayClusterTimeStamp& InNow) const
+	{
+		return InNow.FrameNumber - FrameNumber;	
+	}
+
+	/** Local System time. Not synchronized. */
+	FDateTime Time;
+	
+	/** (Tentative) Keep track of synchronized frame number (for logging and detecting desyncs). */
+	uint32 FrameNumber = 0;
+};
