[Motion Design] Fix CIS 498201 - nDisplay not supported on Mac

diff --git a/Engine/Plugins/Experimental/Avalanche/Source/AvalancheMedia/Private/Playback/AvaPlaybackServer.h b/Engine/Plugins/Experimental/Avalanche/Source/AvalancheMedia/Private/Playback/AvaPlaybackServer.h
--- a/Engine/Plugins/Experimental/Avalanche/Source/AvalancheMedia/Private/Playback/AvaPlaybackServer.h
+++ b/Engine/Plugins/Experimental/Avalanche/Source/AvalancheMedia/Private/Playback/AvaPlaybackServer.h
@@ -230,6 +230,10 @@ private:
 		int32 Priority;
 		FMessageAddress ReplyTo;
 		FAvaPlaybackCommand Command;
+
+		FPendingPlaybackCommand(const FDateTime& InReceivedUtc, uint32 InReceivedFrameNumber, int32 InPriority, const FMessageAddress& InReplyTo, const FAvaPlaybackCommand& InCommand)
+			: ReceivedUtc(InReceivedUtc), ReceivedFrameNumber(InReceivedFrameNumber), Priority(InPriority), ReplyTo(InReplyTo), Command(InCommand)
+		{}
 	};
 
 	/** Accumulate all the playback commands and execute them all in one batch on the next tick. */
diff --git a/Engine/Plugins/Experimental/AvalancheDisplayCluster/AvalancheDisplayCluster.uplugin b/Engine/Plugins/Experimental/AvalancheDisplayCluster/AvalancheDisplayCluster.uplugin
--- a/Engine/Plugins/Experimental/AvalancheDisplayCluster/AvalancheDisplayCluster.uplugin
+++ b/Engine/Plugins/Experimental/AvalancheDisplayCluster/AvalancheDisplayCluster.uplugin
@@ -17,14 +17,17 @@
 	"EnabledByDefault": false,
 	"SupportedTargetPlatforms": [
 		"Win64",
-		"Linux",
-		"Mac"
+		"Linux"
 	],
 	"Modules": [
 		{
 			"Name": "AvalancheDisplayCluster",
 			"Type": "Runtime",
-			"LoadingPhase": "Default"
+			"LoadingPhase": "Default",
+			"PlatformAllowList": [
+				"Win64",
+				"Linux"
+			]
 		}
 	],
 	"Plugins": [
