Fixed include paths following StructUtils being moved to CoreUObject

diff --git a/Engine/Plugins/Animation/DeformerGraph/Source/OptimusCore/Public/OptimusValueContainerStruct.h b/Engine/Plugins/Animation/DeformerGraph/Source/OptimusCore/Public/OptimusValueContainerStruct.h
--- a/Engine/Plugins/Animation/DeformerGraph/Source/OptimusCore/Public/OptimusValueContainerStruct.h
+++ b/Engine/Plugins/Animation/DeformerGraph/Source/OptimusCore/Public/OptimusValueContainerStruct.h
@@ -2,9 +2,12 @@
 
 #pragma once
 
+#if UE_ENABLE_INCLUDE_ORDER_DEPRECATED_IN_5_5
 #include "CoreMinimal.h"
+#endif
+
 #include "OptimusDataType.h"
-#include "PropertyBag.h"
+#include "StructUtils/PropertyBag.h"
 #include "OptimusValueContainerStruct.generated.h"
 
 USTRUCT()
diff --git a/Engine/Plugins/Animation/LiveLinkHub/Source/LiveLinkHub/Private/Recording/Implementations/LiveLinkUAssetRecorder.cpp b/Engine/Plugins/Animation/LiveLinkHub/Source/LiveLinkHub/Private/Recording/Implementations/LiveLinkUAssetRecorder.cpp
--- a/Engine/Plugins/Animation/LiveLinkHub/Source/LiveLinkHub/Private/Recording/Implementations/LiveLinkUAssetRecorder.cpp
+++ b/Engine/Plugins/Animation/LiveLinkHub/Source/LiveLinkHub/Private/Recording/Implementations/LiveLinkUAssetRecorder.cpp
@@ -11,7 +11,7 @@
 #include "Features/IModularFeatures.h"
 #include "Framework/MultiBox/MultiBoxBuilder.h"
 #include "FileHelpers.h"
-#include "InstancedStruct.h"
+#include "StructUtils/InstancedStruct.h"
 #include "IContentBrowserSingleton.h"
 #include "LiveLinkHub.h"
 #include "LiveLinkHubClient.h"
diff --git a/Engine/Plugins/Animation/LiveLinkHub/Source/LiveLinkHub/Private/Recording/Implementations/LiveLinkUAssetRecording.h b/Engine/Plugins/Animation/LiveLinkHub/Source/LiveLinkHub/Private/Recording/Implementations/LiveLinkUAssetRecording.h
--- a/Engine/Plugins/Animation/LiveLinkHub/Source/LiveLinkHub/Private/Recording/Implementations/LiveLinkUAssetRecording.h
+++ b/Engine/Plugins/Animation/LiveLinkHub/Source/LiveLinkHub/Private/Recording/Implementations/LiveLinkUAssetRecording.h
@@ -3,7 +3,7 @@
 #pragma once
 
 #include "Async/AsyncWork.h"
-#include "InstancedStruct.h"
+#include "StructUtils/InstancedStruct.h"
 #include "LiveLinkPreset.h"
 #include "LiveLinkRole.h"
 #include "LiveLinkTypes.h"
diff --git a/Engine/Plugins/Animation/LiveLinkHub/Source/LiveLinkHub/Private/Recording/LiveLinkRecordingPlayer.h b/Engine/Plugins/Animation/LiveLinkHub/Source/LiveLinkHub/Private/Recording/LiveLinkRecordingPlayer.h
--- a/Engine/Plugins/Animation/LiveLinkHub/Source/LiveLinkHub/Private/Recording/LiveLinkRecordingPlayer.h
+++ b/Engine/Plugins/Animation/LiveLinkHub/Source/LiveLinkHub/Private/Recording/LiveLinkRecordingPlayer.h
@@ -4,7 +4,7 @@
 
 #include "LiveLinkRole.h"
 #include "LiveLinkTypes.h"
-#include "StructView.h"
+#include "StructUtils/StructView.h"
 #include "Templates/SubclassOf.h"
 
 class ULiveLinkRecording;
diff --git a/Engine/Plugins/Experimental/ActorModifierCore/Source/ActorModifierCore/Public/Modifiers/Profiler/ActorModifierCoreProfiler.h b/Engine/Plugins/Experimental/ActorModifierCore/Source/ActorModifierCore/Public/Modifiers/Profiler/ActorModifierCoreProfiler.h
--- a/Engine/Plugins/Experimental/ActorModifierCore/Source/ActorModifierCore/Public/Modifiers/Profiler/ActorModifierCoreProfiler.h
+++ b/Engine/Plugins/Experimental/ActorModifierCore/Source/ActorModifierCore/Public/Modifiers/Profiler/ActorModifierCoreProfiler.h
@@ -2,7 +2,7 @@
 
 #pragma once
 
-#include "PropertyBag.h"
+#include "StructUtils/PropertyBag.h"
 
 class FStructOnScope;
 class UActorModifierCoreBase;
diff --git a/Engine/Plugins/Experimental/AnimNext/Source/AnimNext/Internal/Graph/AnimNextGraphState.h b/Engine/Plugins/Experimental/AnimNext/Source/AnimNext/Internal/Graph/AnimNextGraphState.h
--- a/Engine/Plugins/Experimental/AnimNext/Source/AnimNext/Internal/Graph/AnimNextGraphState.h
+++ b/Engine/Plugins/Experimental/AnimNext/Source/AnimNext/Internal/Graph/AnimNextGraphState.h
@@ -3,7 +3,7 @@
 #pragma once
 
 #include "CoreMinimal.h"
-#include "PropertyBag.h"
+#include "StructUtils/PropertyBag.h"
 #include "AnimNextGraphState.generated.h"
 
 // Default state for a graph entry point
diff --git a/Engine/Plugins/Experimental/Avalanche/Source/AvalancheMask/Internal/Materials/AvaMaskMaterialInstanceSubsystem.h b/Engine/Plugins/Experimental/Avalanche/Source/AvalancheMask/Internal/Materials/AvaMaskMaterialInstanceSubsystem.h
--- a/Engine/Plugins/Experimental/Avalanche/Source/AvalancheMask/Internal/Materials/AvaMaskMaterialInstanceSubsystem.h
+++ b/Engine/Plugins/Experimental/Avalanche/Source/AvalancheMask/Internal/Materials/AvaMaskMaterialInstanceSubsystem.h
@@ -2,7 +2,7 @@
 
 #pragma once
 
-#include "SharedStruct.h"
+#include "StructUtils/SharedStruct.h"
 #include "Subsystems/EngineSubsystem.h"
 #include "Subsystems/WorldSubsystem.h"
 #include "UObject/SoftObjectPtr.h"
diff --git a/Engine/Plugins/Experimental/Avalanche/Source/AvalancheMask/Private/Handling/AvaHandleUtilities.h b/Engine/Plugins/Experimental/Avalanche/Source/AvalancheMask/Private/Handling/AvaHandleUtilities.h
--- a/Engine/Plugins/Experimental/Avalanche/Source/AvalancheMask/Private/Handling/AvaHandleUtilities.h
+++ b/Engine/Plugins/Experimental/Avalanche/Source/AvalancheMask/Private/Handling/AvaHandleUtilities.h
@@ -5,7 +5,7 @@
 #include "Components/PrimitiveComponent.h"
 #include "IAvaMaskMaterialHandle.h"
 #include "IAvaObjectHandle.h"
-#include "StructView.h"
+#include "StructUtils/StructView.h"
 #include "Templates/UnrealTypeTraits.h"
 
 class UAvaObjectHandleSubsystem;
diff --git a/Engine/Plugins/Experimental/Avalanche/Source/AvalancheMask/Private/Handling/IAvaMaskMaterialHandle.h b/Engine/Plugins/Experimental/Avalanche/Source/AvalancheMask/Private/Handling/IAvaMaskMaterialHandle.h
--- a/Engine/Plugins/Experimental/Avalanche/Source/AvalancheMask/Private/Handling/IAvaMaskMaterialHandle.h
+++ b/Engine/Plugins/Experimental/Avalanche/Source/AvalancheMask/Private/Handling/IAvaMaskMaterialHandle.h
@@ -7,11 +7,11 @@
 #include "AvaMaskUtilities.h"
 #include "Engine/Texture.h"
 #include "IAvaMaterialHandle.h"
-#include "InstancedStruct.h"
+#include "StructUtils/InstancedStruct.h"
 #include "Materials/AvaMaterialUtils.h"
 #include "Materials/MaterialInstance.h"
 #include "Materials/MaterialInstanceDynamic.h"
-#include "StructView.h"
+#include "StructUtils/StructView.h"
 
 class UAvaMaskMaterialInstanceSubsystem;
 
diff --git a/Engine/Plugins/Experimental/Chooser/Source/Chooser/Public/IObjectChooser.h b/Engine/Plugins/Experimental/Chooser/Source/Chooser/Public/IObjectChooser.h
--- a/Engine/Plugins/Experimental/Chooser/Source/Chooser/Public/IObjectChooser.h
+++ b/Engine/Plugins/Experimental/Chooser/Source/Chooser/Public/IObjectChooser.h
@@ -5,8 +5,8 @@
 #include "CoreMinimal.h"
 #include "UObject/Object.h"
 #include "UObject/Interface.h"
-#include "InstancedStruct.h"
-#include "StructView.h"
+#include "StructUtils/InstancedStruct.h"
+#include "StructUtils/StructView.h"
 #include "IObjectChooser.generated.h"
 
 #if UE_TRACE_ENABLED && !IS_PROGRAM && !UE_BUILD_SHIPPING && !UE_BUILD_TEST
diff --git a/Engine/Plugins/Experimental/Chooser/Source/ProxyTable/Internal/LookupProxy.h b/Engine/Plugins/Experimental/Chooser/Source/ProxyTable/Internal/LookupProxy.h
--- a/Engine/Plugins/Experimental/Chooser/Source/ProxyTable/Internal/LookupProxy.h
+++ b/Engine/Plugins/Experimental/Chooser/Source/ProxyTable/Internal/LookupProxy.h
@@ -5,7 +5,7 @@
 #include "IObjectChooser.h"
 #include "ChooserPropertyAccess.h"
 #include "IChooserParameterProxyTable.h"
-#include "InstancedStructContainer.h"
+#include "StructUtils/InstancedStructContainer.h"
 #include "ProxyTable.h"
 #include "LookupProxy.generated.h"
 
diff --git a/Engine/Plugins/MovieScene/MovieRenderPipeline/Source/MovieRenderPipelineCore/Public/Graph/MovieGraphNode.h b/Engine/Plugins/MovieScene/MovieRenderPipeline/Source/MovieRenderPipelineCore/Public/Graph/MovieGraphNode.h
--- a/Engine/Plugins/MovieScene/MovieRenderPipeline/Source/MovieRenderPipelineCore/Public/Graph/MovieGraphNode.h
+++ b/Engine/Plugins/MovieScene/MovieRenderPipeline/Source/MovieRenderPipelineCore/Public/Graph/MovieGraphNode.h
@@ -4,8 +4,8 @@
 
 #include "CoreMinimal.h"
 #include "MovieGraphPin.h"
-#include "InstancedStruct.h"
-#include "PropertyBag.h"
+#include "StructUtils/InstancedStruct.h"
+#include "StructUtils/PropertyBag.h"
 #include "Graph/MovieGraphValueContainer.h"
 #include "Graph/MovieGraphFilenameResolveParams.h"
 #include "UObject/Interface.h"
diff --git a/Engine/Plugins/PCG/Source/PCG/Private/Elements/PCGUserParameterGet.cpp b/Engine/Plugins/PCG/Source/PCG/Private/Elements/PCGUserParameterGet.cpp
--- a/Engine/Plugins/PCG/Source/PCG/Private/Elements/PCGUserParameterGet.cpp
+++ b/Engine/Plugins/PCG/Source/PCG/Private/Elements/PCGUserParameterGet.cpp
@@ -13,8 +13,8 @@
 #include "Helpers/PCGPropertyHelpers.h"
 #include "Metadata/PCGAttributePropertySelector.h"
 
-#include "PropertyBag.h"
-#include "StructView.h"
+#include "StructUtils/PropertyBag.h"
+#include "StructUtils/StructView.h"
 
 #define LOCTEXT_NAMESPACE "PCGUserParameterGetElement"
 
diff --git a/Engine/Plugins/PCG/Source/PCG/Private/Tests/PCGGraphParametersTest.cpp b/Engine/Plugins/PCG/Source/PCG/Private/Tests/PCGGraphParametersTest.cpp
--- a/Engine/Plugins/PCG/Source/PCG/Private/Tests/PCGGraphParametersTest.cpp
+++ b/Engine/Plugins/PCG/Source/PCG/Private/Tests/PCGGraphParametersTest.cpp
@@ -5,8 +5,8 @@
 #include "Tests/PCGTestsCommon.h"
 
 #include "PCGGraph.h"
-#include "PropertyBag.h"
-#include "StructView.h"
+#include "StructUtils/PropertyBag.h"
+#include "StructUtils/StructView.h"
 
 IMPLEMENT_CUSTOM_SIMPLE_AUTOMATION_TEST(FPCGGraphNoUserParameters, FPCGTestBaseClass, "Plugins.PCG.Graph.NoUserParameters", PCGTestsCommon::TestFlags)
 IMPLEMENT_CUSTOM_SIMPLE_AUTOMATION_TEST(FPCGGraphAddSingleUserParameters, FPCGTestBaseClass, "Plugins.PCG.Graph.AddSingleUserParameter", PCGTestsCommon::TestFlags)
diff --git a/Engine/Plugins/Runtime/InstancedActors/Source/InstancedActors/Public/InstancedActorsSettingsTypes.h b/Engine/Plugins/Runtime/InstancedActors/Source/InstancedActors/Public/InstancedActorsSettingsTypes.h
--- a/Engine/Plugins/Runtime/InstancedActors/Source/InstancedActors/Public/InstancedActorsSettingsTypes.h
+++ b/Engine/Plugins/Runtime/InstancedActors/Source/InstancedActors/Public/InstancedActorsSettingsTypes.h
@@ -4,8 +4,8 @@
 
 #include "Engine/DataTable.h"
 #include "GameFramework/Actor.h"
-#include "StructView.h"
-#include "InstancedStruct.h"
+#include "StructUtils/StructView.h"
+#include "StructUtils/InstancedStruct.h"
 #include "InstancedActorsSettingsTypes.generated.h"
 
 
diff --git a/Engine/Plugins/Runtime/InstancedActors/Source/InstancedActors/Public/InstancedActorsStationaryLODBatchProcessor.h b/Engine/Plugins/Runtime/InstancedActors/Source/InstancedActors/Public/InstancedActorsStationaryLODBatchProcessor.h
--- a/Engine/Plugins/Runtime/InstancedActors/Source/InstancedActors/Public/InstancedActorsStationaryLODBatchProcessor.h
+++ b/Engine/Plugins/Runtime/InstancedActors/Source/InstancedActors/Public/InstancedActorsStationaryLODBatchProcessor.h
@@ -4,7 +4,7 @@
 
 #include "MassProcessor.h"
 #include "InstancedActorsTypes.h"
-#include "SharedStruct.h"
+#include "StructUtils/SharedStruct.h"
 #include "InstancedActorsStationaryLODBatchProcessor.generated.h"
 
 
diff --git a/Engine/Plugins/Runtime/MassEntity/Source/MassEntity/Public/MassEntityManager.h b/Engine/Plugins/Runtime/MassEntity/Source/MassEntity/Public/MassEntityManager.h
--- a/Engine/Plugins/Runtime/MassEntity/Source/MassEntity/Public/MassEntityManager.h
+++ b/Engine/Plugins/Runtime/MassEntity/Source/MassEntity/Public/MassEntityManager.h
@@ -7,8 +7,8 @@
 #include "MassProcessingTypes.h"
 #include "MassEntityQuery.h"
 #if UE_ENABLE_INCLUDE_ORDER_DEPRECATED_IN_5_5
-#include "InstancedStruct.h"
-#include "StructUtilsTypes.h"
+#include "StructUtils/InstancedStruct.h"
+#include "StructUtils/StructUtilsTypes.h"
 #endif
 #include "MassObserverManager.h"
 #include "Containers/MpscQueue.h"
diff --git a/Engine/Plugins/Runtime/StateTree/Source/StateTreeModule/Public/StateTreeInstanceData.h b/Engine/Plugins/Runtime/StateTree/Source/StateTreeModule/Public/StateTreeInstanceData.h
--- a/Engine/Plugins/Runtime/StateTree/Source/StateTreeModule/Public/StateTreeInstanceData.h
+++ b/Engine/Plugins/Runtime/StateTree/Source/StateTreeModule/Public/StateTreeInstanceData.h
@@ -2,7 +2,7 @@
 
 #pragma once
 
-#include "InstancedStructContainer.h"
+#include "StructUtils/InstancedStructContainer.h"
 #include "StateTreeEvents.h"
 #include "StateTreeTypes.h"
 #include "StateTreeExecutionTypes.h"
diff --git a/Engine/Source/Runtime/CoreUObject/Public/InstancedStruct.h b/Engine/Source/Runtime/CoreUObject/Public/InstancedStruct.h
--- a/Engine/Source/Runtime/CoreUObject/Public/InstancedStruct.h
+++ b/Engine/Source/Runtime/CoreUObject/Public/InstancedStruct.h
@@ -4,7 +4,4 @@
 
 #if UE_ENABLE_INCLUDE_ORDER_DEPRECATED_IN_5_5
 #include "StructUtils/InstancedStruct.h"
-#endif
-
-// Temporary and will be removed in a following change
-#include "StructUtils/InstancedStruct.h"
\ No newline at end of file
+#endif
\ No newline at end of file
diff --git a/Engine/Source/Runtime/CoreUObject/Public/InstancedStructContainer.h b/Engine/Source/Runtime/CoreUObject/Public/InstancedStructContainer.h
--- a/Engine/Source/Runtime/CoreUObject/Public/InstancedStructContainer.h
+++ b/Engine/Source/Runtime/CoreUObject/Public/InstancedStructContainer.h
@@ -4,7 +4,4 @@
 
 #if UE_ENABLE_INCLUDE_ORDER_DEPRECATED_IN_5_5
 #include "StructUtils/InstancedStructContainer.h"
-#endif
-
-// Temporary and will be removed in a following change
-#include "StructUtils/InstancedStructContainer.h"
\ No newline at end of file
+#endif
\ No newline at end of file
diff --git a/Engine/Source/Runtime/CoreUObject/Public/PropertyBag.h b/Engine/Source/Runtime/CoreUObject/Public/PropertyBag.h
--- a/Engine/Source/Runtime/CoreUObject/Public/PropertyBag.h
+++ b/Engine/Source/Runtime/CoreUObject/Public/PropertyBag.h
@@ -5,6 +5,3 @@
 #if UE_ENABLE_INCLUDE_ORDER_DEPRECATED_IN_5_5
 #include "StructUtils/PropertyBag.h"
 #endif
-
-// Temporary and will be removed in a following change
-#include "StructUtils/PropertyBag.h"
\ No newline at end of file
diff --git a/Engine/Source/Runtime/CoreUObject/Public/SharedStruct.h b/Engine/Source/Runtime/CoreUObject/Public/SharedStruct.h
--- a/Engine/Source/Runtime/CoreUObject/Public/SharedStruct.h
+++ b/Engine/Source/Runtime/CoreUObject/Public/SharedStruct.h
@@ -4,7 +4,4 @@
 
 #if UE_ENABLE_INCLUDE_ORDER_DEPRECATED_IN_5_5
 #include "StructUtils/SharedStruct.h"
-#endif
-
-// Temporary and will be removed in a following change
-#include "StructUtils/SharedStruct.h"
\ No newline at end of file
+#endif
\ No newline at end of file
diff --git a/Engine/Source/Runtime/CoreUObject/Public/StructView.h b/Engine/Source/Runtime/CoreUObject/Public/StructView.h
--- a/Engine/Source/Runtime/CoreUObject/Public/StructView.h
+++ b/Engine/Source/Runtime/CoreUObject/Public/StructView.h
@@ -4,7 +4,4 @@
 
 #if UE_ENABLE_INCLUDE_ORDER_DEPRECATED_IN_5_5
 #include "StructUtils/StructView.h"
-#endif
-
-// Temporary and will be removed in a following change
-#include "StructUtils/StructView.h"
\ No newline at end of file
+#endif
\ No newline at end of file
