[Motion Design] Playable Blueprint library: Remote Control Values injection (from playable transition) can be triggered from blueprint.

diff --git a/Engine/Plugins/Experimental/Avalanche/Source/AvalancheMedia/Private/Playable/AvaPlayableGroup.cpp b/Engine/Plugins/Experimental/Avalanche/Source/AvalancheMedia/Private/Playable/AvaPlayableGroup.cpp
--- a/Engine/Plugins/Experimental/Avalanche/Source/AvalancheMedia/Private/Playable/AvaPlayableGroup.cpp
+++ b/Engine/Plugins/Experimental/Avalanche/Source/AvalancheMedia/Private/Playable/AvaPlayableGroup.cpp
@@ -2,16 +2,20 @@
 
 #include "Playable/AvaPlayableGroup.h"
 
+#include "AvaPlayableGroupSubsystem.h"
 #include "Engine/GameInstance.h"
 #include "Engine/ViewportStatsSubsystem.h"
 #include "Engine/World.h"
+#include "IAvaMediaModule.h"
 #include "Playable/AvaPlayable.h"
 #include "Playable/AvaPlayableGroupManager.h"
 #include "Playable/PlayableGroups/AvaGameInstancePlayableGroup.h"
 #include "Playable/PlayableGroups/AvaGameViewportPlayableGroup.h"
 #include "Playable/PlayableGroups/AvaRemoteProxyPlayableGroup.h"
 #include "Playable/Transition/AvaPlayableTransition.h"
+#include "Playback/AvaPlaybackManager.h"
 #include "Playback/AvaPlaybackUtils.h"
+#include "Playback/IAvaPlaybackServer.h"
 #include "UObject/Package.h"
 
 #define LOCTEXT_NAMESPACE "AvaPlayableGroup"
@@ -371,6 +375,68 @@ bool UAvaPlayableGroup::IsVisibilityConstrained(const UAvaPlayable* InPlayable)
 	return false;
 }
 
+void UAvaPlayableGroup::ForEachPlayable(TFunctionRef<bool(UAvaPlayable*)> InFunction)
+{
+	for (const TObjectKey<UAvaPlayable>& PlayableKey : Playables)
+	{
+		if (UAvaPlayable* Playable = PlayableKey.ResolveObjectPtr())
+		{
+			if (!InFunction(Playable))
+			{
+				return;
+			}
+		}
+	}
+}
+	
+void UAvaPlayableGroup::ForEachPlayableTransition(TFunctionRef<bool(UAvaPlayableTransition*)> InFunction)
+{
+	for (const TObjectKey<UAvaPlayableTransition>& TransitionKey : PlayableTransitions)
+	{
+		if (UAvaPlayableTransition* Transition = TransitionKey.ResolveObjectPtr())
+		{
+			if (!InFunction(Transition))
+			{
+				return;
+			}
+		}
+	}
+}
+
+UAvaPlayableGroup* UAvaPlayableGroup::FindPlayableGroupForWorld(const UWorld* InWorld)
+{
+	// Ideally, we would like a direct link from GameInstance (or World) to it's owning PlayableGroup and PlayableGroupManager.
+	// Todo: Refactor the AvaGameInstance path to have a link to PlayableGroup.
+	const UAvaPlayableGroupManager* PlayableGroupManager = nullptr;
+
+	// For Game Viewport output, the sub system will give us the playable group manager directly.
+	if (const UGameInstance* GameInstance = InWorld->GetGameInstance())
+	{
+		if (const UAvaPlayableGroupSubsystem* PlayableGroupSubsystem = GameInstance->GetSubsystem<UAvaPlayableGroupSubsystem>())
+		{
+			PlayableGroupManager = PlayableGroupSubsystem->PlayableGroupManager;
+		}
+	}
+
+	if (!PlayableGroupManager)
+	{
+		const FAvaPlaybackManager& PlaybackManager = IAvaMediaModule::Get().GetLocalPlaybackManager();
+		PlayableGroupManager = PlaybackManager.GetPlayableGroupManager();
+	}
+
+	// Search in the local playable group manager for that world.
+	UAvaPlayableGroup* PlayableGroup = PlayableGroupManager->FindPlayableGroupForWorld(InWorld);
+	
+	// If not found, search in the playback server's playback manager.
+	if (!PlayableGroup && IAvaMediaModule::Get().IsPlaybackServerStarted())
+	{
+		PlayableGroupManager = IAvaMediaModule::Get().GetPlaybackServer()->GetPlaybackManager().GetPlayableGroupManager();
+		PlayableGroup = PlayableGroupManager->FindPlayableGroupForWorld(InWorld);
+	}
+
+	return PlayableGroup;
+}
+
 void UAvaPlayableGroup::OnPlayableStatusChanged(UAvaPlayable* InPlayable)
 {
 	// Evaluate the playable visibility requests
diff --git a/Engine/Plugins/Experimental/Avalanche/Source/AvalancheMedia/Private/Playable/AvaPlayableGroupSceneViewExtension.cpp b/Engine/Plugins/Experimental/Avalanche/Source/AvalancheMedia/Private/Playable/AvaPlayableGroupSceneViewExtension.cpp
--- a/Engine/Plugins/Experimental/Avalanche/Source/AvalancheMedia/Private/Playable/AvaPlayableGroupSceneViewExtension.cpp
+++ b/Engine/Plugins/Experimental/Avalanche/Source/AvalancheMedia/Private/Playable/AvaPlayableGroupSceneViewExtension.cpp
@@ -2,13 +2,7 @@
 
 #include "AvaPlayableGroupSceneViewExtension.h"
 
-#include "AvaPlayableGroupManager.h"
-#include "AvaPlayableGroupSubsystem.h"
-#include "Engine/GameInstance.h"
-#include "IAvaMediaModule.h"
 #include "Playable/AvaPlayableGroup.h"
-#include "Playback/AvaPlaybackManager.h"
-#include "Playback/IAvaPlaybackServer.h"
 #include "SceneView.h"
 
 FAvaPlayableGroupSceneViewExtension::FAvaPlayableGroupSceneViewExtension(const FAutoRegister& InAutoReg)
@@ -29,36 +23,7 @@ void FAvaPlayableGroupSceneViewExtension::SetupView(FSceneViewFamily& InViewFami
 		return;
 	}
 
-	// Ideally, we would like a direct link from GameInstance (or World) to it's owning PlayableGroup and PlayableGroupManager.
-	// Todo: Refactor the AvaGameInstance path to have a link to PlayableGroup.
-	const UAvaPlayableGroupManager* PlayableGroupManager = nullptr;
-
-	// For Game Viewport output, the sub system will give us the playable group manager directly.
-	if (const UGameInstance* GameInstance = ViewWorld->GetGameInstance())
-	{
-		if (const UAvaPlayableGroupSubsystem* PlayableGroupSubsystem = GameInstance->GetSubsystem<UAvaPlayableGroupSubsystem>())
-		{
-			PlayableGroupManager = PlayableGroupSubsystem->PlayableGroupManager;
-		}
-	}
-
-	if (!PlayableGroupManager)
-	{
-		const FAvaPlaybackManager& PlaybackManager = IAvaMediaModule::Get().GetLocalPlaybackManager();
-		PlayableGroupManager = PlaybackManager.GetPlayableGroupManager();
-	}
-
-	// Search in the local playable group manager for that world.
-	UAvaPlayableGroup* ViewPlayableGroup = PlayableGroupManager->FindPlayableGroupForWorld(ViewWorld);
-	
-	// If not found, search in the playback server's playback manager.
-	if (!ViewPlayableGroup && IAvaMediaModule::Get().IsPlaybackServerStarted())
-	{
-		PlayableGroupManager = IAvaMediaModule::Get().GetPlaybackServer()->GetPlaybackManager().GetPlayableGroupManager();
-		ViewPlayableGroup = PlayableGroupManager->FindPlayableGroupForWorld(ViewWorld);
-	}
-
-	if (ViewPlayableGroup)
+	if (UAvaPlayableGroup* ViewPlayableGroup = UAvaPlayableGroup::FindPlayableGroupForWorld(ViewWorld))
 	{
 		ViewPlayableGroup->SetupView(InViewFamily, InView);
 	}
diff --git a/Engine/Plugins/Experimental/Avalanche/Source/AvalancheMedia/Private/Playable/AvaPlayableLibrary.cpp b/Engine/Plugins/Experimental/Avalanche/Source/AvalancheMedia/Private/Playable/AvaPlayableLibrary.cpp
new file mode 100644
--- /dev/null
+++ b/Engine/Plugins/Experimental/Avalanche/Source/AvalancheMedia/Private/Playable/AvaPlayableLibrary.cpp
@@ -0,0 +1,111 @@
+// Copyright Epic Games, Inc. All Rights Reserved.
+
+#include "Playable/AvaPlayableLibrary.h"
+
+#include "Engine/Engine.h"
+#include "Engine/Level.h"
+#include "Engine/LevelStreamingDynamic.h"
+#include "Engine/World.h"
+#include "Playable/AvaPlayableGroup.h"
+#include "Playable/Playables/AvaPlayableLevelStreaming.h"
+#include "Playable/Transition/AvaPlayableTransition.h"
+
+namespace UE::AvaPlayableLibrary::Private
+{
+	ULevel* GetLevel(const UObject* InWorldContextObject)
+	{
+		if (!InWorldContextObject)
+		{
+			return nullptr;
+		}
+
+		ULevel* Level = InWorldContextObject->GetTypedOuter<ULevel>();
+		if (!Level && GEngine)
+		{
+			if (const UWorld* World = GEngine->GetWorldFromContextObject(InWorldContextObject, EGetWorldErrorMode::LogAndReturnNull))
+			{
+				Level = World->PersistentLevel;
+			}
+		}
+		return Level;
+	}
+}
+
+
+UAvaPlayable* UAvaPlayableLibrary::GetPlayable(const UObject* InWorldContextObject)
+{
+	using namespace UE::AvaPlayableLibrary::Private;
+	const ULevel* Level = GetLevel(InWorldContextObject);
+	if (!Level)
+	{
+		return nullptr;
+	}
+
+	UAvaPlayableGroup* PlayableGroup = UAvaPlayableGroup::FindPlayableGroupForWorld(Level->OwningWorld);
+	if (!PlayableGroup)
+	{
+		return nullptr;
+	}
+
+	UAvaPlayable* FoundPlayable = nullptr;
+	PlayableGroup->ForEachPlayable([&FoundPlayable, Level](UAvaPlayable* InPlayable)
+	{
+		if (const UAvaPlayableLevelStreaming* PlayableLevelStreaming = Cast<UAvaPlayableLevelStreaming>(InPlayable))
+		{
+			if (const ULevelStreaming* LevelStreaming = PlayableLevelStreaming->GetLevelStreaming())
+			{
+				if (LevelStreaming->GetLoadedLevel() == Level)
+				{
+					FoundPlayable = InPlayable;
+					return false;
+				}
+			}
+		}
+		return true;
+	});
+
+	return FoundPlayable;
+}
+
+UAvaPlayableTransition* UAvaPlayableLibrary::GetPlayableTransition(const UAvaPlayable* InPlayable)
+{
+	if (!InPlayable)
+	{
+		return nullptr;
+	}
+
+	UAvaPlayableGroup* PlayableGroup = InPlayable->GetPlayableGroup();
+	if (!PlayableGroup)
+	{
+		return nullptr;
+	}
+
+	UAvaPlayableTransition* FoundTransition = nullptr;
+	PlayableGroup->ForEachPlayableTransition([&FoundTransition, InPlayable](UAvaPlayableTransition* InTransition)
+	{
+		if (InTransition->IsEnterPlayable(InPlayable)
+			|| InTransition->IsPlayingPlayable(InPlayable)
+			|| InTransition->IsExitPlayable(InPlayable))
+		{
+			FoundTransition = InTransition;
+			return false;
+		}		
+		return true;
+	});
+
+	return FoundTransition;
+}
+
+void UAvaPlayableLibrary::UpdateRemoteControlValues(const UObject* InWorldContextObject)
+{
+	if (UAvaPlayable* Playable = GetPlayable(InWorldContextObject))
+	{
+		if (UAvaPlayableTransition* Transition = GetPlayableTransition(Playable))
+		{
+			if (const TSharedPtr<FAvaPlayableRemoteControlValues> RemoteControlValues = Transition->GetValuesForPlayable(Playable))
+			{
+				Playable->UpdateRemoteControlCommand(RemoteControlValues.ToSharedRef());
+			}
+		}
+	}
+}
\ No newline at end of file
diff --git a/Engine/Plugins/Experimental/Avalanche/Source/AvalancheMedia/Private/Playable/Transition/AvaPlayableTransition.cpp b/Engine/Plugins/Experimental/Avalanche/Source/AvalancheMedia/Private/Playable/Transition/AvaPlayableTransition.cpp
--- a/Engine/Plugins/Experimental/Avalanche/Source/AvalancheMedia/Private/Playable/Transition/AvaPlayableTransition.cpp
+++ b/Engine/Plugins/Experimental/Avalanche/Source/AvalancheMedia/Private/Playable/Transition/AvaPlayableTransition.cpp
@@ -74,17 +74,17 @@ void UAvaPlayableTransition::SetExitPlayables(TArray<TWeakObjectPtr<UAvaPlayable
 	ExitPlayablesWeak = MoveTemp(InPlayablesWeak);
 }
 
-bool UAvaPlayableTransition::IsEnterPlayable(UAvaPlayable* InPlayable) const
+bool UAvaPlayableTransition::IsEnterPlayable(const UAvaPlayable* InPlayable) const
 {
 	return EnterPlayablesWeak.Contains(InPlayable);
 }
 
-bool UAvaPlayableTransition::IsPlayingPlayable(UAvaPlayable* InPlayable) const
+bool UAvaPlayableTransition::IsPlayingPlayable(const UAvaPlayable* InPlayable) const
 {
 	return PlayingPlayablesWeak.Contains(InPlayable);
 }
 
-bool UAvaPlayableTransition::IsExitPlayable(UAvaPlayable* InPlayable) const
+bool UAvaPlayableTransition::IsExitPlayable(const UAvaPlayable* InPlayable) const
 {
 	return ExitPlayablesWeak.Contains(InPlayable);
 }
@@ -94,6 +94,21 @@ void UAvaPlayableTransition::SetEnterPlayableValues(TArray<TSharedPtr<FAvaPlayab
 	EnterPlayableValues = MoveTemp(InPlayableValues);
 }
 
+TSharedPtr<FAvaPlayableRemoteControlValues> UAvaPlayableTransition::GetValuesForPlayable(const UAvaPlayable* InPlayable)
+{
+	int32 ArrayIndex = 0;
+	for (const TWeakObjectPtr<UAvaPlayable> PlayableWeak : EnterPlayablesWeak)
+	{
+		const UAvaPlayable* Playable = PlayableWeak.Get();
+		if (Playable == InPlayable && EnterPlayableValues.IsValidIndex(ArrayIndex))
+		{
+			return EnterPlayableValues[ArrayIndex];
+		}
+		++ArrayIndex;
+	}
+	return nullptr;
+}
+
 void UAvaPlayableTransition::MarkPlayableAsDiscard(UAvaPlayable* InPlayable)
 {
 	DiscardPlayablesWeak.AddUnique(InPlayable);
diff --git a/Engine/Plugins/Experimental/Avalanche/Source/AvalancheMedia/Public/Playable/AvaPlayableGroup.h b/Engine/Plugins/Experimental/Avalanche/Source/AvalancheMedia/Public/Playable/AvaPlayableGroup.h
--- a/Engine/Plugins/Experimental/Avalanche/Source/AvalancheMedia/Public/Playable/AvaPlayableGroup.h
+++ b/Engine/Plugins/Experimental/Avalanche/Source/AvalancheMedia/Public/Playable/AvaPlayableGroup.h
@@ -173,6 +173,15 @@ public:
 
 	virtual void SetupView(FSceneViewFamily& InViewFamily, FSceneView& InView);
 
+	void ForEachPlayable(TFunctionRef<bool(UAvaPlayable*)> InFunction);
+	
+	void ForEachPlayableTransition(TFunctionRef<bool(UAvaPlayableTransition*)> InFunction);
+
+	/**
+	 * Globally search for the playable group corresponding to the given world.
+	 */
+	static UAvaPlayableGroup* FindPlayableGroupForWorld(const UWorld* InWorld);
+	
 protected:
 	bool IsVisibilityConstrained(const UAvaPlayable* InPlayable) const;
 
diff --git a/Engine/Plugins/Experimental/Avalanche/Source/AvalancheMedia/Public/Playable/AvaPlayableLibrary.h b/Engine/Plugins/Experimental/Avalanche/Source/AvalancheMedia/Public/Playable/AvaPlayableLibrary.h
new file mode 100644
--- /dev/null
+++ b/Engine/Plugins/Experimental/Avalanche/Source/AvalancheMedia/Public/Playable/AvaPlayableLibrary.h
@@ -0,0 +1,26 @@
+// Copyright Epic Games, Inc. All Rights Reserved.
+
+#pragma once
+
+#include "Kismet/BlueprintFunctionLibrary.h"
+#include "AvaPlayableLibrary.generated.h"
+
+class UAvaPlayable;
+class UAvaPlayableTransition;
+
+UCLASS(MinimalAPI, DisplayName = "Motion Design Playable Library", meta=(ScriptName = "MotionDesignPlayableLibrary"))
+class UAvaPlayableLibrary : public UBlueprintFunctionLibrary
+{
+	GENERATED_BODY()
+
+public:
+	/** Returns the playable corresponding to that world object. */
+	static UAvaPlayable* GetPlayable(const UObject* InWorldContextObject);
+
+	/** Returns the transition this playable is part of. */
+	static UAvaPlayableTransition* GetPlayableTransition(const UAvaPlayable* InPlayable);
+
+	/** Injects the remote control values from current transition for the current playable. */
+	UFUNCTION(BlueprintCallable, Category = "Motion Design|Playable", meta=(WorldContext = "InWorldContextObject"))
+	static AVALANCHEMEDIA_API void UpdateRemoteControlValues(const UObject* InWorldContextObject);
+};
diff --git a/Engine/Plugins/Experimental/Avalanche/Source/AvalancheMedia/Public/Playable/Transition/AvaPlayableTransition.h b/Engine/Plugins/Experimental/Avalanche/Source/AvalancheMedia/Public/Playable/Transition/AvaPlayableTransition.h
--- a/Engine/Plugins/Experimental/Avalanche/Source/AvalancheMedia/Public/Playable/Transition/AvaPlayableTransition.h
+++ b/Engine/Plugins/Experimental/Avalanche/Source/AvalancheMedia/Public/Playable/Transition/AvaPlayableTransition.h
@@ -40,11 +40,13 @@ public:
 	void SetPlayingPlayables(TArray<TWeakObjectPtr<UAvaPlayable>>&& InPlayablesWeak);
 	void SetExitPlayables(TArray<TWeakObjectPtr<UAvaPlayable>>&& InPlayablesWeak);
 	
-	bool IsEnterPlayable(UAvaPlayable* InPlayable) const;
-	bool IsPlayingPlayable(UAvaPlayable* InPlayable) const;
-	bool IsExitPlayable(UAvaPlayable* InPlayable) const;
+	bool IsEnterPlayable(const UAvaPlayable* InPlayable) const;
+	bool IsPlayingPlayable(const UAvaPlayable* InPlayable) const;
+	bool IsExitPlayable(const UAvaPlayable* InPlayable) const;
 
 	void SetEnterPlayableValues(TArray<TSharedPtr<FAvaPlayableRemoteControlValues>>&& InPlayableValues);
+
+	TSharedPtr<FAvaPlayableRemoteControlValues> GetValuesForPlayable(const UAvaPlayable* InPlayable);
 	
 	/** This is called during the transition evaluation to indicate discarded playables. */
 	void MarkPlayableAsDiscard(UAvaPlayable* InPlayable);
