[Motion Design] Broadcast Blueprint library: Get current channel's viewport resolution.

diff --git a/Engine/Plugins/Experimental/Avalanche/Source/AvalancheMedia/Private/Broadcast/AvaBroadcastLibrary.cpp b/Engine/Plugins/Experimental/Avalanche/Source/AvalancheMedia/Private/Broadcast/AvaBroadcastLibrary.cpp
new file mode 100644
--- /dev/null
+++ b/Engine/Plugins/Experimental/Avalanche/Source/AvalancheMedia/Private/Broadcast/AvaBroadcastLibrary.cpp
@@ -0,0 +1,51 @@
+// Copyright Epic Games, Inc. All Rights Reserved.
+
+#include "Broadcast/AvaBroadcastLibrary.h"
+
+#include "Broadcast/AvaBroadcast.h"
+#include "Broadcast/Channel/AvaBroadcastOutputChannel.h"
+#include "Engine/Engine.h"
+#include "Engine/GameInstance.h"
+#include "Engine/GameViewportClient.h"
+#include "Playable/AvaPlayable.h"
+#include "Playable/AvaPlayableGroup.h"
+#include "Playable/AvaPlayableLibrary.h"
+
+FVector2D UAvaBroadcastLibrary::GetChannelViewportSize(const UObject* InWorldContextObject)
+{
+	// Attempt to access the game instance's viewport (generic method)
+	if (InWorldContextObject && GEngine)
+	{
+		if (const UWorld* World = GEngine->GetWorldFromContextObject(InWorldContextObject, EGetWorldErrorMode::LogAndReturnNull))
+		{
+			if (const UGameInstance* GameInstance = World->GetGameInstance())
+			{
+				if (const UGameViewportClient* ViewportClient = GameInstance->GetGameViewportClient())
+				{
+					FVector2d ViewportSize;
+					ViewportClient->GetViewportSize(ViewportSize);
+					return ViewportSize;					
+				}
+			}
+		}
+	}
+
+	// Fallback to Channel Render Target (for AvaGameInstance)
+	if (const UAvaPlayable* Playable = UAvaPlayableLibrary::GetPlayable(InWorldContextObject))
+	{
+		if (const UAvaPlayableGroup* PlayableGroup = Playable->GetPlayableGroup())
+		{
+			const FAvaBroadcastOutputChannel& Channel = UAvaBroadcast::Get().GetCurrentProfile().GetChannel(PlayableGroup->GetChannelName());
+			if (Channel.IsValidChannel())
+			{
+				if (const UTextureRenderTarget2D* RenderTarget = Channel.GetCurrentRenderTarget(/*bInFallbackToPlaceholder*/true))
+				{
+					return FVector2D(RenderTarget->SizeX, RenderTarget->SizeY);
+				}
+				return Channel.DetermineRenderTargetSize();
+			}
+		}
+	}
+	
+	return FAvaBroadcastOutputChannel::GetDefaultMediaOutputSize(EAvaBroadcastChannelType::Program);
+}
\ No newline at end of file
diff --git a/Engine/Plugins/Experimental/Avalanche/Source/AvalancheMedia/Private/Playable/AvaPlayableGroup.cpp b/Engine/Plugins/Experimental/Avalanche/Source/AvalancheMedia/Private/Playable/AvaPlayableGroup.cpp
--- a/Engine/Plugins/Experimental/Avalanche/Source/AvalancheMedia/Private/Playable/AvaPlayableGroup.cpp
+++ b/Engine/Plugins/Experimental/Avalanche/Source/AvalancheMedia/Private/Playable/AvaPlayableGroup.cpp
@@ -39,25 +39,31 @@ UAvaPlayableGroup* UAvaPlayableGroup::MakePlayableGroup(UObject* InOuter, const
 {
 	UObject* Outer = InOuter ? InOuter : GetTransientPackage();
 	
-	UAvaPlayableGroup* GameInstanceGroup = nullptr;
+	UAvaPlayableGroup* NewPlayableGroup = nullptr;
 	if (InPlayableGroupInfo.bIsRemoteProxy)
 	{
 		// Remote Proxy group doesn't have a game instance.
-		GameInstanceGroup = NewObject<UAvaRemoteProxyPlayableGroup>(Outer);
-		GameInstanceGroup->ParentPlayableGroupManagerWeak = InPlayableGroupInfo.PlayableGroupManager;
+		NewPlayableGroup = NewObject<UAvaRemoteProxyPlayableGroup>(Outer);
+		NewPlayableGroup->ParentPlayableGroupManagerWeak = InPlayableGroupInfo.PlayableGroupManager;
 	}
 	else
 	{
 		if (InPlayableGroupInfo.GameInstance)
 		{
-			GameInstanceGroup = UAvaGameViewportPlayableGroup::Create(InOuter, InPlayableGroupInfo.GameInstance, InPlayableGroupInfo.PlayableGroupManager);
+			NewPlayableGroup = UAvaGameViewportPlayableGroup::Create(InOuter, InPlayableGroupInfo.GameInstance, InPlayableGroupInfo.PlayableGroupManager);
 		}
 		else
 		{
-			GameInstanceGroup = UAvaGameInstancePlayableGroup::Create(InOuter, InPlayableGroupInfo);
+			NewPlayableGroup = UAvaGameInstancePlayableGroup::Create(InOuter, InPlayableGroupInfo);
 		}
 	}
-	return GameInstanceGroup;
+
+	if (NewPlayableGroup)
+	{
+		NewPlayableGroup->ChannelName = InPlayableGroupInfo.ChannelName;
+	}
+	
+	return NewPlayableGroup;
 }
 
 void UAvaPlayableGroup::RegisterPlayable(UAvaPlayable* InPlayable)
@@ -287,6 +293,11 @@ UWorld* UAvaPlayableGroup::GetPlayWorld() const
 	return GameInstance ? GameInstance->GetWorld() : nullptr;
 }
 
+FName UAvaPlayableGroup::GetChannelName() const
+{
+	return ChannelName;
+}
+
 void UAvaPlayableGroup::NotifyLevelStreaming(UAvaPlayable* InPlayable)
 {
 	// If the world is not playing, we need to make sure the level streaming still updates.
diff --git a/Engine/Plugins/Experimental/Avalanche/Source/AvalancheMedia/Public/Broadcast/AvaBroadcastLibrary.h b/Engine/Plugins/Experimental/Avalanche/Source/AvalancheMedia/Public/Broadcast/AvaBroadcastLibrary.h
new file mode 100644
--- /dev/null
+++ b/Engine/Plugins/Experimental/Avalanche/Source/AvalancheMedia/Public/Broadcast/AvaBroadcastLibrary.h
@@ -0,0 +1,17 @@
+// Copyright Epic Games, Inc. All Rights Reserved.
+
+#pragma once
+
+#include "Kismet/BlueprintFunctionLibrary.h"
+#include "AvaBroadcastLibrary.generated.h"
+
+UCLASS(MinimalAPI, DisplayName = "Motion Design Broadcast Library", meta=(ScriptName = "MotionDesignBroadcastLibrary"))
+class UAvaBroadcastLibrary : public UBlueprintFunctionLibrary
+{
+	GENERATED_BODY()
+
+public:
+	/** Returns the current channel's viewport size. */
+	UFUNCTION(BlueprintCallable, Category = "Motion Design|broadcast", meta=(WorldContext = "InWorldContextObject"))
+	static AVALANCHEMEDIA_API FVector2D GetChannelViewportSize(const UObject* InWorldContextObject);
+};
diff --git a/Engine/Plugins/Experimental/Avalanche/Source/AvalancheMedia/Public/Playable/AvaPlayableGroup.h b/Engine/Plugins/Experimental/Avalanche/Source/AvalancheMedia/Public/Playable/AvaPlayableGroup.h
--- a/Engine/Plugins/Experimental/Avalanche/Source/AvalancheMedia/Public/Playable/AvaPlayableGroup.h
+++ b/Engine/Plugins/Experimental/Avalanche/Source/AvalancheMedia/Public/Playable/AvaPlayableGroup.h
@@ -146,6 +146,11 @@ public:
 	* Returns this group's play world, if it has one.
 	*/
 	virtual UWorld* GetPlayWorld() const;
+
+	/**
+	 * Returns the broadcast channel name this playable group is part of.
+	 */
+	FName GetChannelName() const;
 	
 	/**
 	 * Unloads the game instance's world if no more playables are loaded.
@@ -204,6 +209,9 @@ protected:
 	UPROPERTY(Transient)
 	TObjectPtr<UGameInstance> GameInstance;
 
+	/** Broadcast channel name this playable group is part of. */
+	FName ChannelName;
+
 	/** PlayableGroup Manager handling this playable group. */
 	TWeakObjectPtr<UAvaPlayableGroupManager> ParentPlayableGroupManagerWeak;
 	
