[Motion Design] New method to find the corresponding playable and playable group for a given level using AssetUserData.

Having a direct link from the managed asset instance to it's managing playable or playable group makes the system more efficient and future proof.
It avoids having to search the root managers, and becomes unaffected by changes to those root objects. In terms of "scoping" it is much better.
Also, eventually, other implementations will be done, hopefully the IInterface_AssetUserData will be available on the other assets as well so it can be used.

diff --git a/Engine/Plugins/Experimental/Avalanche/Source/AvalancheMedia/Private/Playable/AvaPlayableAssetUserData.h b/Engine/Plugins/Experimental/Avalanche/Source/AvalancheMedia/Private/Playable/AvaPlayableAssetUserData.h
new file mode 100644
--- /dev/null
+++ b/Engine/Plugins/Experimental/Avalanche/Source/AvalancheMedia/Private/Playable/AvaPlayableAssetUserData.h
@@ -0,0 +1,23 @@
+// Copyright Epic Games, Inc. All Rights Reserved.
+
+#pragma once
+
+#include "Engine/AssetUserData.h"
+#include "AvaPlayableAssetUserData.generated.h"
+
+class UAvaPlayable;
+
+/**
+ * AssetUserData for AvaPlayable.
+ *
+ * This is strategically injected in asset instances to be able to retrieve
+ * the corresponding playable directly.
+ */
+UCLASS()
+class UAvaPlayableAssetUserData : public UAssetUserData
+{
+	GENERATED_BODY()
+
+public:
+	TWeakObjectPtr<UAvaPlayable> PlayableWeak;
+};
diff --git a/Engine/Plugins/Experimental/Avalanche/Source/AvalancheMedia/Private/Playable/AvaPlayableGroup.cpp b/Engine/Plugins/Experimental/Avalanche/Source/AvalancheMedia/Private/Playable/AvaPlayableGroup.cpp
--- a/Engine/Plugins/Experimental/Avalanche/Source/AvalancheMedia/Private/Playable/AvaPlayableGroup.cpp
+++ b/Engine/Plugins/Experimental/Avalanche/Source/AvalancheMedia/Private/Playable/AvaPlayableGroup.cpp
@@ -4,10 +4,12 @@
 
 #include "AvaPlayableGroupSubsystem.h"
 #include "Engine/GameInstance.h"
+#include "Engine/Level.h"
 #include "Engine/ViewportStatsSubsystem.h"
 #include "Engine/World.h"
 #include "IAvaMediaModule.h"
 #include "Playable/AvaPlayable.h"
+#include "Playable/AvaPlayableGroupAssetUserData.h"
 #include "Playable/AvaPlayableGroupManager.h"
 #include "Playable/PlayableGroups/AvaGameInstancePlayableGroup.h"
 #include "Playable/PlayableGroups/AvaGameViewportPlayableGroup.h"
@@ -33,6 +35,32 @@ namespace UE::AvaMedia::PlayableGroup::Private
 			}
 		}
 	}
+
+	UAvaPlayableGroupAssetUserData* FindPlayableGroupAssetUserData(IInterface_AssetUserData* InAssetInstance)
+	{
+		return InAssetInstance->GetAssetUserData<UAvaPlayableGroupAssetUserData>();
+	}
+
+	UAvaPlayableGroupAssetUserData* FindPlayableGroupAssetUserDataSafe(IInterface_AssetUserData* InAssetInstance)
+	{
+		return InAssetInstance ? InAssetInstance->GetAssetUserData<UAvaPlayableGroupAssetUserData>() : nullptr;
+	}
+
+	UAvaPlayableGroupAssetUserData* FindOrAddPlayableGroupAssetUserData(IInterface_AssetUserData* InAssetInstance)
+	{
+		if (!InAssetInstance)
+		{
+			return nullptr;
+		}
+		
+		UAvaPlayableGroupAssetUserData* PlayableGroupUserData = FindPlayableGroupAssetUserData(InAssetInstance);
+		if (!PlayableGroupUserData)
+		{
+			PlayableGroupUserData = NewObject<UAvaPlayableGroupAssetUserData>();
+			InAssetInstance->AddAssetUserData(PlayableGroupUserData);
+		}
+		return PlayableGroupUserData;
+	}
 }
 
 UAvaPlayableGroup* UAvaPlayableGroup::MakePlayableGroup(UObject* InOuter, const FPlayableGroupCreationInfo& InPlayableGroupInfo)
@@ -61,6 +89,15 @@ UAvaPlayableGroup* UAvaPlayableGroup::MakePlayableGroup(UObject* InOuter, const
 	if (NewPlayableGroup)
 	{
 		NewPlayableGroup->ChannelName = InPlayableGroupInfo.ChannelName;
+
+		if (const UWorld* PlayWorld = NewPlayableGroup->GetPlayWorld())
+		{
+			using namespace UE::AvaMedia::PlayableGroup::Private;
+			if(UAvaPlayableGroupAssetUserData* PlayableGroupUserData = FindOrAddPlayableGroupAssetUserData(PlayWorld->PersistentLevel))
+			{
+				PlayableGroupUserData->PlayableGroupsWeak.AddUnique(NewPlayableGroup);
+			}
+		}
 	}
 	
 	return NewPlayableGroup;
@@ -414,10 +451,33 @@ void UAvaPlayableGroup::ForEachPlayableTransition(TFunctionRef<bool(UAvaPlayable
 	}
 }
 
-UAvaPlayableGroup* UAvaPlayableGroup::FindPlayableGroupForWorld(const UWorld* InWorld)
+UAvaPlayableGroup* UAvaPlayableGroup::FindPlayableGroupForWorld(const UWorld* InWorld, bool bInFallbackToGlobalSearch)
 {
-	// Ideally, we would like a direct link from GameInstance (or World) to it's owning PlayableGroup and PlayableGroupManager.
-	// Todo: Refactor the AvaGameInstance path to have a link to PlayableGroup.
+	if (!InWorld)
+	{
+		return nullptr;
+	}
+
+	// Fast path: if a world is managed by a playable group, it should have an asset user data that
+	// we can retrieve the corresponding playable group from.
+	using namespace UE::AvaMedia::PlayableGroup::Private;
+	if (UAvaPlayableGroupAssetUserData* PlayableGroupUserData = FindPlayableGroupAssetUserData(InWorld->PersistentLevel))
+	{
+		for (const TWeakObjectPtr<UAvaPlayableGroup>& PlayableGroupWeak : PlayableGroupUserData->PlayableGroupsWeak)
+		{
+			if (UAvaPlayableGroup* PlayableGroup = PlayableGroupWeak.Get())
+			{
+				return PlayableGroup;
+			}
+		}
+	}
+
+	if (!bInFallbackToGlobalSearch)
+	{
+		return nullptr;
+	}
+	
+	// Global Search starting from the system's root playable group managers.
 	const UAvaPlayableGroupManager* PlayableGroupManager = nullptr;
 
 	// For Game Viewport output, the sub system will give us the playable group manager directly.
diff --git a/Engine/Plugins/Experimental/Avalanche/Source/AvalancheMedia/Private/Playable/AvaPlayableGroupAssetUserData.h b/Engine/Plugins/Experimental/Avalanche/Source/AvalancheMedia/Private/Playable/AvaPlayableGroupAssetUserData.h
new file mode 100644
--- /dev/null
+++ b/Engine/Plugins/Experimental/Avalanche/Source/AvalancheMedia/Private/Playable/AvaPlayableGroupAssetUserData.h
@@ -0,0 +1,29 @@
+// Copyright Epic Games, Inc. All Rights Reserved.
+
+#pragma once
+
+#include "Engine/AssetUserData.h"
+#include "AvaPlayableGroupAssetUserData.generated.h"
+
+class UAvaPlayableGroup;
+
+/**
+ * AssetUserData for AvaPlayableGroup.
+ * 
+ * This is strategically injected in asset instances to be able to retrieve
+ * the corresponding playable group directly.
+ */
+UCLASS()
+class UAvaPlayableGroupAssetUserData : public UAssetUserData
+{
+	GENERATED_BODY()
+
+public:
+	/**
+	 * A playable group being loosely defined in the playable framework,
+	 * an asset instance may be associated with more than one playable group.
+	 * The intent is to be injecting this in asset instances that are associated with
+	 * just one playable group.
+	 */
+	TArray<TWeakObjectPtr<UAvaPlayableGroup>> PlayableGroupsWeak;
+};
diff --git a/Engine/Plugins/Experimental/Avalanche/Source/AvalancheMedia/Private/Playable/AvaPlayableLibrary.cpp b/Engine/Plugins/Experimental/Avalanche/Source/AvalancheMedia/Private/Playable/AvaPlayableLibrary.cpp
--- a/Engine/Plugins/Experimental/Avalanche/Source/AvalancheMedia/Private/Playable/AvaPlayableLibrary.cpp
+++ b/Engine/Plugins/Experimental/Avalanche/Source/AvalancheMedia/Private/Playable/AvaPlayableLibrary.cpp
@@ -4,10 +4,10 @@
 
 #include "Engine/Engine.h"
 #include "Engine/Level.h"
-#include "Engine/LevelStreamingDynamic.h"
 #include "Engine/World.h"
+#include "Playable/AvaPlayable.h"
+#include "Playable/AvaPlayableAssetUserData.h"
 #include "Playable/AvaPlayableGroup.h"
-#include "Playable/Playables/AvaPlayableLevelStreaming.h"
 #include "Playable/Transition/AvaPlayableTransition.h"
 
 namespace UE::AvaPlayableLibrary::Private
@@ -35,36 +35,21 @@ namespace UE::AvaPlayableLibrary::Private
 UAvaPlayable* UAvaPlayableLibrary::GetPlayable(const UObject* InWorldContextObject)
 {
 	using namespace UE::AvaPlayableLibrary::Private;
-	const ULevel* Level = GetLevel(InWorldContextObject);
+	ULevel* Level = GetLevel(InWorldContextObject);
 	if (!Level)
 	{
 		return nullptr;
 	}
 
-	UAvaPlayableGroup* PlayableGroup = UAvaPlayableGroup::FindPlayableGroupForWorld(Level->OwningWorld);
-	if (!PlayableGroup)
+	if (const UAvaPlayableAssetUserData* PlayableUserData = Level->GetAssetUserData<UAvaPlayableAssetUserData>())
 	{
-		return nullptr;
-	}
-
-	UAvaPlayable* FoundPlayable = nullptr;
-	PlayableGroup->ForEachPlayable([&FoundPlayable, Level](UAvaPlayable* InPlayable)
-	{
-		if (const UAvaPlayableLevelStreaming* PlayableLevelStreaming = Cast<UAvaPlayableLevelStreaming>(InPlayable))
+		if (UAvaPlayable* Playable = PlayableUserData->PlayableWeak.Get())
 		{
-			if (const ULevelStreaming* LevelStreaming = PlayableLevelStreaming->GetLevelStreaming())
-			{
-				if (LevelStreaming->GetLoadedLevel() == Level)
-				{
-					FoundPlayable = InPlayable;
-					return false;
-				}
-			}
+			return Playable;
 		}
-		return true;
-	});
+	}
 
-	return FoundPlayable;
+	return nullptr;
 }
 
 UAvaPlayableTransition* UAvaPlayableLibrary::GetPlayableTransition(const UAvaPlayable* InPlayable)
@@ -96,7 +81,7 @@ UAvaPlayableTransition* UAvaPlayableLibrary::GetPlayableTransition(const UAvaPla
 	return FoundTransition;
 }
 
-void UAvaPlayableLibrary::UpdateRemoteControlValues(const UObject* InWorldContextObject)
+bool UAvaPlayableLibrary::UpdateRemoteControlValues(const UObject* InWorldContextObject)
 {
 	if (UAvaPlayable* Playable = GetPlayable(InWorldContextObject))
 	{
@@ -104,8 +89,12 @@ void UAvaPlayableLibrary::UpdateRemoteControlValues(const UObject* InWorldContex
 		{
 			if (const TSharedPtr<FAvaPlayableRemoteControlValues> RemoteControlValues = Transition->GetValuesForPlayable(Playable))
 			{
-				Playable->UpdateRemoteControlCommand(RemoteControlValues.ToSharedRef());
+				if (Playable->UpdateRemoteControlCommand(RemoteControlValues.ToSharedRef()) == EAvaPlayableCommandResult::Executed)
+				{
+					return true;
+				}
 			}
 		}
 	}
+	return false;
 }
\ No newline at end of file
diff --git a/Engine/Plugins/Experimental/Avalanche/Source/AvalancheMedia/Private/Playable/Playables/AvaPlayableLevelStreaming.cpp b/Engine/Plugins/Experimental/Avalanche/Source/AvalancheMedia/Private/Playable/Playables/AvaPlayableLevelStreaming.cpp
--- a/Engine/Plugins/Experimental/Avalanche/Source/AvalancheMedia/Private/Playable/Playables/AvaPlayableLevelStreaming.cpp
+++ b/Engine/Plugins/Experimental/Avalanche/Source/AvalancheMedia/Private/Playable/Playables/AvaPlayableLevelStreaming.cpp
@@ -15,6 +15,7 @@
 #include "Framework/AvaSoftAssetPtr.h"
 #include "IAvaMediaModule.h"
 #include "IAvaRemoteControlInterface.h"
+#include "Playable/AvaPlayableAssetUserData.h"
 #include "Playable/AvaPlayableGroup.h"
 #include "Playable/AvaPlayableGroupManager.h"
 #include "Playback/AvaPlaybackUtils.h"
@@ -27,6 +28,22 @@
 
 namespace UE::AvaMedia::LevelStreamingPlayable::Private
 {
+	UAvaPlayableAssetUserData* FindOrAddPlayableAssetUserData(ULevel* InLevel)
+	{
+		if (!InLevel)
+		{
+			return nullptr;
+		}
+
+		UAvaPlayableAssetUserData* PlayableUserData = InLevel->GetAssetUserData<UAvaPlayableAssetUserData>();
+		if (!PlayableUserData)
+		{
+			PlayableUserData = NewObject<UAvaPlayableAssetUserData>();
+			InLevel->AddAssetUserData(PlayableUserData);
+		}
+		return PlayableUserData;
+	}
+
 	AAvaScene* FindAvaScene(const ULevel* InLevel)
 	{
 		AAvaScene* AvaScene = nullptr;
@@ -500,6 +517,14 @@ void UAvaPlayableLevelStreaming::OnLevelStreamingStateChanged(UWorld* InWorld
 		return;
 	}
 
+	using namespace UE::AvaMedia::LevelStreamingPlayable::Private;
+	
+	// Inject Playable user data in streamed level.
+	if (UAvaPlayableAssetUserData* PlayableUserData = FindOrAddPlayableAssetUserData(InLevelIfLoaded))
+	{
+		PlayableUserData->PlayableWeak = this;
+	}
+
 	// Package the event handler for queueing.
 	TWeakObjectPtr<UAvaPlayableLevelStreaming> ThisPlayableWeak(this);
 	auto SyncEventHandler = [ThisPlayableWeak, InNewState]()
diff --git a/Engine/Plugins/Experimental/Avalanche/Source/AvalancheMedia/Public/Playable/AvaPlayableGroup.h b/Engine/Plugins/Experimental/Avalanche/Source/AvalancheMedia/Public/Playable/AvaPlayableGroup.h
--- a/Engine/Plugins/Experimental/Avalanche/Source/AvalancheMedia/Public/Playable/AvaPlayableGroup.h
+++ b/Engine/Plugins/Experimental/Avalanche/Source/AvalancheMedia/Public/Playable/AvaPlayableGroup.h
@@ -183,9 +183,11 @@ public:
 	void ForEachPlayableTransition(TFunctionRef<bool(UAvaPlayableTransition*)> InFunction);
 
 	/**
-	 * Globally search for the playable group corresponding to the given world.
+	 * Search for the first playable group associated to the given world.
+	 * @param InWorld	Play world to search the playable group of.
+	 * @param bInFallbackToGlobalSearch If the direct link from world to the playable group is not found, fallback to a global search of the whole system.
 	 */
-	static UAvaPlayableGroup* FindPlayableGroupForWorld(const UWorld* InWorld);
+	static UAvaPlayableGroup* FindPlayableGroupForWorld(const UWorld* InWorld, bool bInFallbackToGlobalSearch = true);
 	
 protected:
 	bool IsVisibilityConstrained(const UAvaPlayable* InPlayable) const;
diff --git a/Engine/Plugins/Experimental/Avalanche/Source/AvalancheMedia/Public/Playable/AvaPlayableLibrary.h b/Engine/Plugins/Experimental/Avalanche/Source/AvalancheMedia/Public/Playable/AvaPlayableLibrary.h
--- a/Engine/Plugins/Experimental/Avalanche/Source/AvalancheMedia/Public/Playable/AvaPlayableLibrary.h
+++ b/Engine/Plugins/Experimental/Avalanche/Source/AvalancheMedia/Public/Playable/AvaPlayableLibrary.h
@@ -20,7 +20,12 @@ public:
 	/** Returns the transition this playable is part of. */
 	static UAvaPlayableTransition* GetPlayableTransition(const UAvaPlayable* InPlayable);
 
-	/** Injects the remote control values from current transition for the current playable. */
+	/**
+	 * Injects the remote control values from current transition for the current playable.
+	 * @remark This does nothing if there is no current transition the current playable is part of or if
+	 *		the current level is not managed by a playable.
+	 * @return true if the values have been injected, false otherwise.
+	 */
 	UFUNCTION(BlueprintCallable, Category = "Motion Design|Playable", meta=(WorldContext = "InWorldContextObject"))
-	static AVALANCHEMEDIA_API void UpdateRemoteControlValues(const UObject* InWorldContextObject);
+	static AVALANCHEMEDIA_API bool UpdateRemoteControlValues(const UObject* InWorldContextObject);
 };
