Motion Design RC: Customization for RC Controller Id to do a dropdown to show existing Controllers if the id is a prop handle of an object within a ULevel/World.

diff --git a/Engine/Plugins/Experimental/Avalanche/Source/AvalancheRemoteControl/Private/AvaRemoteControlUtils.cpp b/Engine/Plugins/Experimental/Avalanche/Source/AvalancheRemoteControl/Private/AvaRemoteControlUtils.cpp
--- a/Engine/Plugins/Experimental/Avalanche/Source/AvalancheRemoteControl/Private/AvaRemoteControlUtils.cpp
+++ b/Engine/Plugins/Experimental/Avalanche/Source/AvalancheRemoteControl/Private/AvaRemoteControlUtils.cpp
@@ -1,6 +1,7 @@
 // Copyright Epic Games, Inc. All Rights Reserved.
 
 #include "AvaRemoteControlUtils.h"
+#include "Engine/Level.h"
 #include "IRemoteControlModule.h"
 #include "RemoteControlPreset.h"
 #include "Subsystems/RemoteControlComponentsSubsystem.h"
@@ -108,4 +109,26 @@ void FAvaRemoteControlUtils::UnregisterRemoteControlPreset(URemoteControlPreset*
 	}
 	
 	RemoteControlModule.UnregisterEmbeddedPreset(InRemoteControlPreset);
-}
\ No newline at end of file
+}
+
+URemoteControlPreset* FAvaRemoteControlUtils::FindEmbeddedPresetInLevel(ULevel* InLevel)
+{
+	if (!InLevel)
+	{
+		return nullptr;
+	}
+
+	TArray<TWeakObjectPtr<URemoteControlPreset>> EmbeddedPresets;
+	IRemoteControlModule::Get().GetEmbeddedPresets(EmbeddedPresets);
+
+	for (const TWeakObjectPtr<URemoteControlPreset>& EmbeddedPresetWeak : EmbeddedPresets)
+	{
+		URemoteControlPreset* EmbeddedPreset = EmbeddedPresetWeak.Get();
+		if (EmbeddedPreset && EmbeddedPreset->GetTypedOuter<ULevel>() == InLevel)
+		{
+			return EmbeddedPreset;
+		}
+	}
+
+	return nullptr;
+}
diff --git a/Engine/Plugins/Experimental/Avalanche/Source/AvalancheRemoteControl/Public/AvaRCControllerId.h b/Engine/Plugins/Experimental/Avalanche/Source/AvalancheRemoteControl/Public/AvaRCControllerId.h
--- a/Engine/Plugins/Experimental/Avalanche/Source/AvalancheRemoteControl/Public/AvaRCControllerId.h
+++ b/Engine/Plugins/Experimental/Avalanche/Source/AvalancheRemoteControl/Public/AvaRCControllerId.h
@@ -9,7 +9,7 @@ class URCVirtualPropertyBase;
 class URemoteControlPreset;
 
 /** Struct describing data to identify/find a Controller in a given preset */
-USTRUCT(BlueprintType)
+USTRUCT(BlueprintType, DisplayName="Motion Design RC Controller Id")
 struct FAvaRCControllerId
 {
 	GENERATED_BODY()
diff --git a/Engine/Plugins/Experimental/Avalanche/Source/AvalancheRemoteControl/Public/AvaRemoteControlUtils.h b/Engine/Plugins/Experimental/Avalanche/Source/AvalancheRemoteControl/Public/AvaRemoteControlUtils.h
--- a/Engine/Plugins/Experimental/Avalanche/Source/AvalancheRemoteControl/Public/AvaRemoteControlUtils.h
+++ b/Engine/Plugins/Experimental/Avalanche/Source/AvalancheRemoteControl/Public/AvaRemoteControlUtils.h
@@ -4,6 +4,7 @@
 
 #include "HAL/Platform.h"
 
+class ULevel;
 class URemoteControlPreset;
 
 struct FAvaRemoteControlUtils
@@ -22,4 +23,12 @@ struct FAvaRemoteControlUtils
 	 *	Unregister the given RemoteControlPreset from the RemoteControl module.
 	 */
 	AVALANCHEREMOTECONTROL_API static void UnregisterRemoteControlPreset(URemoteControlPreset* InRemoteControlPreset);
+
+	/**
+	 * Iterates the active Embedded Presets and finds the one that is within the given Level
+	 * Note: this will not find presets that have not been registered, even if outered to the given Level.
+	 * @param InLevel the level containing the preset
+	 * @return the registered preset if found, null otherwise
+	 */
+	AVALANCHEREMOTECONTROL_API static URemoteControlPreset* FindEmbeddedPresetInLevel(ULevel* InLevel);
 };
diff --git a/Engine/Plugins/Experimental/Avalanche/Source/AvalancheRemoteControlEditor/AvalancheRemoteControlEditor.Build.cs b/Engine/Plugins/Experimental/Avalanche/Source/AvalancheRemoteControlEditor/AvalancheRemoteControlEditor.Build.cs
--- a/Engine/Plugins/Experimental/Avalanche/Source/AvalancheRemoteControlEditor/AvalancheRemoteControlEditor.Build.cs
+++ b/Engine/Plugins/Experimental/Avalanche/Source/AvalancheRemoteControlEditor/AvalancheRemoteControlEditor.Build.cs
@@ -17,6 +17,8 @@ public class AvalancheRemoteControlEditor : ModuleRules
                 "Core",
                 "CoreUObject",
                 "Engine",
+                "InputCore",
+                "RemoteControl",
                 "RemoteControlComponents",
                 "RemoteControlComponentsEditor",
                 "RemoteControlUI",
diff --git a/Engine/Plugins/Experimental/Avalanche/Source/AvalancheRemoteControlEditor/Private/Customizations/AvaRCControllerIdCustomization.cpp b/Engine/Plugins/Experimental/Avalanche/Source/AvalancheRemoteControlEditor/Private/Customizations/AvaRCControllerIdCustomization.cpp
--- a/Engine/Plugins/Experimental/Avalanche/Source/AvalancheRemoteControlEditor/Private/Customizations/AvaRCControllerIdCustomization.cpp
+++ b/Engine/Plugins/Experimental/Avalanche/Source/AvalancheRemoteControlEditor/Private/Customizations/AvaRCControllerIdCustomization.cpp
@@ -3,7 +3,7 @@
 #include "AvaRCControllerIdCustomization.h"
 #include "AvaRCControllerId.h"
 #include "DetailWidgetRow.h"
-#include "Widgets/Layout/SBox.h"
+#include "Widgets/SAvaRCControllerPicker.h"
 #include "Widgets/SBoxPanel.h"
 
 TSharedRef<IPropertyTypeCustomization> FAvaRCControllerIdCustomization::MakeInstance()
@@ -27,13 +27,9 @@ void FAvaRCControllerIdCustomization::CustomizeHeader(TSharedRef<IPropertyHandle
 			+ SHorizontalBox::Slot()
 			.VAlign(VAlign_Center)
 			.HAlign(HAlign_Fill)
+			.FillWidth(1.f)
 			[
-				NamePropertyHandle->CreatePropertyValueWidget(/*bDisplayDefaultPropertyButtons*/false)
-			]
-			+ SHorizontalBox::Slot()
-			.AutoWidth()
-			[
-				InPropertyHandle->CreateDefaultPropertyButtonWidgets()
+				SNew(SAvaRCControllerPicker, InPropertyHandle)
 			]
 		];
 }
diff --git a/Engine/Plugins/Experimental/Avalanche/Source/AvalancheRemoteControlEditor/Private/Customizations/Widgets/SAvaRCControllerPicker.cpp b/Engine/Plugins/Experimental/Avalanche/Source/AvalancheRemoteControlEditor/Private/Customizations/Widgets/SAvaRCControllerPicker.cpp
new file mode 100644
--- /dev/null
+++ b/Engine/Plugins/Experimental/Avalanche/Source/AvalancheRemoteControlEditor/Private/Customizations/Widgets/SAvaRCControllerPicker.cpp
@@ -0,0 +1,188 @@
+// Copyright Epic Games, Inc. All Rights Reserved.
+
+#include "SAvaRCControllerPicker.h"
+#include "AvaRCControllerId.h"
+#include "AvaRemoteControlUtils.h"
+#include "DetailLayoutBuilder.h"
+#include "Engine/Level.h"
+#include "RCVirtualProperty.h"
+#include "RemoteControlPreset.h"
+#include "Widgets/Input/SComboBox.h"
+#include "Widgets/Input/SEditableTextBox.h"
+#include "Widgets/Text/STextBlock.h"
+
+#define LOCTEXT_NAMESPACE "SAvaRCControllerPicker"
+
+void SAvaRCControllerPicker::Construct(const FArguments& InArgs, const TSharedRef<IPropertyHandle>& InControllerIdHandle)
+{
+	NameHandle = InControllerIdHandle->GetChildHandle(GET_MEMBER_NAME_CHECKED(FAvaRCControllerId, Name));
+
+	// Figure out the Level to use
+	{
+		TArray<UObject*> OuterObjects;
+		InControllerIdHandle->GetOuterObjects(OuterObjects);
+
+		for (UObject* OuterObject : OuterObjects)
+		{
+			if (ULevel* Level = OuterObject->GetTypedOuter<ULevel>())
+			{
+				LevelWeak = Level;
+				break;
+			}
+		}
+	}
+
+	// Editable Text box to deal with the Name Handle
+	TSharedRef<SWidget> Widget = SNew(SEditableTextBox)
+		.Text(this, &SAvaRCControllerPicker::GetControllerName)
+		.OnTextCommitted(this, &SAvaRCControllerPicker::OnControllerNameCommitted)
+		.Font(IDetailLayoutBuilder::GetDetailFont());
+
+	// if the found level is valid, then use a combo box to display the options found for that level instead
+	if (LevelWeak.IsValid())
+	{
+		RefreshOptions();
+
+		Widget = SAssignNew(ComboBox, SComboBox<TSharedPtr<FAvaRCControllerPickerOption>>)
+			.OptionsSource(&Options)
+			.InitiallySelectedItem(SelectedOption)
+			.OnGenerateWidget(this, &SAvaRCControllerPicker::GenerateOptionWidget)
+			.OnComboBoxOpening(this, &SAvaRCControllerPicker::RefreshOptions)
+			.OnSelectionChanged(this, &SAvaRCControllerPicker::OnOptionSelectionChanged)
+			[
+				Widget
+			];
+	}
+
+	ChildSlot
+	[
+		Widget
+	];
+}
+
+FText SAvaRCControllerPicker::GetControllerName() const
+{
+	FName CurrentControllerName;
+	const FPropertyAccess::Result Result = NameHandle->GetValue(CurrentControllerName);
+
+	if (Result == FPropertyAccess::Fail)
+	{
+		return FText::GetEmpty();
+	}
+
+	if (Result == FPropertyAccess::MultipleValues)
+	{
+		return LOCTEXT("MultipleValues", "Multiple Values");
+	}
+
+	return FText::FromName(CurrentControllerName);
+}
+
+void SAvaRCControllerPicker::SetControllerName(FName InControllerName)
+{
+	if (!NameHandle.IsValid())
+	{
+		return;
+	}
+
+	FName CurrentControllerName;
+	bool bSucceeded = NameHandle->GetValue(CurrentControllerName) == FPropertyAccess::Success;
+
+	// Only skip if successfully retrieved the current value and it matches the new value.
+	// Else, attempt to set the new value in the handle
+	if (bSucceeded && CurrentControllerName == InControllerName)
+	{
+		return;
+	}
+
+	NameHandle->SetValue(InControllerName, EPropertyValueSetFlags::DefaultFlags);
+}
+
+void SAvaRCControllerPicker::OnControllerNameCommitted(const FText& InText, ETextCommit::Type InCommitType)
+{
+	SetControllerName(*InText.ToString());
+}
+
+void SAvaRCControllerPicker::RefreshOptions()
+{
+	Options.Reset();
+	SelectedOption.Reset();
+
+	if (!NameHandle.IsValid())
+	{
+		return;	
+	}
+
+	ULevel* Level = LevelWeak.Get();
+	if (!Level)
+	{
+		return;
+	}
+
+	URemoteControlPreset* Preset = FAvaRemoteControlUtils::FindEmbeddedPresetInLevel(Level);
+	if (!Preset)
+	{
+		return;
+	}
+
+	FName CurrentControllerName;
+	NameHandle->GetValue(CurrentControllerName);
+
+	TArray<URCVirtualPropertyBase*> Controllers = Preset->GetControllers();
+	Options.Reserve(Controllers.Num());
+
+	TSet<FName> ProcessedNames;
+	ProcessedNames.Reserve(Controllers.Num());
+
+	for (URCVirtualPropertyBase* Controller : Controllers)
+	{
+		if (!Controller)
+		{
+			continue;
+		}
+
+		// Since FAvaRCControllerPickerOption only handles names directly, avoid adding the same name entry twice
+		bool bAlreadyInSet;
+		ProcessedNames.Add(Controller->DisplayName, &bAlreadyInSet);
+		if (bAlreadyInSet)
+		{
+			continue;
+		}
+
+		TSharedRef<FAvaRCControllerPickerOption> Item = MakeShared<FAvaRCControllerPickerOption>();
+		Item->ControllerName = Controller->DisplayName;
+
+		// Set the first Item that matches the current Controller Name as the selected item
+		if (!SelectedOption.IsValid() && Item->ControllerName == CurrentControllerName)
+		{
+			SelectedOption = Item;
+		}
+
+		Options.Emplace(MoveTemp(Item));
+	}
+
+	// Update the Selected Item to the selected option (can be null if there's no controller found matching the current name stored)
+	if (ComboBox.IsValid())
+	{
+		ComboBox->RefreshOptions();
+		ComboBox->SetSelectedItem(SelectedOption);
+	}
+}
+
+void SAvaRCControllerPicker::OnOptionSelectionChanged(TSharedPtr<FAvaRCControllerPickerOption> InSelectedOption, ESelectInfo::Type InSelectInfo)
+{
+	if (InSelectedOption.IsValid())
+	{
+		SelectedOption = InSelectedOption;
+		SetControllerName(SelectedOption->ControllerName);
+	}
+}
+
+TSharedRef<SWidget> SAvaRCControllerPicker::GenerateOptionWidget(TSharedPtr<FAvaRCControllerPickerOption> InOption)
+{
+	return SNew(STextBlock)
+		.Text(FText::FromName(InOption->ControllerName))
+		.Font(IDetailLayoutBuilder::GetDetailFont());
+}
+
+#undef LOCTEXT_NAMESPACE
diff --git a/Engine/Plugins/Experimental/Avalanche/Source/AvalancheRemoteControlEditor/Private/Customizations/Widgets/SAvaRCControllerPicker.h b/Engine/Plugins/Experimental/Avalanche/Source/AvalancheRemoteControlEditor/Private/Customizations/Widgets/SAvaRCControllerPicker.h
new file mode 100644
--- /dev/null
+++ b/Engine/Plugins/Experimental/Avalanche/Source/AvalancheRemoteControlEditor/Private/Customizations/Widgets/SAvaRCControllerPicker.h
@@ -0,0 +1,47 @@
+// Copyright Epic Games, Inc. All Rights Reserved.
+
+#pragma once
+
+#include "Containers/Array.h"
+#include "Widgets/SCompoundWidget.h"
+
+class IPropertyHandle;
+class ULevel;
+template<typename OptionType> class SComboBox;
+
+struct FAvaRCControllerPickerOption
+{
+	FName ControllerName;
+};
+
+class SAvaRCControllerPicker : public SCompoundWidget
+{
+public:
+	SLATE_BEGIN_ARGS(SAvaRCControllerPicker) {}
+	SLATE_END_ARGS()
+
+	void Construct(const FArguments& InArgs, const TSharedRef<IPropertyHandle>& InControllerIdHandle);
+
+private:
+	FText GetControllerName() const;
+
+	void SetControllerName(FName InControllerName);
+
+	void OnControllerNameCommitted(const FText& InText, ETextCommit::Type InCommitType);
+
+	void RefreshOptions();
+
+	void OnOptionSelectionChanged(TSharedPtr<FAvaRCControllerPickerOption> InSelectedOption, ESelectInfo::Type InSelectInfo);
+
+	TSharedRef<SWidget> GenerateOptionWidget(TSharedPtr<FAvaRCControllerPickerOption> InOption);
+
+	TSharedPtr<SComboBox<TSharedPtr<FAvaRCControllerPickerOption>>> ComboBox;
+
+	TWeakObjectPtr<ULevel> LevelWeak;
+
+	TSharedPtr<IPropertyHandle> NameHandle;
+
+	TArray<TSharedPtr<FAvaRCControllerPickerOption>> Options;
+
+	TSharedPtr<FAvaRCControllerPickerOption> SelectedOption;
+};
