MotionDesign : Modifiers
- When root visibility actor is hidden, do not handle visibility state on children

diff --git a/Engine/Plugins/Experimental/Avalanche/Source/AvalancheModifiers/Private/Modifiers/AvaVisibilityModifier.cpp b/Engine/Plugins/Experimental/Avalanche/Source/AvalancheModifiers/Private/Modifiers/AvaVisibilityModifier.cpp
--- a/Engine/Plugins/Experimental/Avalanche/Source/AvalancheModifiers/Private/Modifiers/AvaVisibilityModifier.cpp
+++ b/Engine/Plugins/Experimental/Avalanche/Source/AvalancheModifiers/Private/Modifiers/AvaVisibilityModifier.cpp
@@ -12,24 +12,30 @@
 void UAvaVisibilityModifier::Apply()
 {
 	AActor* const ModifyActor = GetModifiedActor();
-	
+
 	// Early exit if the modify actor is NOT being isolated. The outliner will manage the visibility for the actor and it's children.
 	if (FAvaModifiersActorUtils::IsActorNotIsolated(ModifyActor))
 	{
 		Next();
 		return;
 	}
-	
+
 	const FAvaSceneTreeUpdateModifierExtension* SceneExtension = GetExtension<FAvaSceneTreeUpdateModifierExtension>();
 	if (!SceneExtension)
 	{
 		Fail(LOCTEXT("InvalidSceneExtension", "Scene extension could not be found"));
 		return;
 	}
-	
+
 	const TArray<TWeakObjectPtr<AActor>> AttachedActors = SceneExtension->GetDirectChildrenActor(ModifyActor);
 	UAvaVisibilityModifierShared* VisibilityShared = GetShared<UAvaVisibilityModifierShared>(true);
 
+#if WITH_EDITOR
+	const bool bHiddenInEditor = ModifyActor->IsTemporarilyHiddenInEditor(/** IncludeParent */false);
+#else
+	const bool bHiddenInEditor = false;
+#endif
+
 	// Top most modifier in tree has priority over this one if it is hiding current one
 	bool bIsNestedVisibilityModifier = false;
 	if (const UAvaVisibilityModifier* VisibilityModifier = GetFirstModifierAbove(ModifyActor))
@@ -40,7 +46,13 @@ void UAvaVisibilityModifier::Apply()
 			bIsNestedVisibilityModifier = true;
 		}
 	}
-	
+	// We are the top root modifier, if this actor is hidden, do not handle children actors
+	else if (bSkipWhenHidden && (ModifyActor->IsHidden() || bHiddenInEditor))
+	{
+		Next();
+		return;
+	}
+
 	TSet<TWeakObjectPtr<AActor>> NewChildrenActorsWeak;
 	for (int32 ChildIndex = 0; ChildIndex < AttachedActors.Num(); ++ChildIndex)
 	{
@@ -50,7 +62,7 @@ void UAvaVisibilityModifier::Apply()
 		{
 			continue;
 		}
-		
+
 		// No need to handle nested children actor, only direct children, visibility will propagate
 		if (AttachedActor->GetAttachParentActor() != ModifyActor)
 		{
@@ -58,7 +70,7 @@ void UAvaVisibilityModifier::Apply()
 		}
 
 		bool bHideActor = false;
-		
+
 		if (!bTreatAsRange)
 		{
 			bHideActor = ChildIndex == Index;
@@ -67,10 +79,10 @@ void UAvaVisibilityModifier::Apply()
 		{
 			bHideActor = ChildIndex <= Index;
 		}
-		
+
 		bHideActor = bInvertVisibility ? bHideActor : !bHideActor;
 		DirectChildrenActorsWeak.Add(AttachedActor, bHideActor);
-		
+
 		TArray<AActor*> AttachedChildActors {AttachedActor};
 		AttachedActor->GetAttachedActors(AttachedChildActors, false, true);
 		for (AActor* AttachedChildActor : AttachedChildActors)
@@ -86,12 +98,12 @@ void UAvaVisibilityModifier::Apply()
 					continue;
 				}
 			}
-			
+
 			if (!bIsNestedVisibilityModifier)
 			{
 				VisibilityShared->SetActorVisibility(this, AttachedChildActor, bHideActor, false);
 			}
-			
+
 			NewChildrenActorsWeak.Add(AttachedActor);
 		}
 	}
@@ -114,7 +126,7 @@ void UAvaVisibilityModifier::PostEditChangeProperty(FPropertyChangedEvent& Prope
 	static const FName IndexPropertyName = GET_MEMBER_NAME_CHECKED(UAvaVisibilityModifier, Index);
 	static const FName InvertVisibilityPropertyName = GET_MEMBER_NAME_CHECKED(UAvaVisibilityModifier, bInvertVisibility);
 	static const FName TreatAsRangePropertyName = GET_MEMBER_NAME_CHECKED(UAvaVisibilityModifier, bTreatAsRange);
-	
+
 	if (PropertyName == IndexPropertyName
 		|| PropertyName == TreatAsRangePropertyName
 		|| PropertyName == InvertVisibilityPropertyName)
@@ -124,7 +136,7 @@ void UAvaVisibilityModifier::PostEditChangeProperty(FPropertyChangedEvent& Prope
 }
 #endif // WITH_EDITOR
 
-void UAvaVisibilityModifier::SetTreatAsRange(const bool bInTreatAsRange)
+void UAvaVisibilityModifier::SetTreatAsRange(bool bInTreatAsRange)
 {
 	if (bTreatAsRange == bInTreatAsRange)
 	{
@@ -135,6 +147,17 @@ void UAvaVisibilityModifier::SetTreatAsRange(const bool bInTreatAsRange)
 	MarkModifierDirty();
 }
 
+void UAvaVisibilityModifier::SetSkipWhenHidden(bool bInSkip)
+{
+	if (bSkipWhenHidden == bInSkip)
+	{
+		return;
+	}
+
+	bSkipWhenHidden = bInSkip;
+	MarkModifierDirty();
+}
+
 void UAvaVisibilityModifier::OnModifierCDOSetup(FActorModifierCoreMetadata& InMetadata)
 {
 	Super::OnModifierCDOSetup(InMetadata);
@@ -157,14 +180,18 @@ void UAvaVisibilityModifier::OnRenderStateUpdated(AActor* InActor, UActorCompone
 	{
 		return;
 	}
-	
+
 	MarkModifierDirty();
 }
 
-void UAvaVisibilityModifier::SetInvertVisibility(const bool bInInvertVisibility)
+void UAvaVisibilityModifier::SetInvertVisibility(bool bInInvertVisibility)
 {
-	bInvertVisibility = bInInvertVisibility;
+	if (bInvertVisibility == bInInvertVisibility)
+	{
+		return;
+	}
 
+	bInvertVisibility = bInInvertVisibility;
 	MarkModifierDirty();
 }
 
@@ -202,7 +229,7 @@ bool UAvaVisibilityModifier::IsChildActorHidden(AActor* InActor) const
 
 		return *DirectChildrenActorsWeak.Find(InActor);
 	}
-	
+
 	return false;
 }
 
@@ -210,12 +237,12 @@ UAvaVisibilityModifier* UAvaVisibilityModifier::GetFirstModifierAbove(AActor* In
 {
 	UAvaVisibilityModifierShared* VisibilityShared = GetShared<UAvaVisibilityModifierShared>(false);
 	UAvaVisibilityModifier* FirstModifierAbove = nullptr;
-	
+
 	if (!InActor || !VisibilityShared)
 	{
 		return FirstModifierAbove;
 	}
-	
+
 	if (FAvaVisibilitySharedActorState* ActorState = VisibilityShared->FindActorState(InActor))
 	{
 		for (const FAvaVisibilitySharedModifierState& ModifierState : ActorState->ModifierStates)
@@ -261,12 +288,12 @@ AActor* UAvaVisibilityModifier::GetDirectChildren(AActor* InParentActor, AActor*
 	{
 		return nullptr;
 	}
-	
+
 	if (InChildActor && InChildActor->GetAttachParentActor() == InParentActor)
 	{
 		return InChildActor;
 	}
-	
+
 	return GetDirectChildren(InParentActor, InChildActor->GetAttachParentActor());
 }
 
diff --git a/Engine/Plugins/Experimental/Avalanche/Source/AvalancheModifiers/Public/Modifiers/AvaVisibilityModifier.h b/Engine/Plugins/Experimental/Avalanche/Source/AvalancheModifiers/Public/Modifiers/AvaVisibilityModifier.h
--- a/Engine/Plugins/Experimental/Avalanche/Source/AvalancheModifiers/Public/Modifiers/AvaVisibilityModifier.h
+++ b/Engine/Plugins/Experimental/Avalanche/Source/AvalancheModifiers/Public/Modifiers/AvaVisibilityModifier.h
@@ -18,7 +18,7 @@ class UAvaVisibilityModifier : public UAvaArrangeBaseModifier
 public:
 	/** Sets the child index range to hide instead of showing. */
 	UFUNCTION(BlueprintCallable, Category="Motion Design|Modifiers|Visibility")
-	AVALANCHEMODIFIERS_API void SetInvertVisibility(const bool bInInvertVisibility);
+	AVALANCHEMODIFIERS_API void SetInvertVisibility(bool bInInvertVisibility);
 
 	/** Returns true if hiding the child index range instead of showing. */
 	UFUNCTION(BlueprintPure, Category="Motion Design|Modifiers|Visibility")
@@ -37,7 +37,7 @@ public:
 	}
 
 	UFUNCTION(BlueprintCallable, Category="Motion Design|Modifiers|Visibility")
-	AVALANCHEMODIFIERS_API void SetTreatAsRange(const bool bInTreatAsRange);
+	AVALANCHEMODIFIERS_API void SetTreatAsRange(bool bInTreatAsRange);
 
 	UFUNCTION(BlueprintPure, Category="Motion Design|Modifiers|Visibility")
 	bool GetTreatAsRange() const
@@ -45,6 +45,15 @@ public:
 		return bTreatAsRange;
 	}
 
+	UFUNCTION(BlueprintCallable, Category="Motion Design|Modifiers|Visibility")
+	AVALANCHEMODIFIERS_API void SetSkipWhenHidden(bool bInSkip);
+
+	UFUNCTION(BlueprintPure, Category="Motion Design|Modifiers|Visibility")
+	bool GetSkipWhenHidden() const
+	{
+		return bSkipWhenHidden;
+	}
+
 protected:
 	//~ Begin UObject
 #if WITH_EDITOR
@@ -83,6 +92,11 @@ protected:
 	UPROPERTY(EditInstanceOnly, Setter="SetInvertVisibility", Getter="GetInvertVisibility", Category="Visibility", meta=(AllowPrivateAccess="true"))
 	bool bInvertVisibility = false;
 
+	/** If true, when the top actor having this modifier is hidden, children actor visibility state will not be handled */
+	UPROPERTY(EditInstanceOnly, Setter="SetSkipWhenHidden", Getter="GetSkipWhenHidden", Category="Visibility", meta=(AllowPrivateAccess="true"))
+	bool bSkipWhenHidden = true;
+
+private:
 	/** Visibility defined for direct children of this actor, used by other modifiers of this type to check */
 	UPROPERTY(Transient, DuplicateTransient, TextExportTransient)
 	TMap<TWeakObjectPtr<AActor>, bool> DirectChildrenActorsWeak;
