[Motion Design] Broadcast Blueprint Library: adding GetChannelName and GetChannelStatus.

Enables blueprint logic to determine if a level is playing as part of a rundown or not, and it the channel is live or not.

diff --git a/Engine/Plugins/Experimental/Avalanche/Source/AvalancheMedia/Private/Broadcast/AvaBroadcastLibrary.cpp b/Engine/Plugins/Experimental/Avalanche/Source/AvalancheMedia/Private/Broadcast/AvaBroadcastLibrary.cpp
--- a/Engine/Plugins/Experimental/Avalanche/Source/AvalancheMedia/Private/Broadcast/AvaBroadcastLibrary.cpp
+++ b/Engine/Plugins/Experimental/Avalanche/Source/AvalancheMedia/Private/Broadcast/AvaBroadcastLibrary.cpp
@@ -48,4 +48,36 @@ FVector2D UAvaBroadcastLibrary::GetChannelViewportSize(const UObject* InWorldCon
 	}
 	
 	return FAvaBroadcastOutputChannel::GetDefaultMediaOutputSize(EAvaBroadcastChannelType::Program);
+}
+
+FName UAvaBroadcastLibrary::GetChannelName(const UObject* InWorldContextObject)
+{
+	if (const UAvaPlayable* Playable = UAvaPlayableLibrary::GetPlayable(InWorldContextObject))
+	{
+		if (const UAvaPlayableGroup* PlayableGroup = Playable->GetPlayableGroup())
+		{
+			return PlayableGroup->GetChannelName();
+		}
+	}
+	return NAME_None;
+}
+
+EAvaBroadcastChannelState UAvaBroadcastLibrary::GetChannelStatus(const FName InChannelName)
+{
+	const FAvaBroadcastOutputChannel& Channel = UAvaBroadcast::Get().GetCurrentProfile().GetChannel(InChannelName);
+	if (Channel.IsValidChannel())
+	{
+		return Channel.GetState();
+	}
+	return EAvaBroadcastChannelState::Offline;;
+}
+
+EAvaBroadcastChannelType UAvaBroadcastLibrary::GetChannelType(const FName InChannelName)
+{
+	const FAvaBroadcastOutputChannel& Channel = UAvaBroadcast::Get().GetCurrentProfile().GetChannel(InChannelName);
+	if (Channel.IsValidChannel())
+	{
+		return Channel.GetChannelType();
+	}
+	return EAvaBroadcastChannelType::Program;
 }
\ No newline at end of file
diff --git a/Engine/Plugins/Experimental/Avalanche/Source/AvalancheMedia/Public/Broadcast/AvaBroadcastLibrary.h b/Engine/Plugins/Experimental/Avalanche/Source/AvalancheMedia/Public/Broadcast/AvaBroadcastLibrary.h
--- a/Engine/Plugins/Experimental/Avalanche/Source/AvalancheMedia/Public/Broadcast/AvaBroadcastLibrary.h
+++ b/Engine/Plugins/Experimental/Avalanche/Source/AvalancheMedia/Public/Broadcast/AvaBroadcastLibrary.h
@@ -2,6 +2,7 @@
 
 #pragma once
 
+#include "AvaMediaDefines.h"
 #include "Kismet/BlueprintFunctionLibrary.h"
 #include "AvaBroadcastLibrary.generated.h"
 
@@ -14,4 +15,16 @@ public:
 	/** Returns the current channel's viewport size. */
 	UFUNCTION(BlueprintCallable, Category = "Motion Design|broadcast", meta=(WorldContext = "InWorldContextObject"))
 	static AVALANCHEMEDIA_API FVector2D GetChannelViewportSize(const UObject* InWorldContextObject);
+
+	/** Returns the current channel's name. Will return "None" if not running as part of the broadcast framework. */
+	UFUNCTION(BlueprintCallable, Category = "Motion Design|broadcast", meta=(WorldContext = "InWorldContextObject"))
+	static AVALANCHEMEDIA_API FName GetChannelName(const UObject* InWorldContextObject);
+
+	/** Returns the channel status. */
+	UFUNCTION(BlueprintCallable, Category = "Motion Design|broadcast")
+	static AVALANCHEMEDIA_API EAvaBroadcastChannelState GetChannelStatus(const FName InChannelName);
+
+	/** Returns the channel type. */
+	UFUNCTION(BlueprintCallable, Category = "Motion Design|broadcast")
+	static AVALANCHEMEDIA_API EAvaBroadcastChannelType GetChannelType(const FName InChannelName);
 };
