[Motion Design] Rundown RemovePage has a direct delegate to refresh the detail panel (if the deleted page was displayed).

This should fix the issue that the detail panel keeps a strong reference on the managed instances temporarily preventing assets from being unloaded.

In order to make this implementation easier, the FOnAvaRundownPageListChanged events where consolidated in one event for all the rundown's page lists. There is less code to deal with (easier to maintain).

Prior to this change, the page detail panel would eventually update after a repaint of the page list widget (SAvaRundownPageList::OnPageSelected) that triggers a selection change (EPageEvent::SelectionChanged).

diff --git a/Engine/Plugins/Experimental/Avalanche/Source/AvalancheMedia/Internal/Rundown/AvaRundownMessages.h b/Engine/Plugins/Experimental/Avalanche/Source/AvalancheMedia/Internal/Rundown/AvaRundownMessages.h
--- a/Engine/Plugins/Experimental/Avalanche/Source/AvalancheMedia/Internal/Rundown/AvaRundownMessages.h
+++ b/Engine/Plugins/Experimental/Avalanche/Source/AvalancheMedia/Internal/Rundown/AvaRundownMessages.h
@@ -599,6 +599,12 @@ struct FAvaRundownPageListChanged : public FAvaRundownMsgBase
 	UPROPERTY()
 	FString Rundown;
 
+	UPROPERTY()
+	EAvaRundownPageListType ListType = EAvaRundownPageListType::Instance;
+
+	UPROPERTY()
+	int32 SubListIndex = INDEX_NONE;
+	
 	/** See EAvaPageListChange flags. */
 	UPROPERTY()
 	uint8 ChangeType = 0;
diff --git a/Engine/Plugins/Experimental/Avalanche/Source/AvalancheMedia/Private/Rundown/AvaRundown.cpp b/Engine/Plugins/Experimental/Avalanche/Source/AvalancheMedia/Private/Rundown/AvaRundown.cpp
--- a/Engine/Plugins/Experimental/Avalanche/Source/AvalancheMedia/Private/Rundown/AvaRundown.cpp
+++ b/Engine/Plugins/Experimental/Avalanche/Source/AvalancheMedia/Private/Rundown/AvaRundown.cpp
@@ -24,10 +24,10 @@ DEFINE_LOG_CATEGORY(LogAvaRundown);
 
 #define LOCTEXT_NAMESPACE "AvaRundown"
 
-void FAvaRundownPageCollection::Empty(UAvaRundown* InRundown)
+void FAvaRundownPageCollection::Empty(UAvaRundown* InRundown, const FAvaRundownPageListReference& InPageListReference)
 {
 	TArray<int32> PageIds;
-	if (OnPageListChanged.IsBound() && Pages.IsEmpty() == false)
+	if (InRundown->GetOnPageListChanged().IsBound() && Pages.IsEmpty() == false)
 	{
 		PageIds.Reserve(Pages.Num());
 		for (const FAvaRundownPage& Page : Pages)
@@ -41,7 +41,7 @@ void FAvaRundownPageCollection::Empty(UAvaRundown* InRundown)
 	
 	if (PageIds.IsEmpty() == false)
 	{
-		OnPageListChanged.Broadcast({InRundown, EAvaRundownPageListChange::RemovedPages, PageIds});
+		InRundown->GetOnPageListChanged().Broadcast({InRundown, InPageListReference, EAvaRundownPageListChange::RemovedPages, PageIds});
 	}
 }
 
@@ -140,8 +140,8 @@ void UAvaRundown::PostEditUndo()
 	
 	//Force Refresh for any Undo
 	RefreshPageIndices();
-	GetOnTemplatePageListChanged().Broadcast({this, EAvaRundownPageListChange::All, {}});
-	GetOnInstancedPageListChanged().Broadcast({this, EAvaRundownPageListChange::All, {}});
+	GetOnPageListChanged().Broadcast({this, TemplatePageList, EAvaRundownPageListChange::All, {}});
+	GetOnPageListChanged().Broadcast({this, InstancePageList, EAvaRundownPageListChange::All, {}});
 }
 #endif
 
@@ -165,8 +165,8 @@ bool UAvaRundown::Empty()
 	}
 	
 	SubLists.Empty();
-	InstancedPages.Empty(this);
-	TemplatePages.Empty(this);
+	InstancedPages.Empty(this, InstancePageList);
+	TemplatePages.Empty(this, TemplatePageList);
 	
 	return true;
 }
@@ -189,7 +189,7 @@ int32 UAvaRundown::AddTemplateInternal(const FAvaRundownPageIdGeneratorParams& I
 	
 	RefreshPageIndices();
 
-	GetOnTemplatePageListChanged().Broadcast({this, EAvaRundownPageListChange::AddedPages, {TemplateId}});
+	GetOnPageListChanged().Broadcast({this, TemplatePageList, EAvaRundownPageListChange::AddedPages, {TemplateId}});
 
 	return TemplateId;
 }
@@ -234,7 +234,7 @@ TArray<int32> UAvaRundown::AddTemplates(const TArray<FAvaRundownPage>& InSourceT
 
 	if (!OutTemplateIds.IsEmpty())
 	{
-		GetOnTemplatePageListChanged().Broadcast({this, EAvaRundownPageListChange::AddedPages, OutTemplateIds});
+		GetOnPageListChanged().Broadcast({this, TemplatePageList, EAvaRundownPageListChange::AddedPages, OutTemplateIds});
 	}
 
 	return OutTemplateIds;
@@ -268,7 +268,7 @@ TArray<int32> UAvaRundown::AddPagesFromTemplates(const TArray<int32>& InTemplate
 	}
 
 	RefreshPageIndices();
-	GetOnInstancedPageListChanged().Broadcast({this, EAvaRundownPageListChange::AddedPages, OutPageIds});
+	GetOnPageListChanged().Broadcast({this, InstancePageList, EAvaRundownPageListChange::AddedPages, OutPageIds});
 
 	return OutPageIds;
 }
@@ -280,7 +280,7 @@ int32 UAvaRundown::AddPageFromTemplate(int32 InTemplateId, const FAvaRundownPage
 	if (NewId != FAvaRundownPage::InvalidPageId)
 	{
 		RefreshPageIndices();
-		GetOnInstancedPageListChanged().Broadcast({this, EAvaRundownPageListChange::AddedPages, {NewId}});
+		GetOnPageListChanged().Broadcast({this, InstancePageList, EAvaRundownPageListChange::AddedPages, {NewId}});
 	}
 
 	return NewId;
@@ -329,7 +329,7 @@ bool UAvaRundown::ChangePageOrder(const FAvaRundownPageListReference& InPageList
 		Collection.PageIndices.Empty();
 		RefreshPageIndices();
 
-		Collection.OnPageListChanged.Broadcast({this, EAvaRundownPageListChange::ReorderedPageView, {}});
+		GetOnPageListChanged().Broadcast({this, InPageListReference, EAvaRundownPageListChange::ReorderedPageView, {}});
 
 		return true;
 	}
@@ -356,7 +356,7 @@ bool UAvaRundown::ChangePageOrder(const FAvaRundownPageListReference& InPageList
 		}
 
 		SubList.PageIds = NewIndices;
-		SubList.OnPageListChanged.Broadcast({this, EAvaRundownPageListChange::ReorderedPageView, {}});
+		GetOnPageListChanged().Broadcast({this, InPageListReference, EAvaRundownPageListChange::ReorderedPageView, {}});
 
 		return true;
 	}
@@ -477,14 +477,15 @@ int32 UAvaRundown::RemovePages(const TArray<int32>& InPageIds)
 
 	if (bRemovedTemplate)
 	{
-		GetOnTemplatePageListChanged().Broadcast({this, EAvaRundownPageListChange::RemovedPages, {InPageIds}});
+		GetOnPageListChanged().Broadcast({this, TemplatePageList, EAvaRundownPageListChange::RemovedPages, {InPageIds}});
 	}
 
 	if (bRemovedInstanced)
 	{
-		GetOnInstancedPageListChanged().Broadcast({this, EAvaRundownPageListChange::RemovedPages, {InPageIds}});
+		GetOnPageListChanged().Broadcast({this, InstancePageList, EAvaRundownPageListChange::RemovedPages, {InPageIds}});
 	}
 
+	int32 SubListIndex = 0;
 	for (FAvaRundownSubList& SubList : SubLists)
 	{
 		bool bFoundInstance = false;
@@ -500,8 +501,9 @@ int32 UAvaRundown::RemovePages(const TArray<int32>& InPageIds)
 
 		if (bFoundInstance)
 		{
-			SubList.OnPageListChanged.Broadcast({this, EAvaRundownPageListChange::RemovedPages, {InPageIds}});
+			GetOnPageListChanged().Broadcast({this, CreateSubListReference(SubListIndex), EAvaRundownPageListChange::RemovedPages, {InPageIds}});
 		}
+		++SubListIndex;
 	}
 
 	return RemovedCount;
@@ -577,26 +579,28 @@ bool UAvaRundown::RenumberPageId(int32 InPageId, int32 InNewPageId)
 			}
 		}
 
-		GetOnTemplatePageListChanged().Broadcast({this, EAvaRundownPageListChange::RenumberedPageId, {InNewPageId}});
+		GetOnPageListChanged().Broadcast({this, TemplatePageList, EAvaRundownPageListChange::RenumberedPageId, {InNewPageId}});
 
 		if (bFoundInstanceOfTemplate)
 		{
-			GetOnInstancedPageListChanged().Broadcast({this, EAvaRundownPageListChange::RenumberedPageId, {InNewPageId}});
+			GetOnPageListChanged().Broadcast({this, InstancePageList, EAvaRundownPageListChange::RenumberedPageId, {InNewPageId}});
 		}
 	}
 	else if (InstanceIdx)
 	{
-		GetOnInstancedPageListChanged().Broadcast({this, EAvaRundownPageListChange::RenumberedPageId, {InNewPageId}});
+		GetOnPageListChanged().Broadcast({this, InstancePageList, EAvaRundownPageListChange::RenumberedPageId, {InNewPageId}});
 
+		int32 SubListIndex = 0;
 		for (FAvaRundownSubList& SubList : SubLists)
 		{
-			int32 Index = SubList.PageIds.Find(InPageId);
+			const int32 Index = SubList.PageIds.Find(InPageId);
 
 			if (Index != INDEX_NONE)
 			{
 				SubList.PageIds[Index] = InNewPageId;
-				SubList.OnPageListChanged.Broadcast({this, EAvaRundownPageListChange::RenumberedPageId, {InNewPageId}});
+				GetOnPageListChanged().Broadcast({this, CreateSubListReference(SubListIndex), EAvaRundownPageListChange::RenumberedPageId, {InNewPageId}});
 			}
+			++SubListIndex;
 		}
 	}
 	else
@@ -1401,7 +1405,7 @@ bool UAvaRundown::AddPageToSubList(int32 InSubListIndex, int32 InPageId, const F
 			SubLists[InSubListIndex].PageIds.Add(InPageId);
 		}
 		
-		SubLists[InSubListIndex].OnPageListChanged.Broadcast({this, EAvaRundownPageListChange::AddedPages, {InPageId}});
+		GetOnPageListChanged().Broadcast({this, CreateSubListReference(InSubListIndex), EAvaRundownPageListChange::AddedPages, {InPageId}});
 		return true;
 	}
 
@@ -1426,7 +1430,7 @@ bool UAvaRundown::AddPagesToSubList(int32 InSubListIndex, const TArray<int32>& I
 
 		if (bAddedPage)
 		{
-			SubLists[InSubListIndex].OnPageListChanged.Broadcast({this, EAvaRundownPageListChange::AddedPages, InPages});
+			GetOnPageListChanged().Broadcast({this, CreateSubListReference(InSubListIndex), EAvaRundownPageListChange::AddedPages, InPages});
 			return true;
 		}
 	}
@@ -1445,7 +1449,7 @@ int32 UAvaRundown::RemovePagesFromSubList(int32 InSubListIndex, const TArray<int
 
 		if (Removed > 0)
 		{
-			SubLists[InSubListIndex].OnPageListChanged.Broadcast({this, EAvaRundownPageListChange::RemovedPages, InPages});
+			GetOnPageListChanged().Broadcast({this, CreateSubListReference(InSubListIndex), EAvaRundownPageListChange::RemovedPages, InPages});
 		}
 
 		return Removed;
diff --git a/Engine/Plugins/Experimental/Avalanche/Source/AvalancheMedia/Private/Rundown/AvaRundownServer.cpp b/Engine/Plugins/Experimental/Avalanche/Source/AvalancheMedia/Private/Rundown/AvaRundownServer.cpp
--- a/Engine/Plugins/Experimental/Avalanche/Source/AvalancheMedia/Private/Rundown/AvaRundownServer.cpp
+++ b/Engine/Plugins/Experimental/Avalanche/Source/AvalancheMedia/Private/Rundown/AvaRundownServer.cpp
@@ -481,6 +481,8 @@ void FAvaRundownServer::OnPageListChanged(const FAvaRundownPageListChangeParams&
 {
 	FAvaRundownPageListChanged* ReplyMessage = FMessageEndpoint::MakeMessage<FAvaRundownPageListChanged>();
 	ReplyMessage->Rundown = FSoftObjectPath(InParams.Rundown).ToString();
+	ReplyMessage->ListType = InParams.PageListReference.Type;
+	ReplyMessage->SubListIndex = InParams.PageListReference.SubListIndex;
 	ReplyMessage->ChangeType = static_cast<uint8>(InParams.ChangeType);
 	ReplyMessage->AffectedPages = InParams.AffectedPages;
 	
@@ -2404,8 +2406,7 @@ FAvaRundownServer::FRundownEntry::FRundownEntry(const TSharedPtr<FAvaRundownServ
 	{
 		const TSharedRef<FAvaRundownServer> RundownServerRef = InRundownServer.ToSharedRef();
 		Rundown->GetOnPagesChanged().AddSP(RundownServerRef, &FAvaRundownServer::OnPagesChanged);
-		Rundown->GetOnInstancedPageListChanged().AddSP(RundownServerRef, &FAvaRundownServer::OnPageListChanged);
-		Rundown->GetOnTemplatePageListChanged().AddSP(RundownServerRef, &FAvaRundownServer::OnPageListChanged);
+		Rundown->GetOnPageListChanged().AddSP(RundownServerRef, &FAvaRundownServer::OnPageListChanged);
 		Rundown->GetOnCanClosePlaybackContext().AddSP(RundownServerRef, &FAvaRundownServer::OnCanClosePlaybackContext);
 	}
 }
@@ -2415,8 +2416,7 @@ FAvaRundownServer::FRundownEntry::~FRundownEntry()
 	if (Rundown)
 	{
 		Rundown->GetOnPagesChanged().RemoveAll(RundownServerRaw);
-		Rundown->GetOnInstancedPageListChanged().RemoveAll(RundownServerRaw);
-		Rundown->GetOnTemplatePageListChanged().RemoveAll(RundownServerRaw);
+		Rundown->GetOnPageListChanged().RemoveAll(RundownServerRaw);
 		Rundown->GetOnCanClosePlaybackContext().RemoveAll(RundownServerRaw);
 	}
 }
diff --git a/Engine/Plugins/Experimental/Avalanche/Source/AvalancheMedia/Public/Rundown/AvaRundown.h b/Engine/Plugins/Experimental/Avalanche/Source/AvalancheMedia/Public/Rundown/AvaRundown.h
--- a/Engine/Plugins/Experimental/Avalanche/Source/AvalancheMedia/Public/Rundown/AvaRundown.h
+++ b/Engine/Plugins/Experimental/Avalanche/Source/AvalancheMedia/Public/Rundown/AvaRundown.h
@@ -89,6 +89,7 @@ struct FAvaRundownPageIdGeneratorParams
 struct FAvaRundownPageListChangeParams
 {
 	UAvaRundown* Rundown;
+	FAvaRundownPageListReference PageListReference;
 	EAvaRundownPageListChange ChangeType;
 	TArray<int32> AffectedPages;
 };
@@ -106,9 +107,7 @@ struct FAvaRundownPageCollection
 
 	/** Cache mapping the Page Id to the index where the Page with such Page Id is at */
 	TMap<int32, int32> PageIndices;
-
-	FOnAvaRundownPageListChanged OnPageListChanged;
-
+	
 	int32 GetPageIndex(const int32 InPageId) const
 	{
 		if (InPageId != FAvaRundownPage::InvalidPageId)
@@ -119,7 +118,7 @@ struct FAvaRundownPageCollection
 		return INDEX_NONE;
 	}
 
-	AVALANCHEMEDIA_API void Empty(UAvaRundown* InRundown);
+	AVALANCHEMEDIA_API void Empty(UAvaRundown* InRundown, const FAvaRundownPageListReference& InPageListReference);
 
 	/** Complete refresh of the page indices. */
 	void RefreshPageIndices()
@@ -151,8 +150,6 @@ struct FAvaRundownSubList
 
 	UPROPERTY()
 	FText Name;
-
-	FOnAvaRundownPageListChanged OnPageListChanged;
 };
 
 namespace UE::AvaRundown
@@ -418,8 +415,7 @@ public:
 	/** Gets the page following the page with the given page id. Using current active page list. */
 	FAvaRundownPage& GetNextPage(int32 InPageId)  { return GetNextPage(InPageId, ActivePageList);}
 	
-	FOnAvaRundownPageListChanged& GetOnTemplatePageListChanged() { return TemplatePages.OnPageListChanged; }
-	FOnAvaRundownPageListChanged& GetOnInstancedPageListChanged() { return InstancedPages.OnPageListChanged; }
+	FOnAvaRundownPageListChanged& GetOnPageListChanged() { return OnPageListChanged; }
 	FOnAvaRundownPagesChanged& GetOnPagesChanged() { return OnPagesChanged; }
 
 	/** Delegate called to determine if the playback context can be closed. */
@@ -728,6 +724,7 @@ public:
 
 protected:
 	FOnAvaRundownPagesChanged OnPagesChanged;
+	FOnAvaRundownPageListChanged OnPageListChanged;
 	FOnActiveListChanged OnActiveListChanged;
 	mutable FOnCanClosePlaybackContext OnCanClosePlaybackContext;
 };
diff --git a/Engine/Plugins/Experimental/Avalanche/Source/AvalancheMediaEditor/Private/Rundown/DetailsView/RemoteControl/Controllers/SAvaRundownRCControllerPanel.cpp b/Engine/Plugins/Experimental/Avalanche/Source/AvalancheMediaEditor/Private/Rundown/DetailsView/RemoteControl/Controllers/SAvaRundownRCControllerPanel.cpp
--- a/Engine/Plugins/Experimental/Avalanche/Source/AvalancheMediaEditor/Private/Rundown/DetailsView/RemoteControl/Controllers/SAvaRundownRCControllerPanel.cpp
+++ b/Engine/Plugins/Experimental/Avalanche/Source/AvalancheMediaEditor/Private/Rundown/DetailsView/RemoteControl/Controllers/SAvaRundownRCControllerPanel.cpp
@@ -80,11 +80,6 @@ bool SAvaRundownRCControllerPanel::HasRemoteControlPreset(const URemoteControlPr
 	return false;
 }
 
-void SAvaRundownRCControllerPanel::OnPageSelectionChanged(const TArray<int32>& InSelectedPageIds)
-{
-	Refresh(InSelectedPageIds);
-}
-
 void SAvaRundownRCControllerPanel::UpdatePropertyRowGenerators(int32 InNumGenerators)
 {
 	if (PropertyRowGenerators.Num() != InNumGenerators)
diff --git a/Engine/Plugins/Experimental/Avalanche/Source/AvalancheMediaEditor/Private/Rundown/DetailsView/RemoteControl/Controllers/SAvaRundownRCControllerPanel.h b/Engine/Plugins/Experimental/Avalanche/Source/AvalancheMediaEditor/Private/Rundown/DetailsView/RemoteControl/Controllers/SAvaRundownRCControllerPanel.h
--- a/Engine/Plugins/Experimental/Avalanche/Source/AvalancheMediaEditor/Private/Rundown/DetailsView/RemoteControl/Controllers/SAvaRundownRCControllerPanel.h
+++ b/Engine/Plugins/Experimental/Avalanche/Source/AvalancheMediaEditor/Private/Rundown/DetailsView/RemoteControl/Controllers/SAvaRundownRCControllerPanel.h
@@ -46,8 +46,6 @@ public:
 	void Construct(const FArguments& InArgs, const TSharedPtr<FAvaRundownEditor>& InRundownEditor);
 	
 	bool HasRemoteControlPreset(const URemoteControlPreset* InPreset) const;
-
-	void OnPageSelectionChanged(const TArray<int32>& InSelectedPageIds);
 	
 	void Refresh(const TArray<int32>& InSelectedPageIds);
 	
diff --git a/Engine/Plugins/Experimental/Avalanche/Source/AvalancheMediaEditor/Private/Rundown/DetailsView/SAvaRundownPageDetails.cpp b/Engine/Plugins/Experimental/Avalanche/Source/AvalancheMediaEditor/Private/Rundown/DetailsView/SAvaRundownPageDetails.cpp
--- a/Engine/Plugins/Experimental/Avalanche/Source/AvalancheMediaEditor/Private/Rundown/DetailsView/SAvaRundownPageDetails.cpp
+++ b/Engine/Plugins/Experimental/Avalanche/Source/AvalancheMediaEditor/Private/Rundown/DetailsView/SAvaRundownPageDetails.cpp
@@ -46,7 +46,8 @@ void SAvaRundownPageDetails::Construct(const FArguments& InArgs, const TSharedPt
 	UAvaRundown* const Rundown = InRundownEditor->GetRundown();
 	if (IsValid(Rundown))
 	{
-		Rundown->GetOnPagesChanged().AddSP(this, &SAvaRundownPageDetails::OnRundownPagesChanged);
+		Rundown->GetOnPagesChanged().AddSP(this, &SAvaRundownPageDetails::OnPagesChanged);
+		Rundown->GetOnPageListChanged().AddSP(this, &SAvaRundownPageDetails::OnPageListChanged);
 	}
 
 	TSharedRef<SHorizontalBox> AnimationHeader = SNew(SHorizontalBox);
@@ -190,7 +191,7 @@ void SAvaRundownPageDetails::Construct(const FArguments& InArgs, const TSharedPt
 		]
 	];
 
-	OnPageSelectionChanged({});
+	ActivePageId = FAvaRundownPage::InvalidPageId;
 }
 
 SAvaRundownPageDetails::~SAvaRundownPageDetails()
@@ -202,6 +203,7 @@ SAvaRundownPageDetails::~SAvaRundownPageDetails()
 		if (IsValid(Rundown))
 		{
 			Rundown->GetOnPagesChanged().RemoveAll(this);
+			Rundown->GetOnPageListChanged().RemoveAll(this);
 		}
 	}
 	if (IAvaMediaModule::IsModuleLoaded() && IAvaMediaModule::Get().IsManagedInstanceCacheAvailable())
@@ -212,19 +214,23 @@ SAvaRundownPageDetails::~SAvaRundownPageDetails()
 
 void SAvaRundownPageDetails::OnPageEvent(const TArray<int32>& InSelectedPageIds, UE::AvaRundown::EPageEvent InPageEvent)
 {
+	bool bRefreshPanels = false;
 	if (InPageEvent == UE::AvaRundown::EPageEvent::SelectionChanged || InPageEvent == UE::AvaRundown::EPageEvent::ReimportRequest)
 	{
-		OnPageSelectionChanged(InSelectedPageIds);
+		const int32 PreviousActivePageId = ActivePageId;
+		ActivePageId = InSelectedPageIds.IsEmpty() ? FAvaRundownPage::InvalidPageId : InSelectedPageIds[0];
+
+		// Only refresh the panels if the page id changed or if a reimport request (forced refresh).
+		bRefreshPanels = ActivePageId != PreviousActivePageId || InPageEvent == UE::AvaRundown::EPageEvent::ReimportRequest;
+	}
+
+	if (bRefreshPanels)
+	{
 		RemoteControlProps->Refresh(InSelectedPageIds);
 		RCControllerPanel->Refresh(InSelectedPageIds);
 	}
 }
 
-void SAvaRundownPageDetails::OnPageSelectionChanged(const TArray<int32>& InSelectedPageIds)
-{
-	ActivePageId = InSelectedPageIds.IsEmpty() ? FAvaRundownPage::InvalidPageId : InSelectedPageIds[0];
-}
-
 void SAvaRundownPageDetails::OnManagedInstanceCacheEntryInvalidated(const FSoftObjectPath& InAssetPath)
 {
 	if (!bRefreshSelectedPageQueued)
@@ -478,7 +484,7 @@ FReply SAvaRundownPageDetails::DuplicateSelectedPage()
 	return FReply::Handled();
 }
 
-void SAvaRundownPageDetails::OnRundownPagesChanged(const UAvaRundown* InRundown, const FAvaRundownPage& InPage, const EAvaRundownPageChanges InChanges)
+void SAvaRundownPageDetails::OnPagesChanged(const UAvaRundown* InRundown, const FAvaRundownPage& InPage, const EAvaRundownPageChanges InChanges)
 {
 	// Refreshing the page while the mouse is captured will result in losing the capture
 	// and ending any drag event that is actively changing the value.
@@ -489,4 +495,13 @@ void SAvaRundownPageDetails::OnRundownPagesChanged(const UAvaRundown* InRundown,
 	}
 }
 
+void SAvaRundownPageDetails::OnPageListChanged(const FAvaRundownPageListChangeParams& InParams)
+{
+	// If the current page is removed, fire off a selection changed immediately.
+	if (InParams.AffectedPages.Contains(ActivePageId) && EnumHasAnyFlags(InParams.ChangeType, EAvaRundownPageListChange::RemovedPages))
+	{
+		OnPageEvent({}, UE::AvaRundown::EPageEvent::SelectionChanged);
+	}
+}
+
 #undef LOCTEXT_NAMESPACE
diff --git a/Engine/Plugins/Experimental/Avalanche/Source/AvalancheMediaEditor/Private/Rundown/DetailsView/SAvaRundownPageDetails.h b/Engine/Plugins/Experimental/Avalanche/Source/AvalancheMediaEditor/Private/Rundown/DetailsView/SAvaRundownPageDetails.h
--- a/Engine/Plugins/Experimental/Avalanche/Source/AvalancheMediaEditor/Private/Rundown/DetailsView/SAvaRundownPageDetails.h
+++ b/Engine/Plugins/Experimental/Avalanche/Source/AvalancheMediaEditor/Private/Rundown/DetailsView/SAvaRundownPageDetails.h
@@ -14,10 +14,11 @@ class FText;
 class SAvaRundownPageRemoteControlProps;
 class SAvaRundownRCControllerPanel;
 class UAvaRundown;
+enum class EAvaRundownPageChanges : uint8;
 struct FAvaRundownPage;
-struct FSoftObjectPath;
+struct FAvaRundownPageListChangeParams;
 struct FSlateBrush;
-enum class EAvaRundownPageChanges : uint8;
+struct FSoftObjectPath;
 
 class SAvaRundownPageDetails : public SCompoundWidget
 {
@@ -31,7 +32,6 @@ public:
 	virtual ~SAvaRundownPageDetails() override;
 
 	void OnPageEvent(const TArray<int32>& InSelectedPageIds, UE::AvaRundown::EPageEvent InPageEvent);
-	void OnPageSelectionChanged(const TArray<int32>& InSelectedPageIds);
 	void OnManagedInstanceCacheEntryInvalidated(const FSoftObjectPath& InAssetPath);
 
 protected:
@@ -59,8 +59,9 @@ protected:
 
 	FReply DuplicateSelectedPage();
 
-	void OnRundownPagesChanged(const UAvaRundown* InRundown, const FAvaRundownPage& InPage, const EAvaRundownPageChanges InChanges);
-
+	void OnPagesChanged(const UAvaRundown* InRundown, const FAvaRundownPage& InPage, const EAvaRundownPageChanges InChanges);
+	void OnPageListChanged(const FAvaRundownPageListChangeParams& InParams);
+	
 private:
 	TWeakPtr<FAvaRundownEditor> RundownEditorWeak;
 
diff --git a/Engine/Plugins/Experimental/Avalanche/Source/AvalancheMediaEditor/Private/Rundown/Pages/Slate/SAvaRundownInstancedPageList.cpp b/Engine/Plugins/Experimental/Avalanche/Source/AvalancheMediaEditor/Private/Rundown/Pages/Slate/SAvaRundownInstancedPageList.cpp
--- a/Engine/Plugins/Experimental/Avalanche/Source/AvalancheMediaEditor/Private/Rundown/Pages/Slate/SAvaRundownInstancedPageList.cpp
+++ b/Engine/Plugins/Experimental/Avalanche/Source/AvalancheMediaEditor/Private/Rundown/Pages/Slate/SAvaRundownInstancedPageList.cpp
@@ -2,7 +2,6 @@
 
 #include "SAvaRundownInstancedPageList.h"
 
-#include "IAvaMediaEditorModule.h"
 #include "Misc/PathViews.h"
 #include "Rundown/AvaRundown.h"
 #include "Rundown/AvaRundownCommands.h"
@@ -21,8 +20,6 @@
 #include "Rundown/Pages/Columns/AvaRundownPageStatusColumn.h"
 #include "Rundown/Pages/Columns/AvaRundownPageTemplateNameColumn.h"
 #include "Rundown/Pages/PageViews/AvaRundownInstancedPageViewImpl.h"
-#include "Rundown/TabFactories/AvaRundownInstancedPageListTabFactory.h"
-#include "Rundown/TabFactories/AvaRundownSubListDocumentTabFactory.h"
 #include "SAvaRundownPageList.h"
 #include "ScopedTransaction.h"
 #include "Styling/StyleColors.h"
@@ -48,16 +45,7 @@ void SAvaRundownInstancedPageList::Construct(const FArguments& InArgs, TSharedPt
 
 	check(InPageListReference.Type == EAvaRundownPageListType::Instance || Rundown->IsValidSubList(PageListReference));
 
-	if (PageListReference.Type == EAvaRundownPageListType::Instance)
-	{
-		TabId = FAvaRundownInstancedPageListTabFactory::TabID;
-		Rundown->GetOnInstancedPageListChanged().AddSP(this, &SAvaRundownInstancedPageList::OnInstancedPageListChanged);
-	}
-	else
-	{
-		FAvaRundownSubListDocumentTabFactory::GetTabId(InPageListReference.SubListIndex);
-		Rundown->GetSubList(InPageListReference.SubListIndex).OnPageListChanged.AddSP(this, &SAvaRundownInstancedPageList::OnInstancedPageListChanged);
-	}
+	Rundown->GetOnPageListChanged().AddSP(this, &SAvaRundownInstancedPageList::OnPageListChanged);
 
 	if (Rundown->IsValidSubList(PageListReference))
 	{
@@ -119,14 +107,7 @@ SAvaRundownInstancedPageList::~SAvaRundownInstancedPageList()
 {
 	if (UAvaRundown* const Rundown = GetValidRundown())
 	{
-		if (PageListReference.Type == EAvaRundownPageListType::Instance)
-		{
-			Rundown->GetOnInstancedPageListChanged().RemoveAll(this);
-		}
-		else if (Rundown->IsValidSubList(PageListReference))
-		{
-			Rundown->GetSubList(PageListReference.SubListIndex).OnPageListChanged.RemoveAll(this);
-		}
+		Rundown->GetOnPageListChanged().RemoveAll(this);
 	}
 }
 
@@ -1157,8 +1138,13 @@ bool SAvaRundownInstancedPageList::CanPlayNextPage() const
 	return IsPageIdValid(PageIdToTakeNext);
 }
 
-void SAvaRundownInstancedPageList::OnInstancedPageListChanged(const FAvaRundownPageListChangeParams& InParams)
+void SAvaRundownInstancedPageList::OnPageListChanged(const FAvaRundownPageListChangeParams& InParams)
 {
+	if (PageListReference != InParams.PageListReference)
+	{
+		return;
+	}
+	
 	if (const TSharedPtr<FAvaRundownEditor> RundownEditor = RundownEditorWeak.Pin())
 	{
 		RundownEditor->RefreshInstancedVisibility();
@@ -1231,7 +1217,7 @@ void SAvaRundownInstancedPageList::OnPageViewNameCommitted(const FText& InNewTex
 		if (Rundown->IsValidSubList(PageListReference))
 		{
 			Rundown->GetSubList(PageListReference.SubListIndex).Name = InNewText;
-			Rundown->GetSubList(PageListReference.SubListIndex).OnPageListChanged.Broadcast({Rundown, EAvaRundownPageListChange::RenamedPageView, {}});
+			Rundown->GetOnPageListChanged().Broadcast({Rundown, PageListReference, EAvaRundownPageListChange::RenamedPageView, {}});
 
 			if (TSharedPtr<SDockTab> MyTab = MyTabWeak.Pin())
 			{
diff --git a/Engine/Plugins/Experimental/Avalanche/Source/AvalancheMediaEditor/Private/Rundown/Pages/Slate/SAvaRundownInstancedPageList.h b/Engine/Plugins/Experimental/Avalanche/Source/AvalancheMediaEditor/Private/Rundown/Pages/Slate/SAvaRundownInstancedPageList.h
--- a/Engine/Plugins/Experimental/Avalanche/Source/AvalancheMediaEditor/Private/Rundown/Pages/Slate/SAvaRundownInstancedPageList.h
+++ b/Engine/Plugins/Experimental/Avalanche/Source/AvalancheMediaEditor/Private/Rundown/Pages/Slate/SAvaRundownInstancedPageList.h
@@ -39,8 +39,6 @@ public:
 
 	void OnTabActivated(TSharedRef<SDockTab> InDockTab, ETabActivationCause InActivationCause);
 
-	FName GetTabId() const { return TabId; }
-
 	TSharedPtr<SDockTab> GetMyTab() const { return MyTabWeak.Pin(); }
 	void SetMyTab(TSharedRef<SDockTab> InTab) { MyTabWeak = InTab; }
 
@@ -71,7 +69,7 @@ public:
 	}
 	
 private:
-	void OnInstancedPageListChanged(const FAvaRundownPageListChangeParams& InParams);
+	void OnPageListChanged(const FAvaRundownPageListChangeParams& InParams);
 
 	TArray<int32> FilterPlayingPages(FFilterPageFunctionRef InFilterPageFunction) const;
 	TArray<int32> FilterSelectedOrPlayingPages(FFilterPageFunctionRef InFilterPageFunction, const bool bInAllowFallback) const;
@@ -84,7 +82,6 @@ private:
 	int32 GetPageIdToTakeNext() const;
 
 protected:
-	FName TabId;
 	TWeakPtr<SDockTab> MyTabWeak;
 
 	FReply MakeActive();
diff --git a/Engine/Plugins/Experimental/Avalanche/Source/AvalancheMediaEditor/Private/Rundown/Pages/Slate/SAvaRundownTemplatePageList.cpp b/Engine/Plugins/Experimental/Avalanche/Source/AvalancheMediaEditor/Private/Rundown/Pages/Slate/SAvaRundownTemplatePageList.cpp
--- a/Engine/Plugins/Experimental/Avalanche/Source/AvalancheMediaEditor/Private/Rundown/Pages/Slate/SAvaRundownTemplatePageList.cpp
+++ b/Engine/Plugins/Experimental/Avalanche/Source/AvalancheMediaEditor/Private/Rundown/Pages/Slate/SAvaRundownTemplatePageList.cpp
@@ -39,7 +39,7 @@ void SAvaRundownTemplatePageList::Construct(const FArguments& InArgs, TSharedPtr
 	UAvaRundown* const Rundown = InRundownEditor->GetRundown();
 	check(Rundown);
 
-	Rundown->GetOnTemplatePageListChanged().AddSP(this, &SAvaRundownTemplatePageList::OnTemplatePageListChanged);
+	Rundown->GetOnPageListChanged().AddSP(this, &SAvaRundownTemplatePageList::OnPageListChanged);
 
 	Refresh();
 }
@@ -48,7 +48,7 @@ SAvaRundownTemplatePageList::~SAvaRundownTemplatePageList()
 {
 	if (UAvaRundown* const Rundown = GetValidRundown())
 	{
-		Rundown->GetOnTemplatePageListChanged().RemoveAll(this);
+		Rundown->GetOnPageListChanged().RemoveAll(this);
 	}
 }
 
@@ -498,8 +498,13 @@ TArray<int32> SAvaRundownTemplatePageList::AddPastedPages(const TArray<FAvaRundo
 	return {};
 }
 
-void SAvaRundownTemplatePageList::OnTemplatePageListChanged(const FAvaRundownPageListChangeParams& InParams)
+void SAvaRundownTemplatePageList::OnPageListChanged(const FAvaRundownPageListChangeParams& InParams)
 {
+	if (PageListReference != InParams.PageListReference)
+	{
+		return;
+	}
+
 	if (const TSharedPtr<FAvaRundownEditor> RundownEditor = RundownEditorWeak.Pin())
 	{
 		RundownEditor->RefreshTemplateVisibility();
diff --git a/Engine/Plugins/Experimental/Avalanche/Source/AvalancheMediaEditor/Private/Rundown/Pages/Slate/SAvaRundownTemplatePageList.h b/Engine/Plugins/Experimental/Avalanche/Source/AvalancheMediaEditor/Private/Rundown/Pages/Slate/SAvaRundownTemplatePageList.h
--- a/Engine/Plugins/Experimental/Avalanche/Source/AvalancheMediaEditor/Private/Rundown/Pages/Slate/SAvaRundownTemplatePageList.h
+++ b/Engine/Plugins/Experimental/Avalanche/Source/AvalancheMediaEditor/Private/Rundown/Pages/Slate/SAvaRundownTemplatePageList.h
@@ -51,5 +51,5 @@ protected:
 	bool CanResetPagesToDefaults() const;
 
 private:
-	void OnTemplatePageListChanged(const FAvaRundownPageListChangeParams& InParams);
+	void OnPageListChanged(const FAvaRundownPageListChangeParams& InParams);
 };
