Motion Design: Fix issue with execution status + order of execution causing inplace transition to have issues with waiting for other scens in layer to finish.

diff --git a/Engine/Plugins/Experimental/Avalanche/Source/AvalancheTransition/Private/Behavior/AvaTransitionBehaviorInstance.cpp b/Engine/Plugins/Experimental/Avalanche/Source/AvalancheTransition/Private/Behavior/AvaTransitionBehaviorInstance.cpp
--- a/Engine/Plugins/Experimental/Avalanche/Source/AvalancheTransition/Private/Behavior/AvaTransitionBehaviorInstance.cpp
+++ b/Engine/Plugins/Experimental/Avalanche/Source/AvalancheTransition/Private/Behavior/AvaTransitionBehaviorInstance.cpp
@@ -76,18 +76,15 @@ void FAvaTransitionBehaviorInstance::SetTransitionType(EAvaTransitionType InTran
 
 bool FAvaTransitionBehaviorInstance::Setup()
 {
-	RunStatus = EStateTreeRunStatus::Unset;
+	// Mark Run Status as running on setup so that on start all the nodes are on this status,
+	// even if their state tree hasn't started yet
+	RunStatus = EStateTreeRunStatus::Running;
 	TOptional<FAvaTransitionExecutionContext> Context = UpdateContext();
 	return Context.IsSet();
 }
 
 void FAvaTransitionBehaviorInstance::Start()
 {
-	if (RunStatus != EStateTreeRunStatus::Unset)
-	{
-		return;
-	}
-
 	// If this Instance is not Enabled for Transition, finish immediately
 	if (!IsEnabled())
 	{
diff --git a/Engine/Plugins/Experimental/Avalanche/Source/AvalancheTransition/Private/Execution/AvaTransitionExecutor.cpp b/Engine/Plugins/Experimental/Avalanche/Source/AvalancheTransition/Private/Execution/AvaTransitionExecutor.cpp
--- a/Engine/Plugins/Experimental/Avalanche/Source/AvalancheTransition/Private/Execution/AvaTransitionExecutor.cpp
+++ b/Engine/Plugins/Experimental/Avalanche/Source/AvalancheTransition/Private/Execution/AvaTransitionExecutor.cpp
@@ -13,11 +13,14 @@
 DEFINE_LOG_CATEGORY_STATIC(LogAvaTransitionExecutor, Log, All);
 
 FAvaTransitionExecutor::FAvaTransitionExecutor(FAvaTransitionExecutorBuilder& InBuilder)
-	: Instances(MoveTemp(InBuilder.Instances))
-	, NullInstance(MoveTemp(InBuilder.NullInstance))
+	: NullInstance(MoveTemp(InBuilder.NullInstance))
 	, ContextName(MoveTemp(InBuilder.ContextName))
 	, OnFinished(MoveTemp(InBuilder.OnFinished))
 {
+	// Add Exit then Enter instances to keep a consistent order of execution
+	Instances.Reserve(InBuilder.ExitInstances.Num() + InBuilder.EnterInstances.Num());
+	Instances.Append(MoveTemp(InBuilder.ExitInstances));
+	Instances.Append(MoveTemp(InBuilder.EnterInstances));
 }
 
 FAvaTransitionExecutor::~FAvaTransitionExecutor()
diff --git a/Engine/Plugins/Experimental/Avalanche/Source/AvalancheTransition/Private/Execution/AvaTransitionExecutorBuilder.cpp b/Engine/Plugins/Experimental/Avalanche/Source/AvalancheTransition/Private/Execution/AvaTransitionExecutorBuilder.cpp
--- a/Engine/Plugins/Experimental/Avalanche/Source/AvalancheTransition/Private/Execution/AvaTransitionExecutorBuilder.cpp
+++ b/Engine/Plugins/Experimental/Avalanche/Source/AvalancheTransition/Private/Execution/AvaTransitionExecutorBuilder.cpp
@@ -22,14 +22,14 @@ FAvaTransitionExecutorBuilder& FAvaTransitionExecutorBuilder::SetContextName(con
 FAvaTransitionExecutorBuilder& FAvaTransitionExecutorBuilder::AddEnterInstance(FAvaTransitionBehaviorInstance& InInstance)
 {
 	InInstance.SetTransitionType(EAvaTransitionType::In);
-	Instances.Add(InInstance);
+	EnterInstances.Add(InInstance);
 	return *this;
 }
 
 FAvaTransitionExecutorBuilder& FAvaTransitionExecutorBuilder::AddExitInstance(FAvaTransitionBehaviorInstance& InInstance)
 {
 	InInstance.SetTransitionType(EAvaTransitionType::Out);
-	Instances.Add(InInstance);
+	ExitInstances.Add(InInstance);
 	return *this;
 }
 
diff --git a/Engine/Plugins/Experimental/Avalanche/Source/AvalancheTransition/Public/Execution/AvaTransitionExecutorBuilder.h b/Engine/Plugins/Experimental/Avalanche/Source/AvalancheTransition/Public/Execution/AvaTransitionExecutorBuilder.h
--- a/Engine/Plugins/Experimental/Avalanche/Source/AvalancheTransition/Public/Execution/AvaTransitionExecutorBuilder.h
+++ b/Engine/Plugins/Experimental/Avalanche/Source/AvalancheTransition/Public/Execution/AvaTransitionExecutorBuilder.h
@@ -38,7 +38,9 @@ public:
 	[[nodiscard]] TSharedRef<IAvaTransitionExecutor> Build(UAvaTransitionSubsystem& InTransitionSubsystem);
 
 private:
-	TArray<FAvaTransitionBehaviorInstance> Instances;
+	TArray<FAvaTransitionBehaviorInstance> EnterInstances;
+
+	TArray<FAvaTransitionBehaviorInstance> ExitInstances;
 
 	FAvaTransitionBehaviorInstance NullInstance;
 
