Motion Design TL: Removed autogenerated states as this is not universally applicable anymore (with the changes to reuse level incoming).

diff --git a/Engine/Plugins/Experimental/Avalanche/Source/AvalancheSequencer/Private/AvaSequencerModule.cpp b/Engine/Plugins/Experimental/Avalanche/Source/AvalancheSequencer/Private/AvaSequencerModule.cpp
--- a/Engine/Plugins/Experimental/Avalanche/Source/AvalancheSequencer/Private/AvaSequencerModule.cpp
+++ b/Engine/Plugins/Experimental/Avalanche/Source/AvalancheSequencer/Private/AvaSequencerModule.cpp
@@ -32,7 +32,6 @@
 #include "SequencerCustomizationManager.h"
 #include "SequencerSettings.h"
 #include "Settings/AvaSequencerSettings.h"
-#include "Transition/AvaSequenceTransitionCompiler.h"
 
 #define LOCTEXT_NAMESPACE "AvaSequencerModule"
 
@@ -57,8 +56,6 @@ void FAvaSequencerModule::StartupModule()
 	RegisterCustomLayouts();
 	RegisterDirectorCompiler();
 
-	FAvaSequenceTransitionCompiler::Get().Register();
-
 	EditorInitializedDelegate = FEditorDelegates::OnEditorInitialized.AddRaw(this, &FAvaSequencerModule::OnEditorInitialized);
 }
 
@@ -83,8 +80,6 @@ void FAvaSequencerModule::ShutdownModule()
 	UnregisterOutlinerItems();
 	UnregisterCustomLayouts();
 
-	FAvaSequenceTransitionCompiler::Get().Unregister();
-
 	FEditorDelegates::OnEditorInitialized.Remove(EditorInitializedDelegate);
 }
 
diff --git a/Engine/Plugins/Experimental/Avalanche/Source/AvalancheSequencer/Private/Transition/AvaSequenceTransitionCompiler.cpp b/Engine/Plugins/Experimental/Avalanche/Source/AvalancheSequencer/Private/Transition/AvaSequenceTransitionCompiler.cpp
deleted file mode 100644
diff --git a/Engine/Plugins/Experimental/Avalanche/Source/AvalancheSequencer/Private/Transition/AvaSequenceTransitionCompiler.h b/Engine/Plugins/Experimental/Avalanche/Source/AvalancheSequencer/Private/Transition/AvaSequenceTransitionCompiler.h
deleted file mode 100644
diff --git a/Engine/Plugins/Experimental/Avalanche/Source/AvalancheTransitionEditor/Private/AvaTransitionEditorModule.cpp b/Engine/Plugins/Experimental/Avalanche/Source/AvalancheTransitionEditor/Private/AvaTransitionEditorModule.cpp
--- a/Engine/Plugins/Experimental/Avalanche/Source/AvalancheTransitionEditor/Private/AvaTransitionEditorModule.cpp
+++ b/Engine/Plugins/Experimental/Avalanche/Source/AvalancheTransitionEditor/Private/AvaTransitionEditorModule.cpp
@@ -59,11 +59,6 @@ IAvaTransitionEditorModule::FOnBuildDefaultTransitionTree& FAvaTransitionEditorM
 	return OnBuildDefaultTransitionTree;
 }
 
-IAvaTransitionEditorModule::FOnCompileTransitionTree& FAvaTransitionEditorModule::GetOnCompileTransitionTree()
-{
-	return OnCompileTransitionTree;
-}
-
 void FAvaTransitionEditorModule::GenerateTransitionTreeOptionsMenu(UToolMenu* InMenu, IAvaTransitionBehavior* InTransitionBehavior)
 {
 	if (!InMenu || !InTransitionBehavior)
diff --git a/Engine/Plugins/Experimental/Avalanche/Source/AvalancheTransitionEditor/Private/AvaTransitionEditorModule.h b/Engine/Plugins/Experimental/Avalanche/Source/AvalancheTransitionEditor/Private/AvaTransitionEditorModule.h
--- a/Engine/Plugins/Experimental/Avalanche/Source/AvalancheTransitionEditor/Private/AvaTransitionEditorModule.h
+++ b/Engine/Plugins/Experimental/Avalanche/Source/AvalancheTransitionEditor/Private/AvaTransitionEditorModule.h
@@ -16,7 +16,6 @@ class FAvaTransitionEditorModule : public IAvaTransitionEditorModule
 
 	//~ Begin IAvaTransitionEditorModule
 	virtual FOnBuildDefaultTransitionTree& GetOnBuildDefaultTransitionTree() override;
-	virtual FOnCompileTransitionTree& GetOnCompileTransitionTree() override;
 	virtual void GenerateTransitionTreeOptionsMenu(UToolMenu* InMenu, IAvaTransitionBehavior* InTransitionBehavior) override;
 	//~ End IAvaTransitionEditorModule
 
@@ -26,7 +25,5 @@ class FAvaTransitionEditorModule : public IAvaTransitionEditorModule
 
 	FOnBuildDefaultTransitionTree OnBuildDefaultTransitionTree;
 
-	FOnCompileTransitionTree OnCompileTransitionTree;
-
 	FDelegateHandle OnPostCompileHandle;
 };
diff --git a/Engine/Plugins/Experimental/Avalanche/Source/AvalancheTransitionEditor/Private/Compiler/AvaTransitionCompiler.cpp b/Engine/Plugins/Experimental/Avalanche/Source/AvalancheTransitionEditor/Private/Compiler/AvaTransitionCompiler.cpp
--- a/Engine/Plugins/Experimental/Avalanche/Source/AvalancheTransitionEditor/Private/Compiler/AvaTransitionCompiler.cpp
+++ b/Engine/Plugins/Experimental/Avalanche/Source/AvalancheTransitionEditor/Private/Compiler/AvaTransitionCompiler.cpp
@@ -75,18 +75,12 @@ bool FAvaTransitionCompiler::Compile(EAvaTransitionEditorMode InCompileMode)
 		// Update Tree's Transition Layer
 		TransitionTree->SetTransitionLayer(EditorData->GetTransitionLayer());
 
-		// Only allow extensible compilation when it's not in Advanced Mode
-		if (InCompileMode != EAvaTransitionEditorMode::Advanced)
-		{
-			// Remove all invalid Sub Trees, if any
-			EditorData->SubTrees.RemoveAll(
-				[](const UStateTreeState* InState)
-				{
-					return !IsValid(InState);
-				});
-
-			IAvaTransitionEditorModule::Get().GetOnCompileTransitionTree().Broadcast(*EditorData);
-		}
+		// Remove all invalid Sub Trees, if any
+		EditorData->SubTrees.RemoveAll(
+			[](const UStateTreeState* InState)
+			{
+				return !IsValid(InState);
+			});
 	}
 
 	UpdateTree();
diff --git a/Engine/Plugins/Experimental/Avalanche/Source/AvalancheTransitionEditor/Private/Serialization/AvaTransitionStateExportData.h b/Engine/Plugins/Experimental/Avalanche/Source/AvalancheTransitionEditor/Private/Serialization/AvaTransitionStateExportData.h
--- a/Engine/Plugins/Experimental/Avalanche/Source/AvalancheTransitionEditor/Private/Serialization/AvaTransitionStateExportData.h
+++ b/Engine/Plugins/Experimental/Avalanche/Source/AvalancheTransitionEditor/Private/Serialization/AvaTransitionStateExportData.h
@@ -3,7 +3,6 @@
 #pragma once
 
 #include "Containers/Array.h"
-#include "State/AvaTransitionAutogeneratedStateContext.h"
 #include "StateTreePropertyBindings.h"
 #include "UObject/Object.h"
 #include "AvaTransitionStateExportData.generated.h"
@@ -16,8 +15,4 @@ class UAvaTransitionStateExportData : public UObject
 public:
 	UPROPERTY()
 	TArray<FStateTreePropertyPathBinding> Bindings;
-
-	/** Map of an autogenerated state id to their context information */
-	UPROPERTY()
-	TMap<FGuid, FAvaTransitionAutogeneratedStateContext> AutogeneratedStates;
 };
diff --git a/Engine/Plugins/Experimental/Avalanche/Source/AvalancheTransitionEditor/Private/Serialization/AvaTransitionStateSerializer.cpp b/Engine/Plugins/Experimental/Avalanche/Source/AvalancheTransitionEditor/Private/Serialization/AvaTransitionStateSerializer.cpp
--- a/Engine/Plugins/Experimental/Avalanche/Source/AvalancheTransitionEditor/Private/Serialization/AvaTransitionStateSerializer.cpp
+++ b/Engine/Plugins/Experimental/Avalanche/Source/AvalancheTransitionEditor/Private/Serialization/AvaTransitionStateSerializer.cpp
@@ -47,12 +47,6 @@ namespace UE::AvaTransitionEditor::Private
 			return;
 		}
 
-		// Copy the autogenerated context for Autogenerated States
-		if (const FAvaTransitionAutogeneratedStateContext* AutogeneratedStateContext = InEditorData.FindAutogeneratedStateContext(InState->ID))
-		{
-			InCopyData.AutogeneratedStates.Add(InState->ID, *AutogeneratedStateContext);
-		}
-
 		InEditorData.VisitStateNodes(*InState,
 			[&InEditorData, &InCopyData](const UStateTreeState*, const FStateTreeBindableStructDesc& InDesc, const FStateTreeDataView)
 			{
@@ -259,12 +253,6 @@ bool FAvaTransitionStateSerializer::ImportText(const FString& InText
 					InEditorData.GetPropertyEditorBindings()->AddPropertyBinding(SourcePath, TargetPath);
 				}
 			}
-
-			// Ensure new states coming from auto-generated states are kept auto-generated (with same context)
-			if (FAvaTransitionAutogeneratedStateContext* AutogeneratedContext = Factory.StateCopyData->AutogeneratedStates.Find(OldId))
-			{
-				InEditorData.SetAutogeneratedState(NewId, *AutogeneratedContext);
-			}
 		}
 	}
 
diff --git a/Engine/Plugins/Experimental/Avalanche/Source/AvalancheTransitionEditor/Private/ViewModels/AvaTransitionEditorViewModel.cpp b/Engine/Plugins/Experimental/Avalanche/Source/AvalancheTransitionEditor/Private/ViewModels/AvaTransitionEditorViewModel.cpp
--- a/Engine/Plugins/Experimental/Avalanche/Source/AvalancheTransitionEditor/Private/ViewModels/AvaTransitionEditorViewModel.cpp
+++ b/Engine/Plugins/Experimental/Avalanche/Source/AvalancheTransitionEditor/Private/ViewModels/AvaTransitionEditorViewModel.cpp
@@ -188,6 +188,8 @@ void FAvaTransitionEditorViewModel::OnInitialize()
 	{
 		ViewModelSharedData->SetReadOnly(Editor->IsReadOnly());
 	}
+
+	UE::StateTree::Delegates::OnPostCompile.AddSP(This, &FAvaTransitionEditorViewModel::OnPostCompile);
 }
 
 void FAvaTransitionEditorViewModel::PostRefresh()
@@ -270,6 +272,15 @@ void FAvaTransitionEditorViewModel::UnbindDelegates()
 	UE::StateTree::Delegates::OnSchemaChanged.RemoveAll(this);
 }
 
+void FAvaTransitionEditorViewModel::OnPostCompile(const UStateTree& InStateTree)
+{
+	if (&InStateTree == GetTransitionTree())
+	{
+		// Full refresh on compile
+		Refresh();
+	}
+}
+
 void FAvaTransitionEditorViewModel::OnIdentifierChanged(const UStateTree& InStateTree)
 {
 	if (&InStateTree == GetTransitionTree())
diff --git a/Engine/Plugins/Experimental/Avalanche/Source/AvalancheTransitionEditor/Private/ViewModels/AvaTransitionEditorViewModel.h b/Engine/Plugins/Experimental/Avalanche/Source/AvalancheTransitionEditor/Private/ViewModels/AvaTransitionEditorViewModel.h
--- a/Engine/Plugins/Experimental/Avalanche/Source/AvalancheTransitionEditor/Private/ViewModels/AvaTransitionEditorViewModel.h
+++ b/Engine/Plugins/Experimental/Avalanche/Source/AvalancheTransitionEditor/Private/ViewModels/AvaTransitionEditorViewModel.h
@@ -100,6 +100,8 @@ private:
 
 	void UnbindDelegates();
 
+	void OnPostCompile(const UStateTree& InStateTree);
+
 	void OnIdentifierChanged(const UStateTree& InStateTree);
 
 	void OnSchemaChanged(const UStateTree& InStateTree);
diff --git a/Engine/Plugins/Experimental/Avalanche/Source/AvalancheTransitionEditor/Private/ViewModels/State/AvaTransitionStateViewModel.cpp b/Engine/Plugins/Experimental/Avalanche/Source/AvalancheTransitionEditor/Private/ViewModels/State/AvaTransitionStateViewModel.cpp
--- a/Engine/Plugins/Experimental/Avalanche/Source/AvalancheTransitionEditor/Private/ViewModels/State/AvaTransitionStateViewModel.cpp
+++ b/Engine/Plugins/Experimental/Avalanche/Source/AvalancheTransitionEditor/Private/ViewModels/State/AvaTransitionStateViewModel.cpp
@@ -206,7 +206,6 @@ FText FAvaTransitionStateViewModel::GetBreakpointTooltip() const
 void FAvaTransitionStateViewModel::UpdateStateDescription(bool bInSetStateName)
 {
 	StateDescription = FText::GetEmpty();
-	AutogeneratedLabel = FText::GetEmpty();
 
 	UStateTreeState* State = GetState();
 	if (!State)
@@ -214,11 +213,6 @@ void FAvaTransitionStateViewModel::UpdateStateDescription(bool bInSetStateName)
 		return;
 	}
 
-	if (const FAvaTransitionAutogeneratedStateContext* AutogeneratedStateContext = FindAutogeneratedStateContext())
-	{
-		AutogeneratedLabel = FText::Format(LOCTEXT("ContextLabelFormat", "Autogenerated by {0}"), AutogeneratedStateContext->GetLabel());
-	}
-
 	// Advanced Mode: If name not none show State Descriptions from the State Name (Raw)
 	if (State->Name != NAME_None && GetSharedData()->GetEditorMode() == EAvaTransitionEditorMode::Advanced)
 	{
@@ -246,20 +240,6 @@ void FAvaTransitionStateViewModel::UpdateStateDescription(bool bInSetStateName)
 	}
 }
 
-const FAvaTransitionAutogeneratedStateContext* FAvaTransitionStateViewModel::FindAutogeneratedStateContext() const
-{
-	if (UAvaTransitionTreeEditorData* EditorData = GetEditorData())
-	{
-		return EditorData->FindAutogeneratedStateContext(StateId);
-	}
-	return nullptr;
-}
-
-FText FAvaTransitionStateViewModel::GetAutogeneratedLabel() const
-{
-	return AutogeneratedLabel;
-}
-
 #if WITH_STATETREE_DEBUGGER
 TSharedRef<SWidget> FAvaTransitionStateViewModel::GetOrCreateDebugIndicatorWidget()
 {
@@ -308,19 +288,7 @@ void FAvaTransitionStateViewModel::SetSelected(bool bInIsSelected)
 
 bool FAvaTransitionStateViewModel::CanGenerateRow() const
 {
-	if (!IsValid())
-	{
-		return false;
-	}
-
-	// Advanced mode should generate all rows
-	if (GetSharedData()->GetEditorMode() == EAvaTransitionEditorMode::Advanced)
-	{
-		return true;
-	}
-
-	// For non-advanced mode, only allow the row if it isn't an autogenerated-state
-	return !FindAutogeneratedStateContext();
+	return IsValid();
 }
 
 TSharedRef<ITableRow> FAvaTransitionStateViewModel::GenerateRow(const TSharedRef<STableViewBase>& InOwningTableView)
diff --git a/Engine/Plugins/Experimental/Avalanche/Source/AvalancheTransitionEditor/Private/ViewModels/State/AvaTransitionStateViewModel.h b/Engine/Plugins/Experimental/Avalanche/Source/AvalancheTransitionEditor/Private/ViewModels/State/AvaTransitionStateViewModel.h
--- a/Engine/Plugins/Experimental/Avalanche/Source/AvalancheTransitionEditor/Private/ViewModels/State/AvaTransitionStateViewModel.h
+++ b/Engine/Plugins/Experimental/Avalanche/Source/AvalancheTransitionEditor/Private/ViewModels/State/AvaTransitionStateViewModel.h
@@ -101,11 +101,6 @@ public:
 
 	void UpdateStateDescription(bool bInSetStateName);
 
-	/** If Autogenerated State Context, will return the context of how this state was autogenerated */
-	const FAvaTransitionAutogeneratedStateContext* FindAutogeneratedStateContext() const;
-
-	FText GetAutogeneratedLabel() const;
-
 	TSharedPtr<FAvaTransitionConditionContainerViewModel> GetConditionContainer() const
 	{
 		return ConditionContainer;
@@ -195,7 +190,5 @@ private:
 
 	FText StateDescription;
 
-	FText AutogeneratedLabel;
-
 	bool bIsSelected = false;
 };
diff --git a/Engine/Plugins/Experimental/Avalanche/Source/AvalancheTransitionEditor/Private/Views/SAvaTransitionStateView.cpp b/Engine/Plugins/Experimental/Avalanche/Source/AvalancheTransitionEditor/Private/Views/SAvaTransitionStateView.cpp
--- a/Engine/Plugins/Experimental/Avalanche/Source/AvalancheTransitionEditor/Private/Views/SAvaTransitionStateView.cpp
+++ b/Engine/Plugins/Experimental/Avalanche/Source/AvalancheTransitionEditor/Private/Views/SAvaTransitionStateView.cpp
@@ -104,19 +104,6 @@ void SAvaTransitionStateView::Construct(const FArguments& InArgs, const TSharedR
 		];
 #endif
 
-	// Autogenerated State
-	RowContent->AddSlot()
-		.HAlign(HAlign_Left)
-		.VAlign(VAlign_Center)
-		.Padding(12.f, 0.f, 0.f, 0.f)
-		.AutoWidth()
-		[
-			SNew(STextBlock)
-			.Text(InStateViewModel, &FAvaTransitionStateViewModel::GetAutogeneratedLabel)
-			.Visibility(this, &SAvaTransitionStateView::GetAutogeneratedLabelVisibility)
-			.TextStyle(FAvaTransitionEditorStyle::Get(), "StateTree.Details")
-		];
-
 	ChildSlot
 		.HAlign(HAlign_Left)
 		.VAlign(VAlign_Center)
@@ -263,14 +250,4 @@ EVisibility SAvaTransitionStateView::GetStateBreakpointVisibility() const
 	return EVisibility::Hidden;
 }
 
-EVisibility SAvaTransitionStateView::GetAutogeneratedLabelVisibility() const
-{
-	TSharedPtr<FAvaTransitionStateViewModel> StateViewModel = StateViewModelWeak.Pin();
-	if (StateViewModel.IsValid() && StateViewModel->FindAutogeneratedStateContext())
-	{
-		return EVisibility::Visible;
-	}
-	return EVisibility::Collapsed;
-}
-
 #undef LOCTEXT_NAMESPACE
diff --git a/Engine/Plugins/Experimental/Avalanche/Source/AvalancheTransitionEditor/Private/Views/SAvaTransitionStateView.h b/Engine/Plugins/Experimental/Avalanche/Source/AvalancheTransitionEditor/Private/Views/SAvaTransitionStateView.h
--- a/Engine/Plugins/Experimental/Avalanche/Source/AvalancheTransitionEditor/Private/Views/SAvaTransitionStateView.h
+++ b/Engine/Plugins/Experimental/Avalanche/Source/AvalancheTransitionEditor/Private/Views/SAvaTransitionStateView.h
@@ -28,7 +28,5 @@ private:
 
 	EVisibility GetStateBreakpointVisibility() const;
 
-	EVisibility GetAutogeneratedLabelVisibility() const;
-
 	TWeakPtr<FAvaTransitionStateViewModel> StateViewModelWeak;
 };
diff --git a/Engine/Plugins/Experimental/Avalanche/Source/AvalancheTransitionEditor/Public/AvaTransitionTreeEditorData.h b/Engine/Plugins/Experimental/Avalanche/Source/AvalancheTransitionEditor/Public/AvaTransitionTreeEditorData.h
--- a/Engine/Plugins/Experimental/Avalanche/Source/AvalancheTransitionEditor/Public/AvaTransitionTreeEditorData.h
+++ b/Engine/Plugins/Experimental/Avalanche/Source/AvalancheTransitionEditor/Public/AvaTransitionTreeEditorData.h
@@ -6,7 +6,6 @@
 #include "Containers/Map.h"
 #include "Delegates/Delegate.h"
 #include "Misc/Guid.h"
-#include "State/AvaTransitionAutogeneratedStateContext.h"
 #include "State/AvaTransitionStateMetadata.h"
 #include "StateTreeEditorData.h"
 #include "AvaTransitionTreeEditorData.generated.h"
@@ -56,22 +55,6 @@ public:
 		return StateMetadata.FindOrAdd(InStateId);
 	}
 
-	const FAvaTransitionAutogeneratedStateContext* FindAutogeneratedStateContext(const FGuid& InStateId) const
-	{
-		return AutogeneratedStates.Find(InStateId);
-	}
-
-	bool IsAutogeneratedState(const UStateTreeState& InState, const FAvaTransitionAutogeneratedStateContext& InContext) const
-	{
-		const FAvaTransitionAutogeneratedStateContext* AutogeneratedContext = FindAutogeneratedStateContext(InState.ID);
-		return AutogeneratedContext && AutogeneratedContext->Matches(InContext);
-	}
-
-	void SetAutogeneratedState(const FGuid& InStateId, const FAvaTransitionAutogeneratedStateContext& InContext)
-	{
-		AutogeneratedStates.Add(InStateId, InContext);
-	}
-
 	FSimpleMulticastDelegate& GetOnTreeRequestRefresh()
 	{
 		return OnTreeRequestRefresh;
@@ -86,9 +69,5 @@ private:
 	UPROPERTY()
 	TMap<FGuid, FAvaTransitionStateMetadata> StateMetadata;
 
-	/** Map of an autogenerated state's id to their context information */
-	UPROPERTY()
-	TMap<FGuid, FAvaTransitionAutogeneratedStateContext> AutogeneratedStates;
-
 	FSimpleMulticastDelegate OnTreeRequestRefresh;
 };
diff --git a/Engine/Plugins/Experimental/Avalanche/Source/AvalancheTransitionEditor/Public/IAvaTransitionEditorModule.h b/Engine/Plugins/Experimental/Avalanche/Source/AvalancheTransitionEditor/Public/IAvaTransitionEditorModule.h
--- a/Engine/Plugins/Experimental/Avalanche/Source/AvalancheTransitionEditor/Public/IAvaTransitionEditorModule.h
+++ b/Engine/Plugins/Experimental/Avalanche/Source/AvalancheTransitionEditor/Public/IAvaTransitionEditorModule.h
@@ -29,9 +29,5 @@ public:
 	/** Returns delegate to build the default Transition Tree */
 	virtual FOnBuildDefaultTransitionTree& GetOnBuildDefaultTransitionTree() = 0;
 
-	DECLARE_MULTICAST_DELEGATE_OneParam(FOnCompileTransitionTree, UAvaTransitionTreeEditorData&);
-	/** Returns multicast delegate to extend Transition Tree Compilation. Only called when compiling in non-Advanced mode */
-	virtual FOnCompileTransitionTree& GetOnCompileTransitionTree() = 0;
-
 	virtual void GenerateTransitionTreeOptionsMenu(UToolMenu* InMenu, IAvaTransitionBehavior* InTransitionBehavior) = 0;
 };
diff --git a/Engine/Plugins/Experimental/Avalanche/Source/AvalancheTransitionEditor/Public/State/AvaTransitionAutogeneratedStateContext.h b/Engine/Plugins/Experimental/Avalanche/Source/AvalancheTransitionEditor/Public/State/AvaTransitionAutogeneratedStateContext.h
deleted file mode 100644
