Motion Design: Rename fix

diff --git a/Engine/Plugins/Experimental/Avalanche/Source/Avalanche/Private/AvaScene.cpp b/Engine/Plugins/Experimental/Avalanche/Source/Avalanche/Private/AvaScene.cpp
--- a/Engine/Plugins/Experimental/Avalanche/Source/Avalanche/Private/AvaScene.cpp
+++ b/Engine/Plugins/Experimental/Avalanche/Source/Avalanche/Private/AvaScene.cpp
@@ -81,7 +81,7 @@ AAvaScene::AAvaScene()
 #if WITH_EDITOR
 	if (!HasAnyFlags(EObjectFlags::RF_ClassDefaultObject))
 	{
-		PreWorldRenameDelegate = FWorldDelegates::OnPreWorldRename.AddUObject(this, &AAvaScene::OnWorldRenamed);
+		PreWorldRenameDelegate = FWorldDelegates::OnPostWorldRename.AddUObject(this, &AAvaScene::OnWorldRenamed);
 		WorldTagGetterDelegate = UObject::FAssetRegistryTag::OnGetExtraObjectTagsWithContext.AddUObject(this, &AAvaScene::OnGetWorldTags);
 	}
 #endif
@@ -113,7 +113,7 @@ IAvaSequencePlaybackObject* AAvaScene::GetScenePlayback() const
 }
 
 #if WITH_EDITOR
-void AAvaScene::OnWorldRenamed(UWorld* InWorld, const TCHAR* InName, UObject* InNewOuter, ERenameFlags InFlags, bool& bOutShouldFailRename)
+void AAvaScene::OnWorldRenamed(UWorld* InWorld)
 {
 	if (FUObjectThreadContext::Get().IsRoutingPostLoad || InWorld != GetWorld())
 	{
@@ -124,7 +124,7 @@ void AAvaScene::OnWorldRenamed(UWorld* InWorld, const TCHAR* InName, UObject* In
 	{
 		if (Sequence)
 		{
-			Sequence->OnOuterWorldRenamed(InName, InNewOuter, InFlags, bOutShouldFailRename);
+			Sequence->OnOuterWorldRenamed(this);
 		}
 	}
 }
diff --git a/Engine/Plugins/Experimental/Avalanche/Source/Avalanche/Public/AvaScene.h b/Engine/Plugins/Experimental/Avalanche/Source/Avalanche/Public/AvaScene.h
--- a/Engine/Plugins/Experimental/Avalanche/Source/Avalanche/Public/AvaScene.h
+++ b/Engine/Plugins/Experimental/Avalanche/Source/Avalanche/Public/AvaScene.h
@@ -39,7 +39,7 @@ public:
 #if WITH_EDITOR
 	TArray<uint8>& GetOutlinerData() { return OutlinerData; }
 
-	void OnWorldRenamed(UWorld* InWorld, const TCHAR* InName, UObject* InNewOuter, ERenameFlags InFlags, bool& bOutShouldFailRename);
+	void OnWorldRenamed(UWorld* InWorld);
 
 	void OnGetWorldTags(FAssetRegistryTagsContext Context) const;
 #endif
diff --git a/Engine/Plugins/Experimental/Avalanche/Source/AvalancheSequence/Private/AvaSequence.cpp b/Engine/Plugins/Experimental/Avalanche/Source/AvalancheSequence/Private/AvaSequence.cpp
--- a/Engine/Plugins/Experimental/Avalanche/Source/AvalancheSequence/Private/AvaSequence.cpp
+++ b/Engine/Plugins/Experimental/Avalanche/Source/AvalancheSequence/Private/AvaSequence.cpp
@@ -622,17 +622,11 @@ FGuid UAvaSequence::FindGuidFromObject(UObject* InObject)
 }
 
 #if WITH_EDITOR
-void UAvaSequence::OnOuterWorldRenamed(const TCHAR* InName, UObject* InNewOuter, ERenameFlags InFlags, bool& bOutShouldFailRename)
+void UAvaSequence::OnOuterWorldRenamed(UObject* InNewOuter)
 {
-	if (!DirectorBlueprint)
+	if (DirectorBlueprint)
 	{
-		return;
-	}
-
-	UAvaSequenceDirectorBlueprint* Director = CastChecked<UAvaSequenceDirectorBlueprint>(DirectorBlueprint);
-	if (!Director->OnOuterWorldRenamed(InName, InNewOuter, InFlags))
-	{
-		bOutShouldFailRename = true;
+		DirectorBlueprint->Rename(*GetDirectorBlueprintName(), InNewOuter, REN_DontCreateRedirectors);
 	}
 }
 #endif
diff --git a/Engine/Plugins/Experimental/Avalanche/Source/AvalancheSequence/Private/Director/AvaSequenceDirectorBlueprint.cpp b/Engine/Plugins/Experimental/Avalanche/Source/AvalancheSequence/Private/Director/AvaSequenceDirectorBlueprint.cpp
--- a/Engine/Plugins/Experimental/Avalanche/Source/AvalancheSequence/Private/Director/AvaSequenceDirectorBlueprint.cpp
+++ b/Engine/Plugins/Experimental/Avalanche/Source/AvalancheSequence/Private/Director/AvaSequenceDirectorBlueprint.cpp
@@ -11,39 +11,6 @@ TConstArrayView<FAvaSequenceInfo> UAvaSequenceDirectorBlueprint::GetSequenceInfo
 }
 
 #if WITH_EDITOR
-bool UAvaSequenceDirectorBlueprint::OnOuterWorldRenamed(const TCHAR* InName, UObject* InNewOuter, ERenameFlags InRenameFlags)
-{
-	const bool bIsRenameTest = (InRenameFlags & REN_Test);
-
-	TArray<UObject*> LastEditedDocumentsObjects;
-	if (!bIsRenameTest)
-	{
-		LastEditedDocumentsObjects.Reserve(LastEditedDocuments.Num());
-		for (const FEditedDocumentInfo& LastEditedDocument : LastEditedDocuments)
-		{
-			LastEditedDocumentsObjects.Add(LastEditedDocument.EditedObjectPath.ResolveObject());
-		}
-	}
-
-	if (!RenameGeneratedClasses(InName, InNewOuter, InRenameFlags))
-	{
-		return false;
-	}
-
-	if (!bIsRenameTest)
-	{
-		for (int32 DocumentIndex = 0; DocumentIndex < LastEditedDocuments.Num(); ++DocumentIndex)
-		{
-			if (LastEditedDocumentsObjects[DocumentIndex])
-			{
-				LastEditedDocuments[DocumentIndex].EditedObjectPath = LastEditedDocumentsObjects[DocumentIndex];
-			}
-		}
-	}
-
-	return true;
-}
-
 UClass* UAvaSequenceDirectorBlueprint::GetBlueprintClass() const
 {
 	return UAvaSequenceDirectorGeneratedClass::StaticClass();
diff --git a/Engine/Plugins/Experimental/Avalanche/Source/AvalancheSequence/Public/AvaSequence.h b/Engine/Plugins/Experimental/Avalanche/Source/AvalancheSequence/Public/AvaSequence.h
--- a/Engine/Plugins/Experimental/Avalanche/Source/AvalancheSequence/Public/AvaSequence.h
+++ b/Engine/Plugins/Experimental/Avalanche/Source/AvalancheSequence/Public/AvaSequence.h
@@ -139,7 +139,7 @@ public:
 	FGuid FindGuidFromObject(UObject* InObject);
 
 #if WITH_EDITOR
-	void OnOuterWorldRenamed(const TCHAR* InName, UObject* InNewOuter, ERenameFlags InFlags, bool& bOutShouldFailRename);
+	void OnOuterWorldRenamed(UObject* InNewOuter);
 #endif
 
 	void OnWorldCleanup(UWorld* InWorld, bool bInSessionEnded, bool bInCleanupResources);
diff --git a/Engine/Plugins/Experimental/Avalanche/Source/AvalancheSequence/Public/Director/AvaSequenceDirectorBlueprint.h b/Engine/Plugins/Experimental/Avalanche/Source/AvalancheSequence/Public/Director/AvaSequenceDirectorBlueprint.h
--- a/Engine/Plugins/Experimental/Avalanche/Source/AvalancheSequence/Public/Director/AvaSequenceDirectorBlueprint.h
+++ b/Engine/Plugins/Experimental/Avalanche/Source/AvalancheSequence/Public/Director/AvaSequenceDirectorBlueprint.h
@@ -16,11 +16,9 @@ class UAvaSequenceDirectorBlueprint : public UBlueprint
 public:
 	AVALANCHESEQUENCE_API TConstArrayView<FAvaSequenceInfo> GetSequenceInfos();
 
-#if WITH_EDITOR
-	bool OnOuterWorldRenamed(const TCHAR* InName, UObject* InNewOuter, ERenameFlags InRenameFlags);
-
 protected:
 	//~ Begin UBlueprint
+#if WITH_EDITOR
 	virtual UClass* GetBlueprintClass() const override;
 	virtual void GetReparentingRules(TSet<const UClass*>& OutAllowedChildrenOfClasses, TSet<const UClass*>& OutDisallowedChildrenOfClasses) const override;
 #endif
