[Motion Design] Fix FAvaPlaybackManager::LoadPlaybackObject edge case when package is loaded but objects have been GC'd.

Using ResolveObject and TryLoad is more reliable and support object redirectors.

diff --git a/Engine/Plugins/Experimental/Avalanche/Source/AvalancheMedia/Private/Playback/AvaPlaybackManager.cpp b/Engine/Plugins/Experimental/Avalanche/Source/AvalancheMedia/Private/Playback/AvaPlaybackManager.cpp
--- a/Engine/Plugins/Experimental/Avalanche/Source/AvalancheMedia/Private/Playback/AvaPlaybackManager.cpp
+++ b/Engine/Plugins/Experimental/Avalanche/Source/AvalancheMedia/Private/Playback/AvaPlaybackManager.cpp
@@ -356,67 +356,49 @@ void FAvaPlaybackManager::InvalidatePlaybackAssetEntry(const FSoftObjectPath& In
 	PlaybackAssetEntries.Remove(InAssetPath);
 }
 
-UAvaPlaybackGraph*  FAvaPlaybackManager::LoadPlaybackObject(const FSoftObjectPath& InAssetPath, const FString& InChannelName) const
+UAvaPlaybackGraph* FAvaPlaybackManager::LoadPlaybackObject(const FSoftObjectPath& InAssetPath, const FString& InChannelName) const
 {
 	TRACE_CPUPROFILER_EVENT_SCOPE(FAvaPlaybackManager::LoadPlaybackObject);
 	
 	const FString PackageName = InAssetPath.GetLongPackageName();
-	const FString AssetName = InAssetPath.GetAssetName();
-	
-	// First check if the package is already loaded.
-	UPackage* TempPackage = FindPackage(nullptr, *PackageName);
 
-	if (!TempPackage)
+	// Fast path: avoid a sync package load for maps.
+	// The generated playback graph will load it async using a level streaming playable.
+	if (FAvaPlaybackUtils::IsMapAsset(PackageName))
 	{
-		// Short cut: avoid sync package load for maps.
-		// The playback object will load it async using level streaming.
-		if (FAvaPlaybackUtils::IsMapAsset(PackageName))
-		{
-			UAvaPlaybackGraph* const AvaPlayback = BuildPlaybackFromWorld(TSoftObjectPtr<UWorld>(InAssetPath), InChannelName);
-			check(AvaPlayback);
-			return AvaPlayback;
-		}
-
-		// Todo: Investigate LoadPackageAsync.
-		// For now, we tolerate a sync load here because there will be hitch from converting the
-		// Motion Design Asset to a world.
-		TempPackage = LoadPackage(nullptr, *PackageName, LOAD_None );
+		UAvaPlaybackGraph* const AvaPlayback = BuildPlaybackFromWorld(TSoftObjectPtr<UWorld>(InAssetPath), InChannelName);
+		check(AvaPlayback);
+		return AvaPlayback;
 	}
 	
-	if (TempPackage)
+	UObject* LoadedAsset = InAssetPath.ResolveObject();
+	if (!LoadedAsset)
 	{
-		if (UObject* FoundObject = FindObject<UObject>(TempPackage, *AssetName))
-		{
-			// When the asset is an Motion Design Playback, it is loaded directly.
-			if (UAvaPlaybackGraph* const AvaPlayback = Cast<UAvaPlaybackGraph>(FoundObject))
-			{
-				return AvaPlayback;
-			}
-
-			if (const UWorld* const World = Cast<UWorld>(FoundObject))
-			{
-				UAvaPlaybackGraph* const AvaPlayback = BuildPlaybackFromWorld(World, InChannelName);
-				check(AvaPlayback);
-				return AvaPlayback;
-			}
-
-			UE_LOG(LogAvaPlaybackManager, Error,
-				TEXT("Asset \"%s\" in package \"%s\" is not a supported Motion Design playback asset (\"%s\")."),
-				*AssetName, *PackageName, *FoundObject->GetClass()->GetFullName());
-		}
-		else
+		LoadedAsset = InAssetPath.TryLoad();
+		if (!LoadedAsset)
 		{
-			UE_LOG(LogAvaPlaybackManager, Error,
-				TEXT("Failed to find asset \"%s\" in package \"%s\""),
-				*AssetName, *PackageName);
+			UE_LOG(LogAvaPlaybackManager, Error, TEXT("Failed to load asset \"%s\""), *InAssetPath.ToString());
+			return nullptr;
 		}
 	}
-	else
+	
+	// When the asset is a Playback Graph, it is used directly.
+	if (UAvaPlaybackGraph* const AvaPlayback = Cast<UAvaPlaybackGraph>(LoadedAsset))
 	{
-		UE_LOG(LogAvaPlaybackManager, Error,
-			TEXT("Failed to load package \"%s\""),
-			*PackageName);
+		return AvaPlayback;
 	}
+
+	if (const UWorld* const World = Cast<UWorld>(LoadedAsset))
+	{
+		UAvaPlaybackGraph* const AvaPlayback = BuildPlaybackFromWorld(World, InChannelName);
+		check(AvaPlayback);
+		return AvaPlayback;
+	}
+
+	UE_LOG(LogAvaPlaybackManager, Error,
+		TEXT("Asset \"%s\" is not a supported Motion Design playback asset (\"%s\")."),
+		*InAssetPath.ToString(), *LoadedAsset->GetClass()->GetFullName());
+	
 	return nullptr;
 }
 
