[Motion Design] Fix Broadcast channel status doesn't refresh when devices are added or removed.

diff --git a/Engine/Plugins/Experimental/Avalanche/Source/AvalancheMedia/Private/Broadcast/Channel/AvaBroadcastOutputChannel.cpp b/Engine/Plugins/Experimental/Avalanche/Source/AvalancheMedia/Private/Broadcast/Channel/AvaBroadcastOutputChannel.cpp
--- a/Engine/Plugins/Experimental/Avalanche/Source/AvalancheMedia/Private/Broadcast/Channel/AvaBroadcastOutputChannel.cpp
+++ b/Engine/Plugins/Experimental/Avalanche/Source/AvalancheMedia/Private/Broadcast/Channel/AvaBroadcastOutputChannel.cpp
@@ -240,7 +240,8 @@ void FAvaBroadcastOutputChannel::DuplicateChannel(const FAvaBroadcastOutputChann
 		OutTargetChannel.MediaOutputInfos.Last().Guid = FGuid::NewGuid();	// Allocate a new guid for the duplicate.
 	}
 	OnChannelChanged.Broadcast(OutTargetChannel, EAvaBroadcastChannelChange::MediaOutputs);
-	OutTargetChannel.UpdateChannelResources(false);
+	OutTargetChannel.UpdateChannelResources(/*bInIsProfileActive*/false);
+	OutTargetChannel.RefreshState();
 }
 
 void FAvaBroadcastOutputChannel::SetChannelIndex(int32 InIndex)
@@ -672,9 +673,9 @@ bool FAvaBroadcastOutputChannel::StartChannelBroadcast()
 	}
 
 	// Ensure placeholder render targets are compatible with media outputs.
-	UpdateChannelResources(true);
+	UpdateChannelResources(/*bInIsProfileActive*/true);
 	
-	UTextureRenderTarget2D* const RenderTarget = GetCurrentRenderTarget(true);
+	UTextureRenderTarget2D* const RenderTarget = GetCurrentRenderTarget(/*bInFallbackToPlaceholder*/true);
 	
 	for (UMediaOutput* MediaOutput : MediaOutputs)
 	{
@@ -839,7 +840,8 @@ void FAvaBroadcastOutputChannel::AddMediaOutput(UMediaOutput* InMediaOutput, con
 		MediaOutputInfos.Add(InOutputInfo);
 		// Note: the media output may not be fully configured at that point.
 		OnChannelChanged.Broadcast(*this, EAvaBroadcastChannelChange::MediaOutputs);
-		UpdateChannelResources(true);
+		UpdateChannelResources(/*bInIsProfileActive*/true);
+		RefreshState();
 	}
 }
 
@@ -856,7 +858,8 @@ int32 FAvaBroadcastOutputChannel::RemoveMediaOutput(UMediaOutput* InMediaOutput)
 		MediaOutputs.RemoveAt(IndexToRemove);
 		MediaOutputInfos.RemoveAt(IndexToRemove);
 		OnChannelChanged.Broadcast(*this, EAvaBroadcastChannelChange::MediaOutputs);
-		UpdateChannelResources(true);
+		UpdateChannelResources(/*bInIsProfileActive*/true);
+		RefreshState();
 	}
 	return RemoveCount;
 }
@@ -881,7 +884,7 @@ void FAvaBroadcastOutputChannel::OnMediaOutputModified(UMediaOutput* InMediaOutp
 	}
 	else
 	{
-		UpdateChannelResources(true);
+		UpdateChannelResources(/*bInIsProfileActive*/true);
 		RefreshState();
 	}
 }
@@ -1133,7 +1136,7 @@ void FAvaBroadcastOutputChannel::UpdateViewportTarget()
 	//Only need to update when Broadcasting, since when starting Broadcast it will Capture and update it anyways
 	if (GetState() == EAvaBroadcastChannelState::Live)
 	{
-		UTextureRenderTarget2D* const RenderTarget = GetCurrentRenderTarget(true);
+		UTextureRenderTarget2D* const RenderTarget = GetCurrentRenderTarget(/*bInFallbackToPlaceholder*/true);
 
 		bool bNeedStateRefresh = false;
 
@@ -1214,7 +1217,7 @@ void FAvaBroadcastOutputChannel::StopCaptureForOutput(const UMediaOutput* InMedi
 	UMediaCapture* const MediaCapture = GetMediaCaptureForOutput(InMediaOutput);
 	if (IsValid(MediaCapture) && UE::AvaMedia::Private::IsCapturing(MediaCapture))
 	{
-		MediaCapture->StopCapture(false);
+		MediaCapture->StopCapture(/*bAllowPendingFrameToBeProcess*/false);
 	}
 	MediaCaptures.Remove(InMediaOutput);
 }
@@ -1222,7 +1225,7 @@ void FAvaBroadcastOutputChannel::StopCaptureForOutput(const UMediaOutput* InMedi
 void FAvaBroadcastOutputChannel::TickPlaceholder(float)
 {
 	// Only Draw if Render Target is Invalid 
-	if (GetState() == EAvaBroadcastChannelState::Live && !IsValid(GetCurrentRenderTarget(false)))
+	if (GetState() == EAvaBroadcastChannelState::Live && !IsValid(GetCurrentRenderTarget(/*bInFallbackToPlaceholder*/false)))
 	{
 		DrawPlaceholderWidget();
 	}
