[Motion Design] Adding Playback Server Logging Entry Customization to make the settings more readable.

diff --git a/Engine/Plugins/Experimental/Avalanche/Source/AvalancheMediaEditor/Private/AvaMediaEditorModule.cpp b/Engine/Plugins/Experimental/Avalanche/Source/AvalancheMediaEditor/Private/AvaMediaEditorModule.cpp
--- a/Engine/Plugins/Experimental/Avalanche/Source/AvalancheMediaEditor/Private/AvaMediaEditorModule.cpp
+++ b/Engine/Plugins/Experimental/Avalanche/Source/AvalancheMediaEditor/Private/AvaMediaEditorModule.cpp
@@ -3,8 +3,10 @@
 #include "AvaMediaEditorModule.h"
 
 #include "AvaMediaEditorStyle.h"
+#include "AvaMediaSettings.h"
 #include "Broadcast/AvaBroadcastEditor.h"
 #include "Broadcast/OutputDevices/AvaBroadcastMediaIOOutputConfigurationCustomization.h"
+#include "Customization/AvaPlaybackServerLoggingEntryCustomization.h"
 #include "Customization/AvaSynchronizedEventsFeatureSelectionCustomization.h"
 #include "Editor.h"
 #include "Engine/GameViewportClient.h"
@@ -272,6 +274,7 @@ void FAvaMediaEditorModule::RegisterCustomizations() const
 	PropertyModule.RegisterCustomPropertyTypeLayout(FAvaRundownMacroCommand::StaticStruct()->GetFName(), FOnGetPropertyTypeCustomizationInstance::CreateStatic(&FAvaRundownMacroCommandCustomization::MakeInstance));
 	PropertyModule.RegisterCustomPropertyTypeLayout(FAvaRundownMacroKeyBinding::StaticStruct()->GetFName(), FOnGetPropertyTypeCustomizationInstance::CreateStatic(&FAvaRundownMacroKeyBindingCustomization::MakeInstance));
 	PropertyModule.RegisterCustomPropertyTypeLayout(FAvaSynchronizedEventsFeatureSelection::StaticStruct()->GetFName(), FOnGetPropertyTypeCustomizationInstance::CreateStatic(&FAvaSynchronizedEventsFeatureSelectionCustomization::MakeInstance));
+	PropertyModule.RegisterCustomPropertyTypeLayout(FAvaPlaybackServerLoggingEntry::StaticStruct()->GetFName(), FOnGetPropertyTypeCustomizationInstance::CreateStatic(&FAvaPlaybackServerLoggingEntryCustomization::MakeInstance));
 }
 
 void FAvaMediaEditorModule::UnregisterCustomizations() const
@@ -281,6 +284,7 @@ void FAvaMediaEditorModule::UnregisterCustomizations() const
 	PropertyModule.UnregisterCustomPropertyTypeLayout(FAvaRundownMacroCommand::StaticStruct()->GetFName());
 	PropertyModule.UnregisterCustomPropertyTypeLayout(FAvaRundownMacroKeyBinding::StaticStruct()->GetFName());
 	PropertyModule.UnregisterCustomPropertyTypeLayout(FAvaSynchronizedEventsFeatureSelection::StaticStruct()->GetFName());
+	PropertyModule.UnregisterCustomPropertyTypeLayout(FAvaPlaybackServerLoggingEntry::StaticStruct()->GetFName());
 }
 
 void FAvaMediaEditorModule::PostEngineInit()
diff --git a/Engine/Plugins/Experimental/Avalanche/Source/AvalancheMediaEditor/Private/Customization/AvaPlaybackServerLoggingEntryCustomization.cpp b/Engine/Plugins/Experimental/Avalanche/Source/AvalancheMediaEditor/Private/Customization/AvaPlaybackServerLoggingEntryCustomization.cpp
new file mode 100644
--- /dev/null
+++ b/Engine/Plugins/Experimental/Avalanche/Source/AvalancheMediaEditor/Private/Customization/AvaPlaybackServerLoggingEntryCustomization.cpp
@@ -0,0 +1,61 @@
+// Copyright Epic Games, Inc. All Rights Reserved.
+
+#include "Customization/AvaPlaybackServerLoggingEntryCustomization.h"
+
+#include "AvaMediaSettings.h"
+#include "IDetailChildrenBuilder.h"
+#include "PropertyCustomizationHelpers.h"
+#include "Widgets/Layout/SBox.h"
+
+#define LOCTEXT_NAMESPACE "AvaPlaybackServerLoggingEntryCustomization"
+
+TSharedRef<IPropertyTypeCustomization> FAvaPlaybackServerLoggingEntryCustomization::MakeInstance()
+{
+	return MakeShared<FAvaPlaybackServerLoggingEntryCustomization>();
+}
+
+void FAvaPlaybackServerLoggingEntryCustomization::CustomizeHeader(TSharedRef<IPropertyHandle> InStructPropertyHandle, FDetailWidgetRow& InHeaderRow, IPropertyTypeCustomizationUtils& InStructCustomizationUtils)
+{
+	const TSharedPtr<IPropertyHandle> CategoryHandle = InStructPropertyHandle->GetChildHandle(GET_MEMBER_NAME_CHECKED(FAvaPlaybackServerLoggingEntry, Category));
+	const TSharedPtr<IPropertyHandle> VerbosityHandle = InStructPropertyHandle->GetChildHandle(GET_MEMBER_NAME_CHECKED(FAvaPlaybackServerLoggingEntry, VerbosityLevel));
+	
+	const FMargin PropertyPadding(2.0f, 0.0f, 2.0f, 0.0f);
+
+	InHeaderRow
+			.NameContent()
+			[
+				InStructPropertyHandle->CreatePropertyNameWidget()
+			]
+			.ValueContent()
+			[
+				SNew(SHorizontalBox)
+				+ SHorizontalBox::Slot()
+				.Padding(PropertyPadding)
+				.HAlign(HAlign_Center)
+				.VAlign(VAlign_Center)
+				.AutoWidth()
+				[
+					SNew(SBox)
+					.HAlign(EHorizontalAlignment::HAlign_Fill)
+					.VAlign(EVerticalAlignment::VAlign_Center)
+					.MinDesiredWidth(200.f) // Note: Minimum size enough for most log categories so the verbosity is visually aligned.
+					[
+						CategoryHandle->CreatePropertyValueWidget()
+					]
+				]
+				+ SHorizontalBox::Slot()
+				.Padding(PropertyPadding)
+				.HAlign(HAlign_Center)
+				.VAlign(VAlign_Center)
+				.AutoWidth()
+				[
+					VerbosityHandle->CreatePropertyValueWidget()
+				]
+			];
+}
+
+void FAvaPlaybackServerLoggingEntryCustomization::CustomizeChildren(TSharedRef<IPropertyHandle> InStructPropertyHandle, IDetailChildrenBuilder& InStructBuilder, IPropertyTypeCustomizationUtils& InStructCustomizationUtils)
+{
+}
+
+#undef LOCTEXT_NAMESPACE
\ No newline at end of file
diff --git a/Engine/Plugins/Experimental/Avalanche/Source/AvalancheMediaEditor/Private/Customization/AvaPlaybackServerLoggingEntryCustomization.h b/Engine/Plugins/Experimental/Avalanche/Source/AvalancheMediaEditor/Private/Customization/AvaPlaybackServerLoggingEntryCustomization.h
new file mode 100644
--- /dev/null
+++ b/Engine/Plugins/Experimental/Avalanche/Source/AvalancheMediaEditor/Private/Customization/AvaPlaybackServerLoggingEntryCustomization.h
@@ -0,0 +1,16 @@
+// Copyright Epic Games, Inc. All Rights Reserved.
+
+#pragma once
+
+#include "IPropertyTypeCustomization.h"
+
+class FAvaPlaybackServerLoggingEntryCustomization : public IPropertyTypeCustomization
+{
+public:
+	static TSharedRef<IPropertyTypeCustomization> MakeInstance();
+	
+	//~ Begin IPropertyTypeCustomization
+	virtual void CustomizeHeader(TSharedRef<IPropertyHandle> InPropertyHandle, FDetailWidgetRow& InHeaderRow, IPropertyTypeCustomizationUtils& InPropertyTypeCustomizationUtils) override;
+	virtual void CustomizeChildren(TSharedRef<IPropertyHandle> InPropertyHandle, IDetailChildrenBuilder& InStructBuilder, IPropertyTypeCustomizationUtils& InPropertyTypeCustomizationUtils) override;
+	//~ End IPropertyTypeCustomization
+};
\ No newline at end of file
