[Motion Design] Adding Transition Instancing Mode to allow scene reuse across transitions.

diff --git a/Engine/Plugins/Experimental/Avalanche/Source/AvalancheTransition/Public/AvaTransitionEnums.h b/Engine/Plugins/Experimental/Avalanche/Source/AvalancheTransition/Public/AvaTransitionEnums.h
--- a/Engine/Plugins/Experimental/Avalanche/Source/AvalancheTransition/Public/AvaTransitionEnums.h
+++ b/Engine/Plugins/Experimental/Avalanche/Source/AvalancheTransition/Public/AvaTransitionEnums.h
@@ -59,3 +59,12 @@ enum class EAvaTransitionLayerCompareType : uint8
 	MatchingTag,
 	Any,
 };
+
+UENUM(BlueprintType, DisplayName="Motion Design Transition Instancing Mode")
+enum class EAvaTransitionInstancingMode : uint8
+{
+	/** A new scene instance is added to the parent world and exit instance is discarded. */
+	New,
+	/** If previous instance exists, reuse it for the entering behavior. */
+	Reuse
+};
\ No newline at end of file
diff --git a/Engine/Plugins/Experimental/Avalanche/Source/AvalancheTransition/Public/AvaTransitionTree.h b/Engine/Plugins/Experimental/Avalanche/Source/AvalancheTransition/Public/AvaTransitionTree.h
--- a/Engine/Plugins/Experimental/Avalanche/Source/AvalancheTransition/Public/AvaTransitionTree.h
+++ b/Engine/Plugins/Experimental/Avalanche/Source/AvalancheTransition/Public/AvaTransitionTree.h
@@ -3,6 +3,7 @@
 #pragma once
 
 #include "AvaTagHandle.h"
+#include "AvaTransitionEnums.h"
 #include "StateTree.h"
 #include "AvaTransitionTree.generated.h"
 
@@ -24,6 +25,16 @@ public:
 
 	void SetEnabled(bool bInEnabled);
 
+	void SetInstancingMode(EAvaTransitionInstancingMode InInstancingMode)
+	{
+		InstancingMode = InInstancingMode;
+	}
+
+	EAvaTransitionInstancingMode GetInstancingMode() const
+	{
+		return InstancingMode;
+	}
+
 	static FName GetEnabledPropertyName()
 	{
 		return GET_MEMBER_NAME_CHECKED(UAvaTransitionTree, bEnabled);
@@ -40,4 +51,7 @@ private:
 	 */
 	UPROPERTY()
 	bool bEnabled = true;
+
+	UPROPERTY()
+	EAvaTransitionInstancingMode InstancingMode = EAvaTransitionInstancingMode::New;
 };
diff --git a/Engine/Plugins/Experimental/Avalanche/Source/AvalancheTransitionEditor/Private/AvaTransitionEditorModule.cpp b/Engine/Plugins/Experimental/Avalanche/Source/AvalancheTransitionEditor/Private/AvaTransitionEditorModule.cpp
--- a/Engine/Plugins/Experimental/Avalanche/Source/AvalancheTransitionEditor/Private/AvaTransitionEditorModule.cpp
+++ b/Engine/Plugins/Experimental/Avalanche/Source/AvalancheTransitionEditor/Private/AvaTransitionEditorModule.cpp
@@ -85,6 +85,19 @@ void FAvaTransitionEditorModule::GenerateTransitionTreeOptionsMenu(UToolMenu* In
 			, FIsActionChecked::CreateStatic(&UE::AvaTransitionEditor::IsTransitionTreeEnabled, TransitionTreeWeak))
 		, EUserInterfaceActionType::ToggleButton);
 
+	if (const TSharedPtr<SWidget> ModeSelector = UE::AvaTransitionEditor::CreateTransitionInstancingModeSelector(TransitionTree))
+	{
+		GeneralSection.AddEntry(FToolMenuEntry::InitWidget(TEXT("TransitionModeSelector")
+			, SNew(SBox)
+				[
+					ModeSelector.ToSharedRef()
+				]
+			, FText::GetEmpty()
+			, /*bNoIndent*/true
+			, /*bSearchable*/false
+			, /*bNoPadding*/true));
+	}
+
 	if (TSharedPtr<SWidget> LayerPicker = UE::AvaTransitionEditor::CreateTransitionLayerPicker(Cast<UAvaTransitionTreeEditorData>(TransitionTree->EditorData), /*bInCompileOnLayerPicked*/true))
 	{
 		FToolMenuSection& LayerSection = InMenu->FindOrAddSection(TEXT("LayerSection"), LOCTEXT("LayerSectionLabel", "Layer"));
diff --git a/Engine/Plugins/Experimental/Avalanche/Source/AvalancheTransitionEditor/Private/AvaTransitionEditorUtils.cpp b/Engine/Plugins/Experimental/Avalanche/Source/AvalancheTransitionEditor/Private/AvaTransitionEditorUtils.cpp
--- a/Engine/Plugins/Experimental/Avalanche/Source/AvalancheTransitionEditor/Private/AvaTransitionEditorUtils.cpp
+++ b/Engine/Plugins/Experimental/Avalanche/Source/AvalancheTransitionEditor/Private/AvaTransitionEditorUtils.cpp
@@ -10,6 +10,7 @@
 #include "ContentBrowserModule.h"
 #include "IContentBrowserSingleton.h"
 #include "ISinglePropertyView.h"
+#include "Menu/Widgets/SAvaTransitionTreeInstancingModeSelector.h"
 #include "Modules/ModuleManager.h"
 #include "PropertyEditorModule.h"
 #include "ScopedTransaction.h"
@@ -17,6 +18,7 @@
 #include "StateTreeTaskBase.h"
 #include "Widgets/Layout/SBox.h"
 #include "Widgets/Layout/SScaleBox.h"
+#include "Widgets/Text/STextBlock.h"
 
 #define LOCTEXT_NAMESPACE "AvaTransitionEditorUtils"
 
@@ -78,6 +80,29 @@ TSharedPtr<SWidget> CreateTransitionLayerPicker(UAvaTransitionTreeEditorData* In
 	return nullptr;
 }
 
+TSharedPtr<SWidget> CreateTransitionInstancingModeSelector(UAvaTransitionTree* InTransitionTree)
+{
+	const FMargin PropertyPadding(2.0f, 0.0f, 2.0f, 0.0f);
+	
+	return SNew(SHorizontalBox)
+		+ SHorizontalBox::Slot()
+		.Padding(PropertyPadding)
+		.HAlign(HAlign_Right)
+		.VAlign(VAlign_Center)
+		.AutoWidth()
+		[
+			SNew(STextBlock)
+			.Text(LOCTEXT("TransitionInstancingModeLabel", "Mode"))
+		]
+		+ SHorizontalBox::Slot()
+		.Padding(PropertyPadding)
+		.HAlign(HAlign_Fill)
+		.VAlign(VAlign_Center)
+		[
+			SNew(SAvaTransitionTreeInstancingModeSelector, InTransitionTree)
+		];
+}
+
 void ValidateTree(UAvaTransitionTree& InTransitionTree)
 {
 	UStateTreeEditorData* EditorData = Cast<UStateTreeEditorData>(InTransitionTree.EditorData);
diff --git a/Engine/Plugins/Experimental/Avalanche/Source/AvalancheTransitionEditor/Private/AvaTransitionEditorUtils.h b/Engine/Plugins/Experimental/Avalanche/Source/AvalancheTransitionEditor/Private/AvaTransitionEditorUtils.h
--- a/Engine/Plugins/Experimental/Avalanche/Source/AvalancheTransitionEditor/Private/AvaTransitionEditorUtils.h
+++ b/Engine/Plugins/Experimental/Avalanche/Source/AvalancheTransitionEditor/Private/AvaTransitionEditorUtils.h
@@ -14,6 +14,8 @@ namespace UE::AvaTransitionEditor
 {
 	TSharedPtr<SWidget> CreateTransitionLayerPicker(UAvaTransitionTreeEditorData* InEditorData, bool bInCompileOnLayerPicked);
 
+	TSharedPtr<SWidget> CreateTransitionInstancingModeSelector(UAvaTransitionTree* InTransitionTree);
+	
 	void ValidateTree(UAvaTransitionTree& InTransitionTree);
 
 	uint32 CalculateTreeHash(UAvaTransitionTree& InTransitionTree);
diff --git a/Engine/Plugins/Experimental/Avalanche/Source/AvalancheTransitionEditor/Private/Menu/Widgets/SAvaTransitionTreeInstancingModeSelector.cpp b/Engine/Plugins/Experimental/Avalanche/Source/AvalancheTransitionEditor/Private/Menu/Widgets/SAvaTransitionTreeInstancingModeSelector.cpp
new file mode 100644
--- /dev/null
+++ b/Engine/Plugins/Experimental/Avalanche/Source/AvalancheTransitionEditor/Private/Menu/Widgets/SAvaTransitionTreeInstancingModeSelector.cpp
@@ -0,0 +1,82 @@
+// Copyright Epic Games, Inc. All Rights Reserved.
+
+#include "SAvaTransitionTreeInstancingModeSelector.h"
+
+#include "AvaTransitionTree.h"
+#include "ScopedTransaction.h"
+#include "Widgets/Input/SComboBox.h"
+
+#define LOCTEXT_NAMESPACE "AvaTransitionTreeInstancingModeSelector"
+
+void SAvaTransitionTreeInstancingModeSelector::Construct(const FArguments& InArgs, UAvaTransitionTree* InTransitionTree)
+{
+	UpdateItems();
+	TransitionTreeWeak = InTransitionTree;
+
+	ChildSlot
+	[
+		SAssignNew(Combo, SComboBox<FName>)
+		.InitiallySelectedItem(GetItemFromProperty())
+		.OptionsSource(&Items)
+		.OnGenerateWidget(this, &SAvaTransitionTreeInstancingModeSelector::GenerateWidget)
+		.OnSelectionChanged(this, &SAvaTransitionTreeInstancingModeSelector::HandleSelectionChanged)
+		[
+			SNew(STextBlock)
+			.Text(this, &SAvaTransitionTreeInstancingModeSelector::GetDisplayTextFromProperty)
+		]
+	];
+}
+	
+TSharedRef<SWidget> SAvaTransitionTreeInstancingModeSelector::GenerateWidget(FName InItem) const
+{
+	return SNew(STextBlock)
+		.Text(GetDisplayTextFromItem(InItem));
+}
+
+void SAvaTransitionTreeInstancingModeSelector::HandleSelectionChanged(FName InProposedSelection, ESelectInfo::Type InSelectInfo)
+{
+	if (UAvaTransitionTree* TransitionTree = TransitionTreeWeak.Get())
+	{
+		FScopedTransaction Transaction(LOCTEXT("SetTransitionTreeInstancingMode", "Set Transition Tree Instancing Mode")); 
+		TransitionTree->Modify();
+		int64 Value = StaticEnum<EAvaTransitionInstancingMode>()->GetValueByName(InProposedSelection); 
+		TransitionTree->SetInstancingMode(static_cast<EAvaTransitionInstancingMode>(Value));
+	}
+}
+
+FText SAvaTransitionTreeInstancingModeSelector::GetDisplayTextFromProperty() const
+{
+	if (const UAvaTransitionTree* TransitionTree = TransitionTreeWeak.Get())
+	{
+		return StaticEnum<EAvaTransitionInstancingMode>()->GetDisplayNameTextByValue(static_cast<int64>(TransitionTree->GetInstancingMode()));
+	}
+	return FText::GetEmpty();
+}
+
+FName SAvaTransitionTreeInstancingModeSelector::GetItemFromProperty() const
+{
+	if (const UAvaTransitionTree* TransitionTree = TransitionTreeWeak.Get())
+	{
+		return StaticEnum<EAvaTransitionInstancingMode>()->GetNameByValue(static_cast<int64>(TransitionTree->GetInstancingMode()));
+	}
+	return NAME_None;
+}
+
+FText SAvaTransitionTreeInstancingModeSelector::GetDisplayTextFromItem(FName InItem) const
+{
+	const int64 Value = StaticEnum<EAvaTransitionInstancingMode>()->GetValueByName(InItem);
+	return StaticEnum<EAvaTransitionInstancingMode>()->GetDisplayNameTextByValue(Value);
+}
+
+void SAvaTransitionTreeInstancingModeSelector::UpdateItems()
+{
+	const UEnum* InstancingModeEnum = StaticEnum<EAvaTransitionInstancingMode>();
+
+	Items.Empty(InstancingModeEnum->NumEnums() - 1);
+	for (int32 Index = 0; Index < InstancingModeEnum->NumEnums() - 1; ++Index)
+	{
+		Items.Add(InstancingModeEnum->GetNameByIndex(Index));
+	}
+}
+
+#undef LOCTEXT_NAMESPACE
\ No newline at end of file
diff --git a/Engine/Plugins/Experimental/Avalanche/Source/AvalancheTransitionEditor/Private/Menu/Widgets/SAvaTransitionTreeInstancingModeSelector.h b/Engine/Plugins/Experimental/Avalanche/Source/AvalancheTransitionEditor/Private/Menu/Widgets/SAvaTransitionTreeInstancingModeSelector.h
new file mode 100644
--- /dev/null
+++ b/Engine/Plugins/Experimental/Avalanche/Source/AvalancheTransitionEditor/Private/Menu/Widgets/SAvaTransitionTreeInstancingModeSelector.h
@@ -0,0 +1,31 @@
+// Copyright Epic Games, Inc. All Rights Reserved.
+
+#pragma once
+
+#include "UObject/WeakObjectPtr.h"
+#include "Widgets/SCompoundWidget.h"
+
+class UAvaTransitionTree;
+template <typename OptionType> class SComboBox;
+
+class SAvaTransitionTreeInstancingModeSelector : public SCompoundWidget
+{
+public:
+	SLATE_BEGIN_ARGS(SAvaTransitionTreeInstancingModeSelector){}
+	SLATE_END_ARGS()
+
+	void Construct(const FArguments& InArgs, UAvaTransitionTree* InTransitionTree);
+
+protected:
+	TSharedRef<SWidget> GenerateWidget(FName InItem) const;
+	void HandleSelectionChanged(FName InProposedSelection, ESelectInfo::Type InSelectInfo);
+	FText GetDisplayTextFromProperty() const;
+	FName GetItemFromProperty() const;
+	FText GetDisplayTextFromItem(FName InItem) const;
+	void UpdateItems();
+
+protected:
+	TWeakObjectPtr<UAvaTransitionTree> TransitionTreeWeak;
+	TSharedPtr<SComboBox<FName>> Combo;
+	TArray<FName> Items;
+};
