[Motion Design] Implementation of Transition Instancing mode for Rundown.

Playable Transition Changes:

- Principle: a reused playable is added both as "enter" and "playing" in the transition (that logic is implemented in page transition, requires rundown page data to resolve).
- If an enter playable is also in the playing list, the transition will not update the RC values automatically at the start of the transition. It is assumed it will be updated through a director BP event instead.
- Temp hack: MarkPlayableAsDiscard will ignore reused playables. Will need to revisit the transition tree logic to avoid that.
- Transition builder: if the transition has reused playables, we need to allow a playable to be added in more than one role.

Rundown Changes:

- TL instancing mode is cached in the page. Todo: need to find a way to indicate that in the rundown editor.
- Reused instance players can either bypass transitions, as before, or can be part of it as "inplace" instances.
- The information about bypassing or reused instances has been moved to the page transition itself. No longer polluting page player class.

diff --git a/Engine/Plugins/Experimental/Avalanche/Source/AvalancheMedia/Private/Playable/Transition/AvaPlayableLocalTransition.cpp b/Engine/Plugins/Experimental/Avalanche/Source/AvalancheMedia/Private/Playable/Transition/AvaPlayableLocalTransition.cpp
--- a/Engine/Plugins/Experimental/Avalanche/Source/AvalancheMedia/Private/Playable/Transition/AvaPlayableLocalTransition.cpp
+++ b/Engine/Plugins/Experimental/Avalanche/Source/AvalancheMedia/Private/Playable/Transition/AvaPlayableLocalTransition.cpp
@@ -280,7 +280,7 @@ bool UAvaPlayableLocalTransition::Start()
 	int32 ArrayIndex = 0;
 	for (UAvaPlayable* Playable : EnterPlayables)
 	{
-		if (EnterPlayableValues.IsValidIndex(ArrayIndex) && EnterPlayableValues[ArrayIndex].IsValid())
+		if (EnterPlayableValues.IsValidIndex(ArrayIndex) && EnterPlayableValues[ArrayIndex].IsValid() && !IsPlayingPlayable(Playable))
 		{
 			Playable->UpdateRemoteControlCommand(EnterPlayableValues[ArrayIndex].ToSharedRef());
 		}
diff --git a/Engine/Plugins/Experimental/Avalanche/Source/AvalancheMedia/Private/Playable/Transition/AvaPlayableTransition.cpp b/Engine/Plugins/Experimental/Avalanche/Source/AvalancheMedia/Private/Playable/Transition/AvaPlayableTransition.cpp
--- a/Engine/Plugins/Experimental/Avalanche/Source/AvalancheMedia/Private/Playable/Transition/AvaPlayableTransition.cpp
+++ b/Engine/Plugins/Experimental/Avalanche/Source/AvalancheMedia/Private/Playable/Transition/AvaPlayableTransition.cpp
@@ -111,6 +111,12 @@ TSharedPtr<FAvaPlayableRemoteControlValues> UAvaPlayableTransition::GetValuesFor
 
 void UAvaPlayableTransition::MarkPlayableAsDiscard(UAvaPlayable* InPlayable)
 {
+	// Hack: Intercept "in-place" playable, don't mark them for discard.
+	if (IsEnterPlayable(InPlayable) && IsPlayingPlayable(InPlayable))
+	{
+		return;
+	}
+	
 	DiscardPlayablesWeak.AddUnique(InPlayable);
 	UAvaPlayable::OnTransitionEvent().Broadcast(InPlayable, this, EAvaPlayableTransitionEventFlags::MarkPlayableDiscard);
 }
@@ -186,9 +192,9 @@ void FAvaPlayableTransitionBuilder::AddEnterPlayableValues(const TSharedPtr<FAva
 	EnterPlayableValues.Add(InValues);
 }
 
-bool FAvaPlayableTransitionBuilder::AddEnterPlayable(UAvaPlayable* InPlayable)
+bool FAvaPlayableTransitionBuilder::AddEnterPlayable(UAvaPlayable* InPlayable, bool bInAllowMultipleAdd)
 {
-	if (ExitPlayablesWeak.Contains(InPlayable))
+	if (!bInAllowMultipleAdd && ExitPlayablesWeak.Contains(InPlayable))
 	{
 		using namespace UE::AvaPlayableTransition::Private;
 		UE_LOG(LogAvaPlayable, Error,
@@ -197,7 +203,7 @@ bool FAvaPlayableTransitionBuilder::AddEnterPlayable(UAvaPlayable* InPlayable)
 		return false;
 	}
 
-	if (PlayingPlayablesWeak.Contains(InPlayable))
+	if (!bInAllowMultipleAdd && PlayingPlayablesWeak.Contains(InPlayable))
 	{
 		using namespace UE::AvaPlayableTransition::Private;
 		UE_LOG(LogAvaPlayable, Error,
@@ -210,9 +216,9 @@ bool FAvaPlayableTransitionBuilder::AddEnterPlayable(UAvaPlayable* InPlayable)
 	return true;
 }
 
-bool FAvaPlayableTransitionBuilder::AddPlayingPlayable(UAvaPlayable* InPlayable)
+bool FAvaPlayableTransitionBuilder::AddPlayingPlayable(UAvaPlayable* InPlayable, bool bInAllowMultipleAdd)
 {
-	if (EnterPlayablesWeak.Contains(InPlayable))
+	if (!bInAllowMultipleAdd && EnterPlayablesWeak.Contains(InPlayable))
 	{
 		using namespace UE::AvaPlayableTransition::Private;
 		UE_LOG(LogAvaPlayable, Error,
@@ -221,7 +227,7 @@ bool FAvaPlayableTransitionBuilder::AddPlayingPlayable(UAvaPlayable* InPlayable)
 		return false;
 	}
 
-	if (ExitPlayablesWeak.Contains(InPlayable))
+	if (!bInAllowMultipleAdd && ExitPlayablesWeak.Contains(InPlayable))
 	{
 		using namespace UE::AvaPlayableTransition::Private;
 		UE_LOG(LogAvaPlayable, Error,
@@ -235,9 +241,9 @@ bool FAvaPlayableTransitionBuilder::AddPlayingPlayable(UAvaPlayable* InPlayable)
 	return true;
 }
 
-bool FAvaPlayableTransitionBuilder::AddExitPlayable(UAvaPlayable* InPlayable)
+bool FAvaPlayableTransitionBuilder::AddExitPlayable(UAvaPlayable* InPlayable, bool bInAllowMultipleAdd)
 {
-	if (EnterPlayablesWeak.Contains(InPlayable))
+	if (!bInAllowMultipleAdd && EnterPlayablesWeak.Contains(InPlayable))
 	{
 		using namespace UE::AvaPlayableTransition::Private;
 		UE_LOG(LogAvaPlayable, Error,
@@ -246,7 +252,7 @@ bool FAvaPlayableTransitionBuilder::AddExitPlayable(UAvaPlayable* InPlayable)
 		return false;
 	}
 
-	if (PlayingPlayablesWeak.Contains(InPlayable))
+	if (!bInAllowMultipleAdd && PlayingPlayablesWeak.Contains(InPlayable))
 	{
 		using namespace UE::AvaPlayableTransition::Private;
 		UE_LOG(LogAvaPlayable, Error,
diff --git a/Engine/Plugins/Experimental/Avalanche/Source/AvalancheMedia/Private/Playback/Transition/AvaPlaybackServerTransition.cpp b/Engine/Plugins/Experimental/Avalanche/Source/AvalancheMedia/Private/Playback/Transition/AvaPlaybackServerTransition.cpp
--- a/Engine/Plugins/Experimental/Avalanche/Source/AvalancheMedia/Private/Playback/Transition/AvaPlaybackServerTransition.cpp
+++ b/Engine/Plugins/Experimental/Avalanche/Source/AvalancheMedia/Private/Playback/Transition/AvaPlaybackServerTransition.cpp
@@ -509,7 +509,8 @@ void UAvaPlaybackServerTransition::MakePlayableTransition()
 	
 	FAvaPlayableTransitionBuilder TransitionBuilder;
 
-	auto AddInstancesToBuilder = [&TransitionBuilder, this](const TArray<TWeakPtr<FAvaPlaybackInstance>>& InPlaybackInstancesWeak, const TCHAR* InCategory, EAvaPlayableTransitionEntryRole InEntryRole)
+	auto AddInstancesToBuilder = [&TransitionBuilder, this](const TArray<TWeakPtr<FAvaPlaybackInstance>>& InPlaybackInstancesWeak,
+		const TCHAR* InCategory, EAvaPlayableTransitionEntryRole InEntryRole, bool bInAllowMultipleAdd)
 	{
 		using namespace UE::AvaPlaybackServerTransition::Private;
 		int32 ArrayIndex = 0;
@@ -519,7 +520,7 @@ void UAvaPlaybackServerTransition::MakePlayableTransition()
 			{
 				if (UAvaPlayable* Playable = GetPlayable(Instance.Get()))
 				{
-					const bool bPlayableAdded = TransitionBuilder.AddPlayable(Playable, InEntryRole);
+					const bool bPlayableAdded = TransitionBuilder.AddPlayable(Playable, InEntryRole, bInAllowMultipleAdd);
 					if (InEntryRole == EAvaPlayableTransitionEntryRole::Enter && bPlayableAdded)
 					{
 						TransitionBuilder.AddEnterPlayableValues(EnterValues.IsValidIndex(ArrayIndex) ? EnterValues[ArrayIndex] : nullptr);	
@@ -536,10 +537,14 @@ void UAvaPlaybackServerTransition::MakePlayableTransition()
 			++ArrayIndex;
 		}
 	};
+
+	constexpr bool bAllowMultipleAddEnter = false;
+	const bool bAllowMultipleAddPlaying = EnumHasAnyFlags(TransitionFlags, EAvaPlayableTransitionFlags::HasReusedPlayables);
+	constexpr bool bAllowMultipleAddExit = false;
  
-	AddInstancesToBuilder(EnterPlaybackInstancesWeak, TEXT("Enter"), EAvaPlayableTransitionEntryRole::Enter);
-	AddInstancesToBuilder(PlayingPlaybackInstancesWeak, TEXT("Playing"), EAvaPlayableTransitionEntryRole::Playing);
-	AddInstancesToBuilder(ExitPlaybackInstancesWeak, TEXT("Exit"), EAvaPlayableTransitionEntryRole::Exit);
+	AddInstancesToBuilder(EnterPlaybackInstancesWeak, TEXT("Enter"), EAvaPlayableTransitionEntryRole::Enter, bAllowMultipleAddEnter);
+	AddInstancesToBuilder(PlayingPlaybackInstancesWeak, TEXT("Playing"), EAvaPlayableTransitionEntryRole::Playing,  bAllowMultipleAddPlaying);
+	AddInstancesToBuilder(ExitPlaybackInstancesWeak, TEXT("Exit"), EAvaPlayableTransitionEntryRole::Exit,  bAllowMultipleAddExit);
 	PlayableTransition = TransitionBuilder.MakeTransition(this, TransitionId);
 
 	if (PlayableTransition)
diff --git a/Engine/Plugins/Experimental/Avalanche/Source/AvalancheMedia/Private/Rundown/AvaRundown.cpp b/Engine/Plugins/Experimental/Avalanche/Source/AvalancheMedia/Private/Rundown/AvaRundown.cpp
--- a/Engine/Plugins/Experimental/Avalanche/Source/AvalancheMedia/Private/Rundown/AvaRundown.cpp
+++ b/Engine/Plugins/Experimental/Avalanche/Source/AvalancheMedia/Private/Rundown/AvaRundown.cpp
@@ -1789,8 +1789,24 @@ namespace UE::AvaMedia::Rundown::Private
 		return true;
 	}
 
+	bool ArePageRCValuesEqualForSubTemplate(const FAvaRundownPage& InSubTemplate, const FAvaRundownPage& InPage, const UAvaRundownPlaybackInstancePlayer* InInstancePlayer)
+	{
+		if (InInstancePlayer && InInstancePlayer->GetPagePlayer())
+		{
+			if (UAvaRundown* Rundown = InInstancePlayer->GetPagePlayer()->GetRundown())
+			{
+				const FAvaRundownPage PlayingPage = Rundown->GetPage(InInstancePlayer->GetPagePlayer()->PageId);
+				if (PlayingPage.IsValidPage())
+				{
+					return ArePageRCValuesEqualForSubTemplate(InSubTemplate, InPage, PlayingPage);
+				}
+			}
+		}
+		return false;
+	}
+
 	/**
-	 * @brief For special template transition logic, search for an existing instance player with same RC values.
+	 * @brief Search for an existing instance player for the given template and sub-template.
 	 * @param InRundown Rundown
 	 * @param InPageToPlay New Page to be played.
 	 * @param InTemplate Template to be played. Should be direct template of the page.
@@ -1817,17 +1833,17 @@ namespace UE::AvaMedia::Rundown::Private
 				continue;
 			}
 			
-			const FAvaRundownPage& OtherPage = InRundown->GetPage(PagePlayer->PageId);
+			const FAvaRundownPage& PlayingPage = InRundown->GetPage(PagePlayer->PageId);
 
 			// Check if same template.
-			if (!OtherPage.IsValidPage() || OtherPage.GetTemplateId() != InTemplate.GetPageId())
+			if (!PlayingPage.IsValidPage() || PlayingPage.GetTemplateId() != InTemplate.GetPageId())
 			{
 				continue;
 			}
 			
-			// Check if same RC values (of the sub-template).
+			// Find Instance Player for the given sub-template.
 			const FAvaRundownPage& SubTemplate = InTemplate.GetTemplate(InRundown, InSubPageIndex);
-			if (SubTemplate.IsValidPage() && ArePageRCValuesEqualForSubTemplate(SubTemplate, InPageToPlay, OtherPage))
+			if (SubTemplate.IsValidPage())
 			{
 				return PagePlayer->FindInstancePlayerByAssetPath(SubTemplate.GetAssetPath(InRundown));
 			}
@@ -1885,38 +1901,54 @@ bool UAvaRundown::PlayPageWithTransition(FAvaRundownPageTransitionBuilder& InBui
 	const int32 NumTemplates = InPage.GetNumTemplates(this);
 	NewPagePlayer->InstancePlayers.Reserve(NumTemplates);
 
-	const FAvaRundownPage& DirectTemplate = InPage.ResolveTemplate(this);
+	const FAvaRundownPage& Template = InPage.ResolveTemplate(this);
+
+	const UAvaMediaSettings& AvaMediaSettings = UAvaMediaSettings::Get();
+	const bool bBypassTransitionOnSameValues = Template.IsComboTemplate()
+		? AvaMediaSettings.bEnableComboTemplateSpecialLogic
+		: AvaMediaSettings.bEnableSingleTemplateSpecialLogic;
 
+	TSet<FGuid> InstancesBypassingTransition;
+	TSet<FGuid> ReusedInstances;
+	
 	for (int32 SubPageIndex = 0; SubPageIndex < NumTemplates; ++SubPageIndex)
 	{
-		bool bFoundExistingInstancePlayer = false;
+		bool bUsingExistingInstancePlayer = false;
 
-		const UAvaMediaSettings& AvaMediaSettings = UAvaMediaSettings::Get();
-		const bool bUseSpecialTransitionLogic = DirectTemplate.IsComboTemplate()
-			? AvaMediaSettings.bEnableComboTemplateSpecialLogic
-			: AvaMediaSettings.bEnableSingleTemplateSpecialLogic;
-		
-		// -- Special Transition Logic --
-		if (bUseSpecialTransitionLogic)
+		const FAvaRundownPage& SubTemplate = Template.GetTemplate(this, SubPageIndex);
+
+		// -- Logic for Instance Player Reuse --
+		if (bBypassTransitionOnSameValues || SubTemplate.GetTransitionMode(this) == EAvaTransitionInstancingMode::Reuse)
 		{
-			// Try to find an existing instance player of the same combo template, sub-template and RC values.
+			// Try to find an existing instance player for this template.
 			UAvaRundownPlaybackInstancePlayer* InstancePlayer =
-				FindExistingInstancePlayer(this, InPage, DirectTemplate, SubPageIndex, bInIsPreview, InPreviewChannelName);
+				FindExistingInstancePlayer(this, InPage, Template, SubPageIndex, bInIsPreview, InPreviewChannelName);
 			
 			if (InstancePlayer && InstancePlayer->PlaybackInstance)
 			{
-				NewPagePlayer->AddInstancePlayer(InstancePlayer);
+				if (bBypassTransitionOnSameValues && ArePageRCValuesEqualForSubTemplate(SubTemplate, InPage, InstancePlayer))
+				{
+					// Mark this instance as "bypassing" the next playable transition.
+					InstancesBypassingTransition.Add(InstancePlayer->GetPlaybackInstanceId());
+					bUsingExistingInstancePlayer = true;
+				}
+				else if (SubTemplate.GetTransitionMode(this) == EAvaTransitionInstancingMode::Reuse)
+				{
+					ReusedInstances.Add(InstancePlayer->GetPlaybackInstanceId());
+					bUsingExistingInstancePlayer = true;
+				}
 
-				// Setup user instance data to be able to track this page.
-				UAvaRundownPagePlayer::SetInstanceUserDataFromPage(*InstancePlayer->PlaybackInstance, InPage);
-				
-				// Mark this instance as "bypassing" the next playable transition.
-				NewPagePlayer->InstancesBypassingTransition.Add(InstancePlayer->GetPlaybackInstanceId());
-				bFoundExistingInstancePlayer = true;
+				if (bUsingExistingInstancePlayer)
+				{
+					NewPagePlayer->AddInstancePlayer(InstancePlayer);
+
+					// Setup user instance data to be able to track this page.
+					UAvaRundownPagePlayer::SetInstanceUserDataFromPage(*InstancePlayer->PlaybackInstance, InPage);
+				}
 			}
 		}
 
-		if (!bFoundExistingInstancePlayer)
+		if (!bUsingExistingInstancePlayer)
 		{
 			NewPagePlayer->LoadInstancePlayer(SubPageIndex, FGuid());
 		}
@@ -1928,6 +1960,9 @@ bool UAvaRundown::PlayPageWithTransition(FAvaRundownPageTransitionBuilder& InBui
 		{
 			if (PageTransition->AddEnterPage(NewPagePlayer))
 			{
+				PageTransition->InstancesBypassingTransition.Append(InstancesBypassingTransition);
+				PageTransition->ReusedInstances.Append(ReusedInstances);
+
 				AddPagePlayer(NewPagePlayer);
 
 				// Start the playback, will only actually start on next tick.
diff --git a/Engine/Plugins/Experimental/Avalanche/Source/AvalancheMedia/Private/Rundown/AvaRundownPage.cpp b/Engine/Plugins/Experimental/Avalanche/Source/AvalancheMedia/Private/Rundown/AvaRundownPage.cpp
--- a/Engine/Plugins/Experimental/Avalanche/Source/AvalancheMedia/Private/Rundown/AvaRundownPage.cpp
+++ b/Engine/Plugins/Experimental/Avalanche/Source/AvalancheMedia/Private/Rundown/AvaRundownPage.cpp
@@ -205,6 +205,7 @@ bool FAvaRundownPage::UpdateTransitionLogic()
 			const UAvaTransitionTree* TransitionTree = FAvaRundownPageAssetUtils::FindTransitionTree(SceneInterface);
 			bHasTransitionLogic = TransitionTree ? TransitionTree->IsEnabled() : false;
 			TransitionLayerTag = FAvaRundownPageAssetUtils::GetTransitionLayerTag(TransitionTree);
+			TransitionMode = TransitionTree ? TransitionTree->GetInstancingMode() : EAvaTransitionInstancingMode::New;
 		}
 	}
 	return false;
@@ -247,6 +248,31 @@ TArray<FAvaTagHandle> FAvaRundownPage::GetTransitionLayers(const UAvaRundown* In
 	return TransitionLayers;
 }
 
+EAvaTransitionInstancingMode FAvaRundownPage::GetTransitionMode( const UAvaRundown* InRundown, int32 InTemplateIndex) const
+{
+	const FAvaRundownPage& Template = GetTemplate(InRundown, InTemplateIndex);
+	return Template.IsValidPage() ?  Template.TransitionMode : TransitionMode;
+}
+
+TArray<EAvaTransitionInstancingMode> FAvaRundownPage::GetTransitionModes(const UAvaRundown* InRundown) const
+{
+	TArray<EAvaTransitionInstancingMode> TransitionModes;
+
+	const int32 NumTemplates = GetNumTemplates(InRundown);
+	TransitionModes.Reserve(NumTemplates);
+
+	for(int32 TemplateIndex = 0; TemplateIndex < NumTemplates; ++TemplateIndex)
+	{
+		const FAvaRundownPage& Template = GetTemplate(InRundown, TemplateIndex);
+		if (Template.IsValidPage())
+		{
+			TransitionModes.Add(Template.TransitionMode);
+		}
+	}
+
+	return TransitionModes;
+}
+
 int32 FAvaRundownPage::AppendPageProgramStatuses(const UAvaRundown* InParentRundown, TArray<FAvaRundownChannelPageStatus>& OutPageStatuses) const
 {
 	if (!InParentRundown || IsTemplate())
diff --git a/Engine/Plugins/Experimental/Avalanche/Source/AvalancheMedia/Private/Rundown/Transition/AvaRundownPageTransition.cpp b/Engine/Plugins/Experimental/Avalanche/Source/AvalancheMedia/Private/Rundown/Transition/AvaRundownPageTransition.cpp
--- a/Engine/Plugins/Experimental/Avalanche/Source/AvalancheMedia/Private/Rundown/Transition/AvaRundownPageTransition.cpp
+++ b/Engine/Plugins/Experimental/Avalanche/Source/AvalancheMedia/Private/Rundown/Transition/AvaRundownPageTransition.cpp
@@ -273,7 +273,6 @@ void UAvaRundownPageTransition::Stop()
 	{
 		if (UAvaRundownPagePlayer* PagePlayer = PagePlayerWeak.Get())
 		{
-			PagePlayer->InstancesBypassingTransition.Reset();
 			UnregisterEnterPagePlayerEvents(PagePlayer);
 		}
 	}
@@ -467,13 +466,9 @@ void UAvaRundownPageTransition::AddPlayersToBuilder(
 {
 	for (const TWeakObjectPtr<UAvaRundownPagePlayer>& PlayerWeak : InPlayersWeak)
 	{
-		if (UAvaRundownPagePlayer* Player = PlayerWeak.Get())
+		if (const UAvaRundownPagePlayer* Player = PlayerWeak.Get())
 		{
 			AddPlayablesToBuilder(InOutBuilder, Player, InCategory, InEntryRole);
-			
-			// Bypassed instances are only for current transition.
-			// Todo: Ideally, it would be stored in the transition itself to avoid the clean up.
-			Player->InstancesBypassingTransition.Reset();
 		}
 	}
 }
@@ -487,7 +482,7 @@ void UAvaRundownPageTransition::AddPlayablesToBuilder(FAvaPlayableTransitionBuil
 		EAvaPlayableTransitionEntryRole EntryRole = InEntryRole;
 
 		// -- Special Transition Logic --
-		if (InPlayer->InstancesBypassingTransition.Contains(InstancePlayer->GetPlaybackInstanceId()))
+		if (InstancesBypassingTransition.Contains(InstancePlayer->GetPlaybackInstanceId()))
 		{
 			if (EntryRole != EAvaPlayableTransitionEntryRole::Enter)
 			{
@@ -504,6 +499,12 @@ void UAvaRundownPageTransition::AddPlayablesToBuilder(FAvaPlayableTransitionBuil
 			if (EntryRole == EAvaPlayableTransitionEntryRole::Enter && bPlayableAdded)
 			{
 				InOutBuilder.AddEnterPlayableValues(GetRemoteControlValues(InPlayer));
+
+				// For reused instances, we add them again in the playing role.
+				if (ReusedInstances.Contains(InstancePlayer->GetPlaybackInstanceId()))
+				{
+					InOutBuilder.AddPlayable(Playable, EAvaPlayableTransitionEntryRole::Playing, /*bInAllowMultipleAdd*/true);
+				}
 			}
 		}
 		else
@@ -524,10 +525,22 @@ void UAvaRundownPageTransition::MakePlayableTransition()
 	AddPlayersToBuilder(TransitionBuilder, ExitPlayersWeak, TEXT("Exit"), EAvaPlayableTransitionEntryRole::Exit);
 	PlayableTransition = TransitionBuilder.MakeTransition(this, TransitionId);
 
-	// For non-TL enter pages, we need to kick out the playing pages too. 
-	if (PlayableTransition && HasEnterPagesWithNoTransitionLogic())
+	if (PlayableTransition)
 	{
-		PlayableTransition->SetTransitionFlags(EAvaPlayableTransitionFlags::TreatPlayingAsExiting);
+		EAvaPlayableTransitionFlags TransitionFlags = EAvaPlayableTransitionFlags::None;
+
+		// For non-TL enter pages, we need to kick out the playing pages too. 
+		if (HasEnterPagesWithNoTransitionLogic())
+		{
+			TransitionFlags |= EAvaPlayableTransitionFlags::TreatPlayingAsExiting;
+		}
+
+		// Server-side validation needs to know about transitions with in-place playables.
+		if (!ReusedInstances.IsEmpty())
+		{
+			TransitionFlags |= EAvaPlayableTransitionFlags::HasReusedPlayables;
+		}
+		PlayableTransition->SetTransitionFlags(TransitionFlags);
 	}
 }
 
diff --git a/Engine/Plugins/Experimental/Avalanche/Source/AvalancheMedia/Private/Rundown/Transition/AvaRundownPageTransition.h b/Engine/Plugins/Experimental/Avalanche/Source/AvalancheMedia/Private/Rundown/Transition/AvaRundownPageTransition.h
--- a/Engine/Plugins/Experimental/Avalanche/Source/AvalancheMedia/Private/Rundown/Transition/AvaRundownPageTransition.h
+++ b/Engine/Plugins/Experimental/Avalanche/Source/AvalancheMedia/Private/Rundown/Transition/AvaRundownPageTransition.h
@@ -57,6 +57,14 @@ public:
 
 	UAvaRundown* GetRundown() const;
 
+	/** Instances will not be added to the playable transition. */
+	UPROPERTY(Transient)
+	TSet<FGuid> InstancesBypassingTransition;
+
+	/** Reused existing instance player. Will be added both as "entering" and "playing" in the playable transition. */
+	UPROPERTY(Transient)
+	TSet<FGuid> ReusedInstances;
+
 protected:
 	UAvaRundownPlaybackInstancePlayer* FindInstancePlayerForPlayable(const UAvaPlayable* InPlayable) const;
 	UAvaRundownPagePlayer* FindPagePlayerForPlayable(const UAvaPlayable* InPlayable) const;
diff --git a/Engine/Plugins/Experimental/Avalanche/Source/AvalancheMedia/Public/AvaMediaDefines.h b/Engine/Plugins/Experimental/Avalanche/Source/AvalancheMedia/Public/AvaMediaDefines.h
--- a/Engine/Plugins/Experimental/Avalanche/Source/AvalancheMedia/Public/AvaMediaDefines.h
+++ b/Engine/Plugins/Experimental/Avalanche/Source/AvalancheMedia/Public/AvaMediaDefines.h
@@ -126,6 +126,8 @@ enum class EAvaPlayableTransitionFlags : uint8
 	None = 0,
 	/** Playing playables will be treated as exit playables. */
 	TreatPlayingAsExiting = 1 << 0,
+	/** Transition contains some reused playables (i.e. both enter and playing). */
+	HasReusedPlayables = 1 << 1,
 };
 ENUM_CLASS_FLAGS(EAvaPlayableTransitionFlags);
 
diff --git a/Engine/Plugins/Experimental/Avalanche/Source/AvalancheMedia/Public/Playable/Transition/AvaPlayableTransition.h b/Engine/Plugins/Experimental/Avalanche/Source/AvalancheMedia/Public/Playable/Transition/AvaPlayableTransition.h
--- a/Engine/Plugins/Experimental/Avalanche/Source/AvalancheMedia/Public/Playable/Transition/AvaPlayableTransition.h
+++ b/Engine/Plugins/Experimental/Avalanche/Source/AvalancheMedia/Public/Playable/Transition/AvaPlayableTransition.h
@@ -82,20 +82,20 @@ public:
 
 	void AddEnterPlayableValues(const TSharedPtr<FAvaPlayableRemoteControlValues>& InValues);
 	
-	bool AddEnterPlayable(UAvaPlayable* InPlayable);
-	bool AddPlayingPlayable(UAvaPlayable* InPlayable);
-	bool AddExitPlayable(UAvaPlayable* InPlayable);
+	bool AddEnterPlayable(UAvaPlayable* InPlayable, bool bInAllowMultipleAdd = false);
+	bool AddPlayingPlayable(UAvaPlayable* InPlayable, bool bInAllowMultipleAdd = false);
+	bool AddExitPlayable(UAvaPlayable* InPlayable, bool bInAllowMultipleAdd = false);
 	
-	bool AddPlayable(UAvaPlayable* InPlayable, EAvaPlayableTransitionEntryRole InPlayableRole)
+	bool AddPlayable(UAvaPlayable* InPlayable, EAvaPlayableTransitionEntryRole InPlayableRole, bool bInAllowMultipleAdd = false)
 	{
 		switch (InPlayableRole)
 		{
 			case EAvaPlayableTransitionEntryRole::Enter:
-				return AddEnterPlayable(InPlayable);
+				return AddEnterPlayable(InPlayable, bInAllowMultipleAdd);
 			case EAvaPlayableTransitionEntryRole::Playing:
-				return AddPlayingPlayable(InPlayable);
+				return AddPlayingPlayable(InPlayable, bInAllowMultipleAdd);
 			case EAvaPlayableTransitionEntryRole::Exit:
-				return AddExitPlayable(InPlayable);
+				return AddExitPlayable(InPlayable, bInAllowMultipleAdd);
 		}
 		return false;
 	}
diff --git a/Engine/Plugins/Experimental/Avalanche/Source/AvalancheMedia/Public/Rundown/AvaRundownPage.h b/Engine/Plugins/Experimental/Avalanche/Source/AvalancheMedia/Public/Rundown/AvaRundownPage.h
--- a/Engine/Plugins/Experimental/Avalanche/Source/AvalancheMedia/Public/Rundown/AvaRundownPage.h
+++ b/Engine/Plugins/Experimental/Avalanche/Source/AvalancheMedia/Public/Rundown/AvaRundownPage.h
@@ -4,6 +4,7 @@
 
 #include "AvaMediaDefines.h"
 #include "AvaTagHandle.h"
+#include "AvaTransitionEnums.h"
 #include "Playable/AvaPlayableRemoteControlValues.h"
 #include "AvaRundownPage.generated.h"
 
@@ -114,6 +115,9 @@ public:
 	
 	FAvaTagHandle GetTransitionLayer(const UAvaRundown* InRundown, int32 InTemplateIndex = 0) const;
 	TArray<FAvaTagHandle> GetTransitionLayers(const UAvaRundown* InRundown) const;
+
+	EAvaTransitionInstancingMode GetTransitionMode( const UAvaRundown* InRundown, int32 InTemplateIndex = 0) const;
+	TArray<EAvaTransitionInstancingMode> GetTransitionModes(const UAvaRundown* InRundown) const;
 	
 	/**
 	 * Appends the page's program status(es).
@@ -280,4 +284,7 @@ protected:
 	/** Transition Layer Tag cached from the transition tree. Cached for fast display in page/template list. */
 	UPROPERTY(VisibleAnywhere, Category = "Motion Design")
 	FAvaTagHandle TransitionLayerTag;
+
+	UPROPERTY(VisibleAnywhere, BlueprintReadOnly, Category = "Motion Design")
+	EAvaTransitionInstancingMode TransitionMode = EAvaTransitionInstancingMode::New;
 };
diff --git a/Engine/Plugins/Experimental/Avalanche/Source/AvalancheMedia/Public/Rundown/AvaRundownPagePlayer.h b/Engine/Plugins/Experimental/Avalanche/Source/AvalancheMedia/Public/Rundown/AvaRundownPagePlayer.h
--- a/Engine/Plugins/Experimental/Avalanche/Source/AvalancheMedia/Public/Rundown/AvaRundownPagePlayer.h
+++ b/Engine/Plugins/Experimental/Avalanche/Source/AvalancheMedia/Public/Rundown/AvaRundownPagePlayer.h
@@ -49,7 +49,7 @@ public:
 
 	void SetPagePlayer(UAvaRundownPagePlayer* InPagePlayer);
 
-public:	
+public:
 	UPROPERTY()
 	FAvaTagHandle TransitionLayer;
 
@@ -181,10 +181,6 @@ public:
 	UPROPERTY()
 	TArray<TObjectPtr<UAvaRundownPlaybackInstancePlayer>> InstancePlayers;
 
-	/** Instances that should bypass the next transition. */
-	UPROPERTY(Transient)
-	TSet<FGuid> InstancesBypassingTransition;
-
 protected:
 	TWeakObjectPtr<UAvaRundown> RundownWeak;
 };
