[Motion Design] Rundown Editor: Fix Page View Search Filtering

Moving the filtering in the RundownPageList class. It is cleaner and will correctly filter the page list for the given widget.

diff --git a/Engine/Plugins/Experimental/Avalanche/Source/AvalancheMediaEditor/Private/Rundown/AvaRundownEditor.cpp b/Engine/Plugins/Experimental/Avalanche/Source/AvalancheMediaEditor/Private/Rundown/AvaRundownEditor.cpp
--- a/Engine/Plugins/Experimental/Avalanche/Source/AvalancheMediaEditor/Private/Rundown/AvaRundownEditor.cpp
+++ b/Engine/Plugins/Experimental/Avalanche/Source/AvalancheMediaEditor/Private/Rundown/AvaRundownEditor.cpp
@@ -25,7 +25,6 @@
 #include "Rundown/AvaRundownPagePlayer.h"
 #include "Rundown/AvaRundownPlaybackUtils.h"
 #include "Rundown/Factories/Filters/AvaRundownFactoriesUtils.h"
-#include "Rundown/Filters/AvaRundownPageTextFilter.h"
 #include "Rundown/Pages/PageViews/IAvaRundownPageView.h"
 #include "Rundown/Pages/Slate/SAvaRundownPageList.h"
 #include "ScopedTransaction.h"
@@ -183,8 +182,6 @@ private:
 };
 
 FAvaRundownEditor::FAvaRundownEditor()
-	: TextFilterTemplatePage(MakeShared<FAvaRundownPageTextFilter>())
-	, TextFilterInstancedPage(MakeShared<FAvaRundownPageTextFilter>())
 {}
 
 FAvaRundownEditor::~FAvaRundownEditor()
@@ -233,12 +230,6 @@ void FAvaRundownEditor::InitRundownEditor(const EToolkitMode::Type InMode
 
 	AvaRundown = InRundown;
 
-	InitVisibilityTemplatePages();
-	InitVisibilityInstancedPages();
-
-	TextFilterTemplatePage->OnChanged().AddSP(this, &FAvaRundownEditor::OnTemplateFilterChanged);
-	TextFilterInstancedPage->OnChanged().AddSP(this, &FAvaRundownEditor::OnInstancedFilterChanged);
-
 	CreateRundownCommands();
 
 	if (IsValid(InRundown))
@@ -809,67 +800,6 @@ void FAvaRundownEditor::RemoveSelectedPages()
 	}
 }
 
-void FAvaRundownEditor::RefreshTemplateVisibility()
-{
-	VisibleTemplatePageIds.Reset();
-
-	if (AvaRundown.IsValid())
-	{
-		for (const FAvaRundownPage& Page : AvaRundown->GetTemplatePages().Pages)
-		{
-			TextFilterTemplatePage->SetItem(Page, AvaRundown.Get(), EAvaRundownSearchListType::Template);
-			if (TextFilterTemplatePage->PassesFilter(Page))
-			{
-				VisibleTemplatePageIds.Add(Page.GetPageId());
-			}
-		}
-	}
-}
-
-void FAvaRundownEditor::RefreshInstancedVisibility()
-{
-	VisibleInstancedPageIds.Reset();
-
-	if (AvaRundown.IsValid())
-	{
-		for (const FAvaRundownPage& Page : AvaRundown->GetInstancedPages().Pages)
-		{
-			TextFilterInstancedPage->SetItem(Page, AvaRundown.Get(), EAvaRundownSearchListType::Instanced);
-			if (TextFilterInstancedPage->PassesFilter(Page))
-			{
-				VisibleInstancedPageIds.Add(Page.GetPageId());
-			}
-		}
-	}
-}
-
-bool FAvaRundownEditor::IsTemplatePageVisible(const FAvaRundownPage& InPage) const
-{
-	return VisibleTemplatePageIds.Contains(InPage.GetPageId());
-}
-
-bool FAvaRundownEditor::IsInstancedPageVisible(const FAvaRundownPage& InPage) const
-{
-	return VisibleInstancedPageIds.Contains(InPage.GetPageId());
-}
-
-void FAvaRundownEditor::SetSearchText(const FText& InText, EAvaRundownSearchListType& InPageListType)
-{
-	switch (InPageListType)
-	{
-	case EAvaRundownSearchListType::Template:
-		SetTemplateSearchText(InText);
-		break;
-	case EAvaRundownSearchListType::Instanced:
-		SetInstancedSearchText(InText);
-		break;
-
-	case EAvaRundownSearchListType::None:
-	default:
-		break;
-	}
-}
-
 FName FAvaRundownEditor::GetToolkitFName() const
 {
 	return TEXT("AvaRundownEditor");
@@ -1550,58 +1480,6 @@ const FAvaRundownEditor::FBindableMacroCommands& FAvaRundownEditor::GetBindableM
 	return BindableMacroCommands;
 }
 
-void FAvaRundownEditor::SetTemplateSearchText(const FText& InText)
-{
-	TextFilterTemplatePage->SetFilterText(InText);
-}
-
-void FAvaRundownEditor::SetInstancedSearchText(const FText& InText)
-{
-	TextFilterInstancedPage->SetFilterText(InText);
-}
-
-void FAvaRundownEditor::OnTemplateFilterChanged()
-{
-	RefreshTemplateVisibility();
-
-	if (GetTemplateListWidget().IsValid())
-	{
-		GetTemplateListWidget()->Refresh();
-	}
-}
-
-void FAvaRundownEditor::OnInstancedFilterChanged()
-{
-	RefreshInstancedVisibility();
-
-	if (GetInstanceListWidget().IsValid())
-	{
-		GetInstanceListWidget()->Refresh();
-	}
-}
-
-void FAvaRundownEditor::InitVisibilityTemplatePages()
-{
-	if (const UAvaRundown* Rundown = GetRundown())
-	{
-		for (const FAvaRundownPage& Page : Rundown->GetTemplatePages().Pages)
-		{
-			VisibleTemplatePageIds.Add(Page.GetPageId());
-		}
-	}
-}
-
-void FAvaRundownEditor::InitVisibilityInstancedPages()
-{
-	if (const UAvaRundown* Rundown = GetRundown())
-	{
-		for (const FAvaRundownPage& Page : Rundown->GetInstancedPages().Pages)
-		{
-			VisibleInstancedPageIds.Add(Page.GetPageId());
-		}
-	}
-}
-
 FAvaRundownEditor::FAutoPlayTicker::FAutoPlayTicker(TWeakPtr<FAvaRundownEditor> InRundownEditorWeak, double InTickInterval)
 	: RundownEditorWeak(InRundownEditorWeak)
 	, PlayInterval(InTickInterval)
diff --git a/Engine/Plugins/Experimental/Avalanche/Source/AvalancheMediaEditor/Private/Rundown/Pages/Slate/SAvaRundownInstancedPageList.cpp b/Engine/Plugins/Experimental/Avalanche/Source/AvalancheMediaEditor/Private/Rundown/Pages/Slate/SAvaRundownInstancedPageList.cpp
--- a/Engine/Plugins/Experimental/Avalanche/Source/AvalancheMediaEditor/Private/Rundown/Pages/Slate/SAvaRundownInstancedPageList.cpp
+++ b/Engine/Plugins/Experimental/Avalanche/Source/AvalancheMediaEditor/Private/Rundown/Pages/Slate/SAvaRundownInstancedPageList.cpp
@@ -11,7 +11,6 @@
 #include "Rundown/AvaRundownPage.h"
 #include "Rundown/AvaRundownPlaybackUtils.h"
 #include "Rundown/AvaRundownSerializationUtils.h"
-#include "Rundown/Factories/Filters/AvaRundownFactoriesUtils.h"
 #include "Rundown/Pages/Columns/AvaRundownPageAssetNameColumn.h"
 #include "Rundown/Pages/Columns/AvaRundownPageChannelSelectorColumn.h"
 #include "Rundown/Pages/Columns/AvaRundownPageEnabledColumn.h"
@@ -38,7 +37,7 @@ void SAvaRundownInstancedPageList::PrivateRegisterAttributes(struct FSlateAttrib
 
 void SAvaRundownInstancedPageList::Construct(const FArguments& InArgs, TSharedPtr<FAvaRundownEditor> InRundownEditor, const FAvaRundownPageListReference& InPageListReference)
 {
-	SAvaRundownPageList::Construct(SAvaRundownPageList::FArguments(), InRundownEditor, InPageListReference, EAvaRundownSearchListType::Instanced);
+	SAvaRundownPageList::Construct(SAvaRundownPageList::FArguments(), InRundownEditor, InPageListReference);
 
 	UAvaRundown* const Rundown = InRundownEditor->GetRundown();
 	check(Rundown);
@@ -118,30 +117,33 @@ void SAvaRundownInstancedPageList::Refresh()
 		UAvaRundown* const Rundown = RundownEditor->GetRundown();
 		check(Rundown);
 
-		TArray<FAvaRundownPage> Pages = Rundown->GetInstancedPages().Pages;
-
-		if (Rundown->IsValidSubList(PageListReference))
+		if (PageListReference.Type == EAvaRundownPageListType::Instance)
 		{
-			Pages.Empty();
+			const TArray<FAvaRundownPage>& Pages = Rundown->GetInstancedPages().Pages;
+			
+			PageViews.Reset(Pages.Num());
 
-			for (const int32 PageId : Rundown->GetSubList(PageListReference.SubListIndex).PageIds)
+			for (const FAvaRundownPage& Page : Pages)
 			{
-				if (const int32* Index = Rundown->GetInstancedPages().PageIndices.Find(PageId))
+				if (IsPageVisible(Page))
 				{
-					Pages.Add(Rundown->GetInstancedPages().Pages[*Index]);
+					PageViews.Emplace(MakeShared<FAvaRundownInstancedPageViewImpl>(Page.GetPageId(), Rundown, SharedThis(this)));
 				}
 			}
 		}
-
-		PageViews.Reset(Pages.Num());
-
-		for (const FAvaRundownPage& Page : Pages)
+		else if (Rundown->IsValidSubList(PageListReference))
 		{
-			if (RundownEditor->IsInstancedPageVisible(Page))
+			const FAvaRundownSubList& SubList = Rundown->GetSubList(PageListReference.SubListIndex);
+			PageViews.Reset(SubList.PageIds.Num());
+			
+			for (const int32 PageId : SubList.PageIds)
 			{
-				PageViews.Emplace(MakeShared<FAvaRundownInstancedPageViewImpl>(Page.GetPageId(), Rundown, SharedThis(this)));
+				if (IsPageVisible(PageId))
+				{
+					PageViews.Emplace(MakeShared<FAvaRundownInstancedPageViewImpl>(PageId, Rundown, SharedThis(this)));
+				}
 			}
-		}
+		}	
 
 		PageListView->RequestListRefresh();
 	}
@@ -1145,10 +1147,7 @@ void SAvaRundownInstancedPageList::OnPageListChanged(const FAvaRundownPageListCh
 		return;
 	}
 	
-	if (const TSharedPtr<FAvaRundownEditor> RundownEditor = RundownEditorWeak.Pin())
-	{
-		RundownEditor->RefreshInstancedVisibility();
-	}
+	RefreshPagesVisibility();
 	Refresh();
 }
 
diff --git a/Engine/Plugins/Experimental/Avalanche/Source/AvalancheMediaEditor/Private/Rundown/Pages/Slate/SAvaRundownPageList.cpp b/Engine/Plugins/Experimental/Avalanche/Source/AvalancheMediaEditor/Private/Rundown/Pages/Slate/SAvaRundownPageList.cpp
--- a/Engine/Plugins/Experimental/Avalanche/Source/AvalancheMediaEditor/Private/Rundown/Pages/Slate/SAvaRundownPageList.cpp
+++ b/Engine/Plugins/Experimental/Avalanche/Source/AvalancheMediaEditor/Private/Rundown/Pages/Slate/SAvaRundownPageList.cpp
@@ -20,6 +20,7 @@
 #include "Rundown/AvaRundownSerializationUtils.h"
 #include "Rundown/Factories/Filters/AvaRundownFactoriesUtils.h"
 #include "Rundown/Factories/Filters/IAvaRundownFilterSuggestionFactory.h"
+#include "Rundown/Filters/AvaRundownPageTextFilter.h"
 #include "Rundown/Pages/AvaRundownPageContext.h"
 #include "Rundown/Pages/AvaRundownPageContextMenu.h"
 #include "Rundown/Pages/PageViews/AvaRundownPageViewImpl.h"
@@ -120,6 +121,28 @@ namespace UE::AvaRundown::Private
 			}
 		}
 	}
+
+	const FAvaRundownPageCollection& GetPageCollection(const UAvaRundown* InRundown, EAvaRundownPageListType InPageListType)
+	{
+		check(InRundown);
+		return InPageListType == EAvaRundownPageListType::Template ? InRundown->GetTemplatePages() : InRundown->GetInstancedPages();
+	}
+	
+	EAvaRundownSearchListType GetPageSearchListType(EAvaRundownPageListType InPageListType)
+	{
+		switch(InPageListType)
+		{
+		case EAvaRundownPageListType::Instance:
+		case EAvaRundownPageListType::View:
+			return EAvaRundownSearchListType::Instanced;
+			
+		case EAvaRundownPageListType::Template:
+			return EAvaRundownSearchListType::Template;
+			
+		default:
+			return EAvaRundownSearchListType::None;
+		}
+	}
 }
 
 void SAvaRundownPageList::PrivateRegisterAttributes(struct FSlateAttributeDescriptor::FInitializer&)
@@ -127,7 +150,7 @@ void SAvaRundownPageList::PrivateRegisterAttributes(struct FSlateAttributeDescri
 
 }
 
-void SAvaRundownPageList::Construct(const FArguments& InArgs, TSharedPtr<FAvaRundownEditor> InRundownEditor, const FAvaRundownPageListReference& InPageListReference, EAvaRundownSearchListType InPageListType)
+void SAvaRundownPageList::Construct(const FArguments& InArgs, TSharedPtr<FAvaRundownEditor> InRundownEditor, const FAvaRundownPageListReference& InPageListReference)
 {
 	RundownEditorWeak = InRundownEditor;
 	check(InRundownEditor.IsValid());
@@ -135,10 +158,13 @@ void SAvaRundownPageList::Construct(const FArguments& InArgs, TSharedPtr<FAvaRun
 	PageContextMenu = MakeShared<FAvaRundownPageContextMenu>();
 
 	PageListReference = InPageListReference;
-	PageListType = InPageListType;
 	CommandList = MakeShared<FUICommandList>();
 	BindCommands();
 
+	PageTextFilter = MakeShared<FAvaRundownPageTextFilter>();
+	PageTextFilter->OnChanged().AddSP(this, &SAvaRundownPageList::OnPageTextFilterChanged);
+	RefreshPagesVisibility();
+	
 	UAvaRundown* const Rundown = InRundownEditor->GetRundown();
 	check(Rundown);
 
@@ -1108,6 +1134,56 @@ FReply SAvaRundownPageList::OnKeyDown(const FGeometry& InMyGeometry, const FKeyE
 	return FReply::Unhandled();
 }
 
+void SAvaRundownPageList::RefreshPagesVisibility()
+{
+	VisiblePageIds.Reset();
+
+	const UAvaRundown* Rundown = GetValidRundown();
+	if (!Rundown)
+	{
+		return;
+	}
+
+	const EAvaRundownSearchListType PageSearchListType = UE::AvaRundown::Private::GetPageSearchListType(PageListReference.Type);
+
+	auto FilterPage = [this, Rundown, PageSearchListType](const FAvaRundownPage& InPage)
+	{
+		if (PageTextFilter)
+		{
+			PageTextFilter->SetItem(InPage, Rundown, PageSearchListType);
+			if (PageTextFilter->PassesFilter(InPage))
+			{
+				VisiblePageIds.Add(InPage.GetPageId());
+			}
+		}
+		else
+		{
+			VisiblePageIds.Add(InPage.GetPageId());
+		}
+	};
+
+	if (Rundown->IsValidSubList(PageListReference))
+	{
+		const FAvaRundownSubList& SubList = Rundown->GetSubList(PageListReference.SubListIndex);
+		for (const int32 PageId : SubList.PageIds)
+		{
+			const FAvaRundownPage& Page = Rundown->GetPage(PageId);
+			if (Page.IsValidPage())
+			{
+				FilterPage(Page);
+			}
+		}
+	}
+	else
+	{
+		const FAvaRundownPageCollection& Collection = UE::AvaRundown::Private::GetPageCollection(Rundown, PageListReference.Type);
+		for (const FAvaRundownPage& Page : Collection.Pages)
+		{
+			FilterPage(Page);
+		}
+	}
+}
+
 FText SAvaRundownPageList::OnAssetSearchBoxSuggestionChosen(const FText& InSearchText, const FString& InSuggestion)
 {
 	int32 SuggestionInsertionIndex = 0;
@@ -1121,17 +1197,17 @@ FText SAvaRundownPageList::OnAssetSearchBoxSuggestionChosen(const FText& InSearc
 
 void SAvaRundownPageList::OnSearchTextChanged(const FText& FilterText)
 {
-	if (TSharedPtr<FAvaRundownEditor> RundownEditor = RundownEditorWeak.Pin())
+	if (PageTextFilter)
 	{
-		RundownEditor->SetSearchText(FilterText, PageListType);
+		PageTextFilter->SetFilterText(FilterText);
 	}
 }
 
 void SAvaRundownPageList::OnSearchTextCommitted(const FText& FilterText, ETextCommit::Type CommitType)
 {
-	if (TSharedPtr<FAvaRundownEditor> RundownEditor = RundownEditorWeak.Pin())
+	if (PageTextFilter)
 	{
-		RundownEditor->SetSearchText(FilterText, PageListType);
+		PageTextFilter->SetFilterText(FilterText);
 	}
 }
 
@@ -1145,6 +1221,7 @@ void SAvaRundownPageList::OnSearchBoxSuggestionFilter(const FText& InSearchText,
 	UE::AvaRundown::Private::ExtractAssetSearchFilterTerms(InSearchText, &FilterKey, &FilterValue, nullptr);
 
 	const IAvaMediaEditorModule& AvaMediaEditorModule = IAvaMediaEditorModule::Get();
+	const EAvaRundownSearchListType PageListType = UE::AvaRundown::Private::GetPageSearchListType(PageListReference.Type);
 
 	TSet<FString> FilterCache;
 	const TSharedRef<FAvaRundownFilterSuggestionPayload> SimplePayload =
@@ -1161,12 +1238,27 @@ void SAvaRundownPageList::OnSearchBoxSuggestionFilter(const FText& InSearchText,
 			const TSharedRef<FAvaRundownFilterSuggestionPayload> ComplexPayload =
 				MakeShared<FAvaRundownFilterSuggestionPayload>(FAvaRundownFilterSuggestionPayload{ OutPossibleSuggestions, FilterValue, FAvaRundownPage::InvalidPageId, Rundown, FilterCache });
 
-			for (const FAvaRundownPage& Page : GetPagesByType(PageListType))
+			if (Rundown->IsValidSubList(PageListReference))
 			{
-				for (const TSharedPtr<IAvaRundownFilterSuggestionFactory>& ComplexSuggestion : AvaMediaEditorModule.GetComplexSuggestions(PageListType))
+				for (const int32 PageId : Rundown->GetSubList(PageListReference.SubListIndex).PageIds)
 				{
-					ComplexPayload->ItemPageId = Page.GetPageId();
-					ComplexSuggestion->AddSuggestion(ComplexPayload);
+					for (const TSharedPtr<IAvaRundownFilterSuggestionFactory>& ComplexSuggestion : AvaMediaEditorModule.GetComplexSuggestions(PageListType))
+					{
+						ComplexPayload->ItemPageId = PageId;
+						ComplexSuggestion->AddSuggestion(ComplexPayload);
+					}
+				}
+			}
+			else
+			{
+				using namespace UE::AvaRundown::Private;
+				for (const FAvaRundownPage& Page : GetPageCollection(Rundown, PageListReference.Type).Pages)
+				{
+					for (const TSharedPtr<IAvaRundownFilterSuggestionFactory>& ComplexSuggestion : AvaMediaEditorModule.GetComplexSuggestions(PageListType))
+					{
+						ComplexPayload->ItemPageId = Page.GetPageId();
+						ComplexSuggestion->AddSuggestion(ComplexPayload);
+					}
 				}
 			}
 		}
@@ -1174,32 +1266,10 @@ void SAvaRundownPageList::OnSearchBoxSuggestionFilter(const FText& InSearchText,
 	OutSuggestionHighlightText = FText::FromString(FilterValue);
 }
 
-TArray<FAvaRundownPage> SAvaRundownPageList::GetPagesByType(EAvaRundownSearchListType InRundownSearchListType) const
+void SAvaRundownPageList::OnPageTextFilterChanged()
 {
-	TArray<FAvaRundownPage> Pages = TArray<FAvaRundownPage>();
-
-	if (const TSharedPtr<FAvaRundownEditor> RundownEditor = RundownEditorWeak.Pin())
-	{
-		if (const UAvaRundown* Rundown = RundownEditor->GetRundown())
-		{
-			switch (InRundownSearchListType)
-			{
-			case EAvaRundownSearchListType::Template:
-				Pages = Rundown->GetTemplatePages().Pages;
-				break;
-
-			case EAvaRundownSearchListType::Instanced:
-				Pages = Rundown->GetInstancedPages().Pages;
-				break;
-
-			case EAvaRundownSearchListType::None:
-			default:
-				break;
-			}
-		}
-	}
-
-	return Pages;
+	RefreshPagesVisibility();
+	Refresh();
 }
 
 TArray<int32> SAvaRundownPageList::FilterSelectedPages(FFilterPageFunctionRef InFilterPageFunction) const
diff --git a/Engine/Plugins/Experimental/Avalanche/Source/AvalancheMediaEditor/Private/Rundown/Pages/Slate/SAvaRundownPageList.h b/Engine/Plugins/Experimental/Avalanche/Source/AvalancheMediaEditor/Private/Rundown/Pages/Slate/SAvaRundownPageList.h
--- a/Engine/Plugins/Experimental/Avalanche/Source/AvalancheMediaEditor/Private/Rundown/Pages/Slate/SAvaRundownPageList.h
+++ b/Engine/Plugins/Experimental/Avalanche/Source/AvalancheMediaEditor/Private/Rundown/Pages/Slate/SAvaRundownPageList.h
@@ -10,19 +10,17 @@
 
 class FAvaRundownEditor;
 class FAvaRundownPageContextMenu;
+class FAvaRundownPageTextFilter;
 class FUICommandList;
 class IAvaRundownPageViewColumn;
 class ITableRow;
 class SAssetSearchBox;
 class SHeaderRow;
 class SHorizontalBox;
-class SSearchBox;
 class STableViewBase;
-enum class EAvaRundownSearchListType : uint8;
 enum class EItemDropZone;
 struct FAssetSearchBoxSuggestion;
 struct FAvaRundownPage;
-struct FAvaRundownPageListChangeParams;
 template<typename ItemType> class SListView;
 
 class SAvaRundownPageList : public SCompoundWidget
@@ -34,7 +32,7 @@ public:
 	SLATE_END_ARGS()
 
 	/** Constructs this widget with InArgs */
-	void Construct(const FArguments& InArgs, TSharedPtr<FAvaRundownEditor> InRundownEditor, const FAvaRundownPageListReference& InPageListReference, EAvaRundownSearchListType InPageListType);
+	void Construct(const FArguments& InArgs, TSharedPtr<FAvaRundownEditor> InRundownEditor, const FAvaRundownPageListReference& InPageListReference);
 	virtual ~SAvaRundownPageList() override;
 
 	const FAvaRundownPageListReference& GetPageListReference() const { return PageListReference; }
@@ -81,8 +79,6 @@ public:
 
 	TSharedRef<SWidget> GetPageListContextMenu();
 
-	EAvaRundownSearchListType GetPageListType() const { return PageListType; }
-
 	bool CanCopySelectedPages() const;
 	void CopySelectedPages();
 
@@ -180,6 +176,19 @@ protected:
 	TArray<int32> GetPagesToPreviewContinue() const;
 	TArray<int32> GetPagesToTakeToProgram() const;
 	int32 GetPageIdToPreviewNext() const;
+
+	bool IsPageVisible(const FAvaRundownPage& InPage) const
+	{
+		return VisiblePageIds.Contains(InPage.GetPageId());
+	}
+
+	bool IsPageVisible(int32 InPageId) const
+	{
+		return VisiblePageIds.Contains(InPageId);
+	}
+
+	/** Applies PageTextFilter to current page list and caches the result. */
+	void RefreshPagesVisibility();
 	
 private:
 	FText OnAssetSearchBoxSuggestionChosen(const FText& InSearchText, const FString& InSuggestion);
@@ -190,7 +199,7 @@ private:
 
 	void OnSearchBoxSuggestionFilter(const FText& InSearchText, TArray<FAssetSearchBoxSuggestion>& OutPossibleSuggestions, FText& OutSuggestionHighlightText);
 
-	TArray<FAvaRundownPage> GetPagesByType(EAvaRundownSearchListType InRundownSearchListType) const;
+	void OnPageTextFilterChanged();
 
 protected:
 	TWeakPtr<FAvaRundownEditor> RundownEditorWeak;
@@ -213,6 +222,11 @@ protected:
 
 	TSharedPtr<FUICommandList> CommandList;
 
+	TSharedPtr<FAvaRundownPageTextFilter> PageTextFilter;
+
+	/** Cached result of PageTextFilter applied to current page list. */
+	TSet<int32> VisiblePageIds;
+	
 	/**
 	 * For template and instance lists, it will return the list of page ids created.
 	 * For a sub list it will return the list of page ids that were added to the page view.
@@ -221,6 +235,4 @@ protected:
 
 private:
 	TSharedPtr<SAssetSearchBox> AssetSearchBoxPtr;
-
-	EAvaRundownSearchListType PageListType;
 };
diff --git a/Engine/Plugins/Experimental/Avalanche/Source/AvalancheMediaEditor/Private/Rundown/Pages/Slate/SAvaRundownTemplatePageList.cpp b/Engine/Plugins/Experimental/Avalanche/Source/AvalancheMediaEditor/Private/Rundown/Pages/Slate/SAvaRundownTemplatePageList.cpp
--- a/Engine/Plugins/Experimental/Avalanche/Source/AvalancheMediaEditor/Private/Rundown/Pages/Slate/SAvaRundownTemplatePageList.cpp
+++ b/Engine/Plugins/Experimental/Avalanche/Source/AvalancheMediaEditor/Private/Rundown/Pages/Slate/SAvaRundownTemplatePageList.cpp
@@ -2,15 +2,12 @@
 
 #include "SAvaRundownTemplatePageList.h"
 
-#include "IAvaMediaEditorModule.h"
-#include "IAvaMediaModule.h"
 #include "Rundown/AvaRundown.h"
 #include "Rundown/AvaRundownCommands.h"
 #include "Rundown/AvaRundownEditor.h"
 #include "Rundown/AvaRundownEditorUtils.h"
 #include "Rundown/AvaRundownManagedInstanceCache.h"
 #include "Rundown/AvaRundownPage.h"
-#include "Rundown/Factories/Filters/AvaRundownFactoriesUtils.h"
 #include "Rundown/Pages/Columns/AvaRundownPageAssetSelectorColumn.h"
 #include "Rundown/Pages/Columns/AvaRundownPageIdColumn.h"
 #include "Rundown/Pages/Columns/AvaRundownPageNameColumn.h"
@@ -31,7 +28,7 @@ void SAvaRundownTemplatePageList::PrivateRegisterAttributes(struct FSlateAttribu
 
 void SAvaRundownTemplatePageList::Construct(const FArguments& InArgs, TSharedPtr<FAvaRundownEditor> InRundownEditor)
 {
-	SAvaRundownPageList::Construct(SAvaRundownPageList::FArguments(), InRundownEditor, UAvaRundown::TemplatePageList, EAvaRundownSearchListType::Template);
+	SAvaRundownPageList::Construct(SAvaRundownPageList::FArguments(), InRundownEditor, UAvaRundown::TemplatePageList);
 
 	RundownEditorWeak = InRundownEditor;
 	check(InRundownEditor.IsValid());
@@ -67,7 +64,7 @@ void SAvaRundownTemplatePageList::Refresh()
 
 		for (const FAvaRundownPage& Page : PageCollection.Pages)
 		{
-			if (RundownEditor->IsTemplatePageVisible(Page))
+			if (IsPageVisible(Page))
 			{
 				PageViews.Emplace(MakeShared<FAvaRundownTemplatePageViewImpl>(Page.GetPageId(), Rundown, SharedThis(this)));
 			}
@@ -505,10 +502,7 @@ void SAvaRundownTemplatePageList::OnPageListChanged(const FAvaRundownPageListCha
 		return;
 	}
 
-	if (const TSharedPtr<FAvaRundownEditor> RundownEditor = RundownEditorWeak.Pin())
-	{
-		RundownEditor->RefreshTemplateVisibility();
-	}
+	RefreshPagesVisibility();
 	Refresh();
 }
 
diff --git a/Engine/Plugins/Experimental/Avalanche/Source/AvalancheMediaEditor/Public/Rundown/AvaRundownEditor.h b/Engine/Plugins/Experimental/Avalanche/Source/AvalancheMediaEditor/Public/Rundown/AvaRundownEditor.h
--- a/Engine/Plugins/Experimental/Avalanche/Source/AvalancheMediaEditor/Public/Rundown/AvaRundownEditor.h
+++ b/Engine/Plugins/Experimental/Avalanche/Source/AvalancheMediaEditor/Public/Rundown/AvaRundownEditor.h
@@ -9,7 +9,6 @@
 
 class FAvaRundownEditorInputProcessor;
 class FAvaRundownManagedInstance;
-class FAvaRundownPageTextFilter;
 class SAvaRundownInstancedPageList;
 class SAvaRundownPageList;
 class SAvaRundownReadPage;
@@ -104,14 +103,6 @@ public:
 	bool CanRemoveSelectedPages() const;
 	void RemoveSelectedPages();
 
-	void RefreshTemplateVisibility();
-	void RefreshInstancedVisibility();
-
-	bool IsTemplatePageVisible(const FAvaRundownPage& InPage) const;
-	bool IsInstancedPageVisible(const FAvaRundownPage& InPage) const;
-
-	void SetSearchText(const FText& InText, EAvaRundownSearchListType& InPageListType);
-
 protected:
 	void RegisterApplicationModes();
 
@@ -141,19 +132,6 @@ protected:
 	using FBindableMacroCommands = TMap<FName, FMacroCommandFunction>;
 	const FBindableMacroCommands& GetBindableMacroCommands();
 
-private:
-	void SetTemplateSearchText(const FText& InText);
-
-	void SetInstancedSearchText(const FText& InText);
-
-	void OnTemplateFilterChanged();
-
-	void OnInstancedFilterChanged();
-
-	void InitVisibilityTemplatePages();
-
-	void InitVisibilityInstancedPages();
-
 protected:
 	/**
 	 *	Tickable object that keeps track of the play time of a page and will automatically
@@ -199,12 +177,6 @@ private:
 
 	TSharedPtr<FAvaRundownEditorInputProcessor> InputProcessor;
 
-	TSharedRef<FAvaRundownPageTextFilter> TextFilterTemplatePage;
-	TArray<int32> VisibleTemplatePageIds;
-
-	TSharedRef<FAvaRundownPageTextFilter> TextFilterInstancedPage;
-	TArray<int32> VisibleInstancedPageIds;
-
 	/**
 	 * All rundown editor instances share the same command handlers.
 	 * The commands get dispatched only to the active editor.
