[Motion Design] AvalancheMedia Clean up: Replacing FText() by FText::GetEmpty() where applicable.

diff --git a/Engine/Plugins/Experimental/Avalanche/Source/AvalancheMedia/Private/Broadcast/OutputDevices/AvaBroadcastDeviceProviderRegistry.cpp b/Engine/Plugins/Experimental/Avalanche/Source/AvalancheMedia/Private/Broadcast/OutputDevices/AvaBroadcastDeviceProviderRegistry.cpp
--- a/Engine/Plugins/Experimental/Avalanche/Source/AvalancheMedia/Private/Broadcast/OutputDevices/AvaBroadcastDeviceProviderRegistry.cpp
+++ b/Engine/Plugins/Experimental/Avalanche/Source/AvalancheMedia/Private/Broadcast/OutputDevices/AvaBroadcastDeviceProviderRegistry.cpp
@@ -57,7 +57,7 @@ FName FAvaBroadcastDeviceProviderRegistry::GetDeviceProviderName(const UClass* I
 	return FName();
 }
 
-FText FAvaBroadcastDeviceProviderRegistry::GetOutputClassDisplayText(const UClass* InMediaOutputClass) const
+const FText& FAvaBroadcastDeviceProviderRegistry::GetOutputClassDisplayText(const UClass* InMediaOutputClass) const
 {
 	if (InMediaOutputClass)
 	{
@@ -66,7 +66,7 @@ FText FAvaBroadcastDeviceProviderRegistry::GetOutputClassDisplayText(const UClas
 			return *FoundText;
 		}
 	}
-	return FText();
+	return FText::GetEmpty();
 }
 
 
diff --git a/Engine/Plugins/Experimental/Avalanche/Source/AvalancheMedia/Private/Broadcast/OutputDevices/AvaBroadcastDeviceProviderRegistry.h b/Engine/Plugins/Experimental/Avalanche/Source/AvalancheMedia/Private/Broadcast/OutputDevices/AvaBroadcastDeviceProviderRegistry.h
--- a/Engine/Plugins/Experimental/Avalanche/Source/AvalancheMedia/Private/Broadcast/OutputDevices/AvaBroadcastDeviceProviderRegistry.h
+++ b/Engine/Plugins/Experimental/Avalanche/Source/AvalancheMedia/Private/Broadcast/OutputDevices/AvaBroadcastDeviceProviderRegistry.h
@@ -29,7 +29,7 @@ public:
 
 	bool HasDeviceProviderName(const UClass* InMediaOutputClass) const;
 	FName GetDeviceProviderName(const UClass* InMediaOutputClass) const;
-	FText GetOutputClassDisplayText(const UClass* InMediaOutputClass) const;
+	const FText& GetOutputClassDisplayText(const UClass* InMediaOutputClass) const;
 
 protected:
 	FAvaBroadcastDeviceProviderRegistry();
diff --git a/Engine/Plugins/Experimental/Avalanche/Source/AvalancheMediaEditor/Private/Broadcast/AvaBroadcastEditor.cpp b/Engine/Plugins/Experimental/Avalanche/Source/AvalancheMediaEditor/Private/Broadcast/AvaBroadcastEditor.cpp
--- a/Engine/Plugins/Experimental/Avalanche/Source/AvalancheMediaEditor/Private/Broadcast/AvaBroadcastEditor.cpp
+++ b/Engine/Plugins/Experimental/Avalanche/Source/AvalancheMediaEditor/Private/Broadcast/AvaBroadcastEditor.cpp
@@ -451,7 +451,7 @@ TSharedRef<SWidget> FAvaBroadcastEditor::MakeProfileComboButton()
 	MenuBuilder.BeginSection(TEXT("AddProfile"));
 	{
 		MenuBuilder.AddMenuEntry(LOCTEXT("AddProfile_Label", "New Profile...")
-			, FText()
+			, FText::GetEmpty()
 			, FSlateIcon()
 			, FUIAction(FExecuteAction::CreateSP(this, &FAvaBroadcastEditor::CreateNewProfile)));
 	}	
diff --git a/Engine/Plugins/Experimental/Avalanche/Source/AvalancheMediaEditor/Private/Broadcast/ChannelGrid/Slate/SAvaBroadcastChannel.cpp b/Engine/Plugins/Experimental/Avalanche/Source/AvalancheMediaEditor/Private/Broadcast/ChannelGrid/Slate/SAvaBroadcastChannel.cpp
--- a/Engine/Plugins/Experimental/Avalanche/Source/AvalancheMediaEditor/Private/Broadcast/ChannelGrid/Slate/SAvaBroadcastChannel.cpp
+++ b/Engine/Plugins/Experimental/Avalanche/Source/AvalancheMediaEditor/Private/Broadcast/ChannelGrid/Slate/SAvaBroadcastChannel.cpp
@@ -403,7 +403,7 @@ TSharedRef<SWidget> SAvaBroadcastChannel::GetChannelStatusOptions()
 				if (Channel.GetState() != State && State != EAvaBroadcastChannelState::Offline)
 				{
 					Builder.AddMenuEntry(StaticEnum<EAvaBroadcastChannelState>()->GetDisplayNameTextByIndex(Index)
-						, FText()
+						, FText::GetEmpty()
 						, FSlateIcon()
 						, FUIAction(FExecuteAction::CreateSP(this, &SAvaBroadcastChannel::OnChannelStatusSelected, State)));	
 				}
diff --git a/Engine/Plugins/Experimental/Avalanche/Source/AvalancheMediaEditor/Private/Playback/Graph/Nodes/AvaPlaybackEditorGraphNode.cpp b/Engine/Plugins/Experimental/Avalanche/Source/AvalancheMediaEditor/Private/Playback/Graph/Nodes/AvaPlaybackEditorGraphNode.cpp
--- a/Engine/Plugins/Experimental/Avalanche/Source/AvalancheMediaEditor/Private/Playback/Graph/Nodes/AvaPlaybackEditorGraphNode.cpp
+++ b/Engine/Plugins/Experimental/Avalanche/Source/AvalancheMediaEditor/Private/Playback/Graph/Nodes/AvaPlaybackEditorGraphNode.cpp
@@ -388,7 +388,7 @@ void UAvaPlaybackEditorGraphNode::GetNodeContextMenuActions(UToolMenu* Menu, UGr
 			
 			Section.AddSubMenu("Alignment"
 				, LOCTEXT("AlignmentHeader", "Alignment")
-				, FText()
+				, FText::GetEmpty()
 				, FNewToolMenuDelegate::CreateLambda([](UToolMenu* SubMenu)
 				{
 					{
diff --git a/Engine/Plugins/Experimental/Avalanche/Source/AvalancheMediaEditor/Private/Rundown/Pages/Columns/Slate/SAvaRundownPageChannelSelector.cpp b/Engine/Plugins/Experimental/Avalanche/Source/AvalancheMediaEditor/Private/Rundown/Pages/Columns/Slate/SAvaRundownPageChannelSelector.cpp
--- a/Engine/Plugins/Experimental/Avalanche/Source/AvalancheMediaEditor/Private/Rundown/Pages/Columns/Slate/SAvaRundownPageChannelSelector.cpp
+++ b/Engine/Plugins/Experimental/Avalanche/Source/AvalancheMediaEditor/Private/Rundown/Pages/Columns/Slate/SAvaRundownPageChannelSelector.cpp
@@ -59,7 +59,7 @@ FText SAvaRundownPageChannelSelector::GetCurrentChannelName() const
 	{
 		return FText::FromName(PageView->GetChannelName());
 	}
-	return FText();
+	return FText::GetEmpty();
 }
 
 void SAvaRundownPageChannelSelector::OnComboBoxOpening()
diff --git a/Engine/Plugins/Experimental/Avalanche/Source/AvalancheMediaEditor/Private/Rundown/Preview/SAvaRundownPagePreview.cpp b/Engine/Plugins/Experimental/Avalanche/Source/AvalancheMediaEditor/Private/Rundown/Preview/SAvaRundownPagePreview.cpp
--- a/Engine/Plugins/Experimental/Avalanche/Source/AvalancheMediaEditor/Private/Rundown/Preview/SAvaRundownPagePreview.cpp
+++ b/Engine/Plugins/Experimental/Avalanche/Source/AvalancheMediaEditor/Private/Rundown/Preview/SAvaRundownPagePreview.cpp
@@ -278,7 +278,7 @@ void SAvaRundownPagePreview::AddResolutionMenuEntry(FMenuBuilder& InOutMenuBuild
 {
 	InOutMenuBuilder.AddMenuEntry(
 			Label,
-			FText(),
+			FText::GetEmpty(),
 			FSlateIcon(),
 			FUIAction(
 				FExecuteAction::CreateSP(this, &SAvaRundownPagePreview::SetPreviewResolution, InResolution),
