[Motion Design] Rundown Editor: Improving Page View workflow (Add/Delete Page View, Undo)

Rundown Backend Refactor: Adding a SubList uuid to avoid having to use unreliable indices.

In order to support adding and removing sublists, using indices to reference sublists causes a lot of problems to refresh the UI widgets.
It makes the operation undo very hard to implement. To make this easier (it is hard enough), the sublists have a unique id so the widgets can keep the reference regardless of it's index or display name.

This makes it much easier to implement the refresh (adding missing ones, closing deleted ones, etc) of the sub list tabs after an undo.

diff --git a/Engine/Plugins/Experimental/Avalanche/Source/AvalancheMedia/Internal/Rundown/AvaRundownMessages.h b/Engine/Plugins/Experimental/Avalanche/Source/AvalancheMedia/Internal/Rundown/AvaRundownMessages.h
--- a/Engine/Plugins/Experimental/Avalanche/Source/AvalancheMedia/Internal/Rundown/AvaRundownMessages.h
+++ b/Engine/Plugins/Experimental/Avalanche/Source/AvalancheMedia/Internal/Rundown/AvaRundownMessages.h
@@ -603,7 +603,7 @@ struct FAvaRundownPageListChanged : public FAvaRundownMsgBase
 	EAvaRundownPageListType ListType = EAvaRundownPageListType::Instance;
 
 	UPROPERTY()
-	int32 SubListIndex = INDEX_NONE;
+	FGuid SubListId;
 	
 	/** See EAvaPageListChange flags. */
 	UPROPERTY()
diff --git a/Engine/Plugins/Experimental/Avalanche/Source/AvalancheMedia/Private/Rundown/AvaRundown.cpp b/Engine/Plugins/Experimental/Avalanche/Source/AvalancheMedia/Private/Rundown/AvaRundown.cpp
--- a/Engine/Plugins/Experimental/Avalanche/Source/AvalancheMedia/Private/Rundown/AvaRundown.cpp
+++ b/Engine/Plugins/Experimental/Avalanche/Source/AvalancheMedia/Private/Rundown/AvaRundown.cpp
@@ -45,9 +45,20 @@ void FAvaRundownPageCollection::Empty(UAvaRundown* InRundown, const FAvaRundownP
 	}
 }
 
-const FAvaRundownPageListReference UAvaRundown::TemplatePageList = {EAvaRundownPageListType::Template, -2};
-const FAvaRundownPageListReference UAvaRundown::InstancePageList = {EAvaRundownPageListType::Instance, -1};
+namespace UE::AvaMedia::Rundown::Private
+{
+	FString ToString(const FAvaRundownPageListReference& InPageListReference)
+	{
+		return FString::Printf(TEXT("{ Type: %s, Id: %s"),
+			*StaticEnum<EAvaRundownPageListType>()->GetNameStringByValue((int64)InPageListReference.Type),
+			*InPageListReference.SubListId.ToString());
+	}
+}
+
+const FAvaRundownPageListReference UAvaRundown::TemplatePageList = {EAvaRundownPageListType::Template, FGuid()};
+const FAvaRundownPageListReference UAvaRundown::InstancePageList = {EAvaRundownPageListType::Instance, FGuid()};
 const FAvaRundownSubList UAvaRundown::InvalidSubList;
+FAvaRundownSubList UAvaRundown::InvalidSubListMutable;
 
 UAvaRundown::UAvaRundown()
 {
@@ -97,6 +108,15 @@ void UAvaRundown::RefreshPageIndices()
 	InstancedPages.RefreshPageIndices();
 }
 
+void UAvaRundown::RefreshSubListIndices()
+{
+	SubListIndices.Empty(SubLists.Num());
+	for (int32 Index = 0; Index < SubLists.Num(); ++Index)
+	{
+		SubListIndices.Add(SubLists[Index].Id, Index);
+	}	
+}
+
 namespace UE::AvaMedia::Rundown::Private
 {
 	void OnPostLoadPages(TArray<FAvaRundownPage>& InPages)
@@ -130,6 +150,15 @@ void UAvaRundown::PostLoad()
 	UE::AvaMedia::Rundown::Private::OnPostLoadPages(TemplatePages.Pages);
 	UE::AvaMedia::Rundown::Private::OnPostLoadPages(InstancedPages.Pages);
 	
+	for (FAvaRundownSubList& SubList : SubLists)
+	{
+		if (!SubList.Id.IsValid())
+		{
+			SubList.Id = FGuid::NewGuid();
+		}
+	}
+
+	RefreshSubListIndices();
 	RefreshPageIndices();
 }
 
@@ -140,8 +169,17 @@ void UAvaRundown::PostEditUndo()
 	
 	//Force Refresh for any Undo
 	RefreshPageIndices();
+	RefreshSubListIndices();
+	
 	GetOnPageListChanged().Broadcast({this, TemplatePageList, EAvaRundownPageListChange::All, {}});
 	GetOnPageListChanged().Broadcast({this, InstancePageList, EAvaRundownPageListChange::All, {}});
+
+	for (const FAvaRundownSubList& SubList : SubLists)
+	{
+		GetOnPageListChanged().Broadcast({this, CreateSubListReference(SubList), EAvaRundownPageListChange::All, {}});
+	}
+
+	GetOnActiveListChanged().Broadcast();
 }
 #endif
 
@@ -334,9 +372,9 @@ bool UAvaRundown::ChangePageOrder(const FAvaRundownPageListReference& InPageList
 		return true;
 	}
 
-	if (SubLists.IsValidIndex(InPageListReference.SubListIndex))
+	if (IsValidSubList(InPageListReference))
 	{
-		FAvaRundownSubList& SubList = SubLists[InPageListReference.SubListIndex];
+		FAvaRundownSubList& SubList = GetSubList(InPageListReference);
 		TArray<int32> NewIndices;
 		NewIndices.Reserve(SubList.PageIds.Num());
 
@@ -485,7 +523,6 @@ int32 UAvaRundown::RemovePages(const TArray<int32>& InPageIds)
 		GetOnPageListChanged().Broadcast({this, InstancePageList, EAvaRundownPageListChange::RemovedPages, {InPageIds}});
 	}
 
-	int32 SubListIndex = 0;
 	for (FAvaRundownSubList& SubList : SubLists)
 	{
 		bool bFoundInstance = false;
@@ -501,9 +538,8 @@ int32 UAvaRundown::RemovePages(const TArray<int32>& InPageIds)
 
 		if (bFoundInstance)
 		{
-			GetOnPageListChanged().Broadcast({this, CreateSubListReference(SubListIndex), EAvaRundownPageListChange::RemovedPages, {InPageIds}});
+			GetOnPageListChanged().Broadcast({this, CreateSubListReference(SubList), EAvaRundownPageListChange::RemovedPages, {InPageIds}});
 		}
-		++SubListIndex;
 	}
 
 	return RemovedCount;
@@ -590,7 +626,6 @@ bool UAvaRundown::RenumberPageId(int32 InPageId, int32 InNewPageId)
 	{
 		GetOnPageListChanged().Broadcast({this, InstancePageList, EAvaRundownPageListChange::RenumberedPageId, {InNewPageId}});
 
-		int32 SubListIndex = 0;
 		for (FAvaRundownSubList& SubList : SubLists)
 		{
 			const int32 Index = SubList.PageIds.Find(InPageId);
@@ -598,9 +633,8 @@ bool UAvaRundown::RenumberPageId(int32 InPageId, int32 InNewPageId)
 			if (Index != INDEX_NONE)
 			{
 				SubList.PageIds[Index] = InNewPageId;
-				GetOnPageListChanged().Broadcast({this, CreateSubListReference(SubListIndex), EAvaRundownPageListChange::RenumberedPageId, {InNewPageId}});
+				GetOnPageListChanged().Broadcast({this, CreateSubListReference(SubList), EAvaRundownPageListChange::RenumberedPageId, {InNewPageId}});
 			}
-			++SubListIndex;
 		}
 	}
 	else
@@ -809,16 +843,17 @@ const FAvaRundownPage& UAvaRundown::GetNextPage(int32 InPageId, const FAvaRundow
 
 	if (const int32* InstancedIdx = InstancedPages.PageIndices.Find(InPageId))
 	{
-		if (InPageListReference.Type == EAvaRundownPageListType::View && SubLists.IsValidIndex(InPageListReference.SubListIndex))
+		if (IsValidSubList(InPageListReference))
 		{
-			int32 Index = SubLists[InPageListReference.SubListIndex].PageIds.Find(InPageId);
+			const FAvaRundownSubList& SubList = GetSubList(InPageListReference);
+			const int32 Index = SubList.PageIds.Find(InPageId);
 
 			if (Index != INDEX_NONE)
 			{
-				return GetNextFromSubList(SubLists[InPageListReference.SubListIndex].PageIds, Index);
+				return GetNextFromSubList(SubList.PageIds, Index);
 			}
 
-			if (SubLists[InPageListReference.SubListIndex].PageIds.IsEmpty())
+			if (SubList.PageIds.IsEmpty())
 			{
 				return GetNextFromPages(InstancedPages.Pages, (*InstancedIdx));
 			}
@@ -849,16 +884,17 @@ FAvaRundownPage& UAvaRundown::GetNextPage(int32 InPageId, const FAvaRundownPageL
 
 	if (const int32* InstancedIdx = InstancedPages.PageIndices.Find(InPageId))
 	{
-		if (InPageListReference.Type == EAvaRundownPageListType::View && SubLists.IsValidIndex(InPageListReference.SubListIndex))
+		if (IsValidSubList(InPageListReference))
 		{
-			const int32 Index = SubLists[InPageListReference.SubListIndex].PageIds.Find(InPageId);
+			const FAvaRundownSubList& SubList = GetSubList(InPageListReference);
+			const int32 Index = SubList.PageIds.Find(InPageId);
 
 			if (Index != INDEX_NONE)
 			{
-				return GetNextFromSubList(SubLists[InPageListReference.SubListIndex].PageIds, Index);
+				return GetNextFromSubList(SubList.PageIds, Index);
 			}
 
-			if (SubLists[InPageListReference.SubListIndex].PageIds.IsEmpty())
+			if (SubList.PageIds.IsEmpty())
 			{
 				return GetNextFromPages(InstancedPages.Pages, (*InstancedIdx));
 			}
@@ -1283,12 +1319,58 @@ bool UAvaRundown::CanContinuePage(int32 InPageId, bool bInPreview, const FName&
 	return false;
 }
 
-int32 UAvaRundown::AddSubList()
+FAvaRundownPageListReference UAvaRundown::AddSubList()
 {
 	const int32 SubListIdx = SubLists.Add(FAvaRundownSubList());
-	SetActivePageList(CreateSubListReference(SubListIdx));
+	SubLists[SubListIdx].Id = FGuid::NewGuid();
+	SubListIndices.Add(SubLists[SubListIdx].Id, SubListIdx);
+
+	const FAvaRundownPageListReference SubListReference = CreateSubListReference(SubLists[SubListIdx].Id);
+	GetOnPageListChanged().Broadcast({this, SubListReference, EAvaRundownPageListChange::SubListAddedOrRemoved, {}});
+	
+	return SubListReference;
+}
+
+bool UAvaRundown::RemoveSubList(const FAvaRundownPageListReference& InPageListReference)
+{
+	if (IsValidSubList(InPageListReference))
+	{
+		// Update active list
+		if (ActivePageList == InPageListReference)
+		{
+			SetActivePageList(InstancePageList);
+		}
+
+		if (const int32 *SubListIndex = SubListIndices.Find(InPageListReference.SubListId))
+		{
+			if (SubLists.IsValidIndex(*SubListIndex))
+			{
+				SubLists.RemoveAt(*SubListIndex);
+				SubListIndices.Remove(InPageListReference.SubListId);
+				GetOnPageListChanged().Broadcast({this, InPageListReference, EAvaRundownPageListChange::SubListAddedOrRemoved, {}});
+				return true;
+			}
+		}
+	}
 
-	return SubListIdx;
+	using namespace UE::AvaMedia::Rundown::Private;
+	UE_LOG(LogAvaRundown, Error, TEXT("Remove SubList failed: Invalid SubList Reference: %s."), *ToString(InPageListReference));
+	return false;
+}
+
+bool UAvaRundown::RenameSubList(const FAvaRundownPageListReference& InPageListReference, const FText& InNewName)
+{
+	FAvaRundownSubList& SubList = GetSubList(InPageListReference); 
+	if (SubList.IsValid())
+	{
+		SubList.Name = InNewName;
+		GetOnPageListChanged().Broadcast({this, InPageListReference, EAvaRundownPageListChange::SubListRenamed, {}});
+		return true;
+	}
+
+	using namespace UE::AvaMedia::Rundown::Private;
+	UE_LOG(LogAvaRundown, Error, TEXT("Rename SubList failed: Invalid SubList Reference: %s."), *ToString(InPageListReference));
+	return false;
 }
 
 TArray<int32> UAvaRundown::GetPlayingPageIds(const FName InProgramChannelName) const
@@ -1342,7 +1424,7 @@ bool UAvaRundown::SetActivePageList(const FAvaRundownPageListReference& InPageLi
 		return true;
 	}
 
-	if (InPageListReference.Type == EAvaRundownPageListType::View && SubLists.IsValidIndex(InPageListReference.SubListIndex))
+	if (IsValidSubList(InPageListReference))
 	{
 		ActivePageList = InPageListReference;
 		OnActiveListChanged.Broadcast();
@@ -1354,83 +1436,102 @@ bool UAvaRundown::SetActivePageList(const FAvaRundownPageListReference& InPageLi
 
 bool UAvaRundown::HasActiveSubList() const
 {
-	return (ActivePageList.Type == EAvaRundownPageListType::View && SubLists.IsValidIndex(ActivePageList.SubListIndex));
+	return IsValidSubList(ActivePageList);
 }
 
 const FAvaRundownSubList& UAvaRundown::GetSubList(int32 InSubListIndex) const
 {
-	if (SubLists.IsValidIndex(InSubListIndex))
-	{
-		return SubLists[InSubListIndex];
-	}
-
-	return InvalidSubList;
+	return SubLists.IsValidIndex(InSubListIndex) ? SubLists[InSubListIndex] : InvalidSubList; 
 }
 
 FAvaRundownSubList& UAvaRundown::GetSubList(int32 InSubListIndex)
 {
-	check(SubLists.IsValidIndex(InSubListIndex));
+	return SubLists.IsValidIndex(InSubListIndex) ? SubLists[InSubListIndex] : InvalidSubListMutable; 
+}
+
+const FAvaRundownSubList& UAvaRundown::GetSubList(const FGuid& InSubListId) const
+{
+	const int32 *Index = SubListIndices.Find(InSubListId);
+	return Index ? GetSubList(*Index) : InvalidSubList;
+}
+
+FAvaRundownSubList& UAvaRundown::GetSubList(const FGuid& InSubListId)
+{
+	const int32 *Index = SubListIndices.Find(InSubListId);
+	return Index ? GetSubList(*Index) : InvalidSubListMutable;
+}
 
-	return SubLists[InSubListIndex];
+int32 UAvaRundown::GetSubListIndex(const FAvaRundownSubList& InSubList) const
+{
+	const int32 *Index = SubListIndices.Find(InSubList.Id);
+	return Index ? *Index : INDEX_NONE;
 }
 
 bool UAvaRundown::IsValidSubList(const FAvaRundownPageListReference& InPageListReference) const
 {
-	return (InPageListReference.Type == EAvaRundownPageListType::View && SubLists.IsValidIndex(InPageListReference.SubListIndex));
+	return InPageListReference.Type == EAvaRundownPageListType::View && SubListIndices.Contains(InPageListReference.SubListId);
 }
 
-bool UAvaRundown::AddPageToSubList(int32 InSubListIndex, int32 InPageId, const FAvaRundownPageInsertPosition& InInsertPosition)
+bool UAvaRundown::AddPageToSubList(const FAvaRundownPageListReference& InPageListReference, int32 InPageId, const FAvaRundownPageInsertPosition& InInsertPosition)
 {
-	if (SubLists.IsValidIndex(InSubListIndex) && InstancedPages.PageIndices.Contains(InPageId) 
-		&& !SubLists[InSubListIndex].PageIds.Contains(InPageId))
+	if (!IsValidSubList(InPageListReference))
+	{
+		return false;
+	}
+	
+	FAvaRundownSubList& SubList = GetSubList(InPageListReference);
+	
+	if (InstancedPages.PageIndices.Contains(InPageId) && !SubList.PageIds.Contains(InPageId))
 	{
 		int32 ExistingPageIndex = INDEX_NONE;
 
 		if (InInsertPosition.IsValid())
 		{
-			ExistingPageIndex = SubLists[InSubListIndex].PageIds.IndexOfByKey(InInsertPosition.AdjacentId);
+			ExistingPageIndex = SubList.PageIds.IndexOfByKey(InInsertPosition.AdjacentId);
 		}
 
-		if (InInsertPosition.IsAddBelow() && SubLists[InSubListIndex].PageIds.IsValidIndex(ExistingPageIndex))
+		if (InInsertPosition.IsAddBelow() && SubList.PageIds.IsValidIndex(ExistingPageIndex))
 		{
 			++ExistingPageIndex;
 		}
 
-		if (SubLists[InSubListIndex].PageIds.IsValidIndex(ExistingPageIndex))
+		if (SubList.PageIds.IsValidIndex(ExistingPageIndex))
 		{
-			SubLists[InSubListIndex].PageIds.Insert(InPageId, ExistingPageIndex);
+			SubList.PageIds.Insert(InPageId, ExistingPageIndex);
 		}
 		else
 		{
-			SubLists[InSubListIndex].PageIds.Add(InPageId);
+			SubList.PageIds.Add(InPageId);
 		}
 		
-		GetOnPageListChanged().Broadcast({this, CreateSubListReference(InSubListIndex), EAvaRundownPageListChange::AddedPages, {InPageId}});
+		GetOnPageListChanged().Broadcast({this, InPageListReference, EAvaRundownPageListChange::AddedPages, {InPageId}});
 		return true;
 	}
 
 	return false;
 }
 
-bool UAvaRundown::AddPagesToSubList(int32 InSubListIndex, const TArray<int32>& InPages)
+bool UAvaRundown::AddPagesToSubList(const FAvaRundownPageListReference& InPageListReference, const TArray<int32>& InPages)
 {
-	if (SubLists.IsValidIndex(InSubListIndex))
+	if (IsValidSubList(InPageListReference))
 	{
+		FAvaRundownSubList& SubList = GetSubList(InPageListReference);
+		
 		bool bAddedPage = false;
 
 		// Super inefficient for now.
 		for (int32 PageId : InPages)
 		{
-			if (InstancedPages.PageIndices.Contains(PageId) && !SubLists[InSubListIndex].PageIds.Contains(PageId))
+			if (InstancedPages.PageIndices.Contains(PageId) && !SubList.PageIds.Contains(PageId))
 			{
-				SubLists[InSubListIndex].PageIds.Add(PageId);
+				SubList.PageIds.Add(PageId);
 				bAddedPage = true;
 			}
 		}
 
 		if (bAddedPage)
 		{
-			GetOnPageListChanged().Broadcast({this, CreateSubListReference(InSubListIndex), EAvaRundownPageListChange::AddedPages, InPages});
+			GetOnPageListChanged().Broadcast({this, InPageListReference, EAvaRundownPageListChange::AddedPages, InPages});
 			return true;
 		}
 	}
@@ -1438,18 +1539,18 @@ bool UAvaRundown::AddPagesToSubList(int32 InSubListIndex, const TArray<int32>& I
 	return false;
 }
 
-int32 UAvaRundown::RemovePagesFromSubList(int32 InSubListIndex, const TArray<int32>& InPages)
+int32 UAvaRundown::RemovePagesFromSubList(const FAvaRundownPageListReference& InPageListReference, const TArray<int32>& InPages)
 {
-	if (SubLists.IsValidIndex(InSubListIndex))
+	if (IsValidSubList(InPageListReference))
 	{
-		const int32 Removed = SubLists[InSubListIndex].PageIds.RemoveAll([InPages](const int32& PageId)
+		const int32 Removed = GetSubList(InPageListReference).PageIds.RemoveAll([InPages](const int32& PageId)
 			{
 				return InPages.Contains(PageId);
 			});
 
 		if (Removed > 0)
 		{
-			GetOnPageListChanged().Broadcast({this, CreateSubListReference(InSubListIndex), EAvaRundownPageListChange::RemovedPages, InPages});
+			GetOnPageListChanged().Broadcast({this, InPageListReference, EAvaRundownPageListChange::RemovedPages, InPages});
 		}
 
 		return Removed;
@@ -2086,7 +2187,7 @@ const FAvaRundownPage& UAvaRundown::GetNextFromSubList(const TArray<int32>& InSu
 	return FAvaRundownPage::NullPage;
 }
 
-FAvaRundownPage& UAvaRundown::GetNextFromSubList(TArray<int32>& InSubListIds, int32 InStartingIndex)
+FAvaRundownPage& UAvaRundown::GetNextFromSubList(const TArray<int32>& InSubListIds, int32 InStartingIndex)
 {
 	if (InSubListIds.IsEmpty())
 	{
diff --git a/Engine/Plugins/Experimental/Avalanche/Source/AvalancheMedia/Private/Rundown/AvaRundownServer.cpp b/Engine/Plugins/Experimental/Avalanche/Source/AvalancheMedia/Private/Rundown/AvaRundownServer.cpp
--- a/Engine/Plugins/Experimental/Avalanche/Source/AvalancheMedia/Private/Rundown/AvaRundownServer.cpp
+++ b/Engine/Plugins/Experimental/Avalanche/Source/AvalancheMedia/Private/Rundown/AvaRundownServer.cpp
@@ -482,7 +482,7 @@ void FAvaRundownServer::OnPageListChanged(const FAvaRundownPageListChangeParams&
 	FAvaRundownPageListChanged* ReplyMessage = FMessageEndpoint::MakeMessage<FAvaRundownPageListChanged>();
 	ReplyMessage->Rundown = FSoftObjectPath(InParams.Rundown).ToString();
 	ReplyMessage->ListType = InParams.PageListReference.Type;
-	ReplyMessage->SubListIndex = InParams.PageListReference.SubListIndex;
+	ReplyMessage->SubListId = InParams.PageListReference.SubListId;
 	ReplyMessage->ChangeType = static_cast<uint8>(InParams.ChangeType);
 	ReplyMessage->AffectedPages = InParams.AffectedPages;
 	
diff --git a/Engine/Plugins/Experimental/Avalanche/Source/AvalancheMedia/Public/AvaMediaDefines.h b/Engine/Plugins/Experimental/Avalanche/Source/AvalancheMedia/Public/AvaMediaDefines.h
--- a/Engine/Plugins/Experimental/Avalanche/Source/AvalancheMedia/Public/AvaMediaDefines.h
+++ b/Engine/Plugins/Experimental/Avalanche/Source/AvalancheMedia/Public/AvaMediaDefines.h
@@ -207,13 +207,14 @@ enum class EAvaRundownPageListChange : uint8
 {
 	None = 0,
 
-	AddedPages        = 1 << 0,
-	RemovedPages      = 1 << 1,
-	RenumberedPageId  = 1 << 2,
-	RenamedPageView   = 1 << 3,
-	ReorderedPageView = 1 << 4,
-
-	All              = 0xFF,
+	AddedPages            = 1 << 0,
+	RemovedPages          = 1 << 1,
+	RenumberedPageId      = 1 << 2,
+	SubListAddedOrRemoved = 1 << 3,
+	SubListRenamed        = 1 << 4,
+	ReorderedPageView     = 1 << 5,
+
+	All                   = 0xFF,
 };
 ENUM_CLASS_FLAGS(EAvaRundownPageListChange);
 
@@ -307,7 +308,7 @@ struct FAvaRundownPageListReference
 	EAvaRundownPageListType Type = EAvaRundownPageListType::Instance;
 
 	UPROPERTY()
-	int32 SubListIndex = INDEX_NONE;
+	FGuid SubListId;
 
 	bool operator==(const FAvaRundownPageListReference& InOther) const
 	{
@@ -316,7 +317,7 @@ struct FAvaRundownPageListReference
 			return false;
 		}
 
-		if (Type == EAvaRundownPageListType::View && SubListIndex != InOther.SubListIndex)
+		if (Type == EAvaRundownPageListType::View && SubListId != InOther.SubListId)
 		{
 			return false;
 		}
diff --git a/Engine/Plugins/Experimental/Avalanche/Source/AvalancheMedia/Public/Rundown/AvaRundown.h b/Engine/Plugins/Experimental/Avalanche/Source/AvalancheMedia/Public/Rundown/AvaRundown.h
--- a/Engine/Plugins/Experimental/Avalanche/Source/AvalancheMedia/Public/Rundown/AvaRundown.h
+++ b/Engine/Plugins/Experimental/Avalanche/Source/AvalancheMedia/Public/Rundown/AvaRundown.h
@@ -145,11 +145,21 @@ struct FAvaRundownSubList
 {
 	GENERATED_BODY()
 
+	/** Unique identifier for the page list. */
+	UPROPERTY()
+	FGuid Id;
+
 	UPROPERTY()
 	TArray<int32> PageIds;
 
+	/** Display Name */
 	UPROPERTY()
 	FText Name;
+
+	bool IsValid() const
+	{
+		return Id.IsValid();
+	}
 };
 
 namespace UE::AvaRundown
@@ -295,7 +305,15 @@ public:
 	static const FAvaRundownPageListReference TemplatePageList;
 	static const FAvaRundownPageListReference InstancePageList;
 
-	static FAvaRundownPageListReference CreateSubListReference(int32 InSubListIndex) { return {EAvaRundownPageListType::View, InSubListIndex}; }
+	static FAvaRundownPageListReference CreateSubListReference(const FGuid& InSubListId)
+	{
+		return {EAvaRundownPageListType::View, InSubListId};
+	}
+
+	static FAvaRundownPageListReference CreateSubListReference(const FAvaRundownSubList& InSubList)
+	{
+		return CreateSubListReference(InSubList.Id);
+	}
 
 protected:
 	bool IsPageIdUnique(int32 InPageId) const { return !TemplatePages.PageIndices.Contains(InPageId) && !InstancedPages.PageIndices.Contains(InPageId); }
@@ -304,9 +322,12 @@ protected:
 	
 	/** Caches the Page's Id to its Index in the Pages Array*/
 	void RefreshPageIndices();
+
+	void RefreshSubListIndices();
 	
 public:
 	static const FAvaRundownSubList InvalidSubList;
+	static FAvaRundownSubList InvalidSubListMutable;
 
 	//~ Begin UObject Interface
 	virtual void BeginDestroy() override;
@@ -516,7 +537,11 @@ public:
 	/** Used to reconcile playing state with a remote playback if connection was lost. */
 	bool RestorePlaySubPage(int32 InPageId, int32 InSubPageIndex, const FGuid& InExistingInstanceId, bool bInIsPreview, const FName& InPreviewChannelName);
 
-	int32 AddSubList();
+	FAvaRundownPageListReference AddSubList();
+	
+	bool RemoveSubList(const FAvaRundownPageListReference& InPageListReference);
+	
+	bool RenameSubList(const FAvaRundownPageListReference& InPageListReference, const FText& InNewName);
 
 	/**
 	 * Return the current playing Page Ids on the specified program channel.
@@ -537,22 +562,35 @@ public:
 	/** Returns true only if a sub list, not the main list, is active. */
 	bool HasActiveSubList() const;
 
-	const FAvaRundownSubList& GetActiveSubList() const { return GetSubList(ActivePageList.SubListIndex); }
-	FAvaRundownSubList& GetActiveSubList() { return GetSubList(ActivePageList.SubListIndex); }
+	const FAvaRundownSubList& GetActiveSubList() const { return GetSubList(ActivePageList.SubListId); }
+	FAvaRundownSubList& GetActiveSubList() { return GetSubList(ActivePageList.SubListId); }
 
 	const FAvaRundownSubList& GetSubList(int32 InSubListIndex) const;
-
 	FAvaRundownSubList& GetSubList(int32 InSubListIndex);
 
-	bool IsValidSubListIndex(int32 InIndex) const { return SubLists.IsValidIndex(InIndex); }
+	const FAvaRundownSubList& GetSubList(const FGuid& InSubListId) const;
+	FAvaRundownSubList& GetSubList(const FGuid& InSubListId);
+
+	const FAvaRundownSubList& GetSubList(const FAvaRundownPageListReference& InPageListReference) const
+	{
+		return InPageListReference.Type == EAvaRundownPageListType::View ? GetSubList(InPageListReference.SubListId) : InvalidSubList;
+	}
+	
+	FAvaRundownSubList& GetSubList(const FAvaRundownPageListReference& InPageListReference)
+	{
+		return InPageListReference.Type == EAvaRundownPageListType::View ? GetSubList(InPageListReference.SubListId) : InvalidSubListMutable;
+	}
+
+	int32 GetSubListIndex(const FAvaRundownSubList& InSubList) const;
+
 	bool IsValidSubList(const FAvaRundownPageListReference& InPageListReference) const;
 
 	const TArray<FAvaRundownSubList>& GetSubLists() const { return SubLists; }
 
-	bool AddPageToSubList(int32 InSubListIndex, int32 InPageId, const FAvaRundownPageInsertPosition& InInsertPosition = FAvaRundownPageInsertPosition());
-	bool AddPagesToSubList(int32 InSubListIndex, const TArray<int32>& InPages);
+	bool AddPageToSubList(const FAvaRundownPageListReference& InPageListReference, int32 InPageId, const FAvaRundownPageInsertPosition& InInsertPosition = FAvaRundownPageInsertPosition());
+	bool AddPagesToSubList(const FAvaRundownPageListReference& InPageListReference, const TArray<int32>& InPages);
 
-	int32 RemovePagesFromSubList(int32 InSubListIndex, const TArray<int32>& InPages);
+	int32 RemovePagesFromSubList(const FAvaRundownPageListReference& InPageListReference, const TArray<int32>& InPages);
 
 	DECLARE_MULTICAST_DELEGATE(FOnActiveListChanged)
 	FOnActiveListChanged& GetOnActiveListChanged() { return OnActiveListChanged; }
@@ -619,7 +657,7 @@ private:
 	FAvaRundownPage& GetNextFromPages(TArray<FAvaRundownPage>& InPages, int32 InStartingIndex) const;
 
 	const FAvaRundownPage& GetNextFromSubList(const TArray<int32>& InSubListIds, int32 InStartingIndex) const;
-	FAvaRundownPage& GetNextFromSubList(TArray<int32>& InSubListIds, int32 InStartingIndex);
+	FAvaRundownPage& GetNextFromSubList(const TArray<int32>& InSubListIds, int32 InStartingIndex);
 
 protected:
 	UPROPERTY()
@@ -727,4 +765,6 @@ protected:
 	FOnAvaRundownPageListChanged OnPageListChanged;
 	FOnActiveListChanged OnActiveListChanged;
 	mutable FOnCanClosePlaybackContext OnCanClosePlaybackContext;
+
+	TMap<FGuid, int32> SubListIndices;
 };
diff --git a/Engine/Plugins/Experimental/Avalanche/Source/AvalancheMediaEditor/Private/Rundown/AvaRundownEditor.cpp b/Engine/Plugins/Experimental/Avalanche/Source/AvalancheMediaEditor/Private/Rundown/AvaRundownEditor.cpp
--- a/Engine/Plugins/Experimental/Avalanche/Source/AvalancheMediaEditor/Private/Rundown/AvaRundownEditor.cpp
+++ b/Engine/Plugins/Experimental/Avalanche/Source/AvalancheMediaEditor/Private/Rundown/AvaRundownEditor.cpp
@@ -24,7 +24,6 @@
 #include "Rundown/AvaRundownManagedInstanceCache.h"
 #include "Rundown/AvaRundownPagePlayer.h"
 #include "Rundown/AvaRundownPlaybackUtils.h"
-#include "Rundown/Factories/Filters/AvaRundownFactoriesUtils.h"
 #include "Rundown/Pages/PageViews/IAvaRundownPageView.h"
 #include "Rundown/Pages/Slate/SAvaRundownPageList.h"
 #include "ScopedTransaction.h"
@@ -191,6 +190,7 @@ FAvaRundownEditor::~FAvaRundownEditor()
 	if (IsValid(Rundown))
 	{
 		Rundown->GetOnActiveListChanged().RemoveAll(this);
+		Rundown->GetOnPageListChanged().RemoveAll(this);
 		Rundown->GetOnPagePlayerAdded().RemoveAll(this);
 		Rundown->GetOnCanClosePlaybackContext().RemoveAll(this);
 
@@ -234,7 +234,8 @@ void FAvaRundownEditor::InitRundownEditor(const EToolkitMode::Type InMode
 
 	if (IsValid(InRundown))
 	{
-		InRundown->GetOnActiveListChanged().AddSP(this, &FAvaRundownEditor::OnActiveSubListChanged);
+		InRundown->GetOnActiveListChanged().AddSP(this, &FAvaRundownEditor::OnActiveListChanged);
+		InRundown->GetOnPageListChanged().AddSP(this, &FAvaRundownEditor::OnPageListChanged);
 		InRundown->GetOnPagePlayerAdded().AddSP(this, &FAvaRundownEditor::HandleOnPagePlayerAdded);
 		InRundown->GetOnCanClosePlaybackContext().AddSP(this, &FAvaRundownEditor::OnCanClosePlaybackContext);
 		InRundown->InitializePlaybackContext();
@@ -255,7 +256,7 @@ void FAvaRundownEditor::InitRundownEditor(const EToolkitMode::Type InMode
 	RegisterApplicationModes();
 	FSharedConsoleCommands::RegisterEditor(SharedThis(this));
 
-	CreateSubListTabs();
+	RefreshSubListTabs();
 }
 
 bool FAvaRundownEditor::IsKeyRelevant(const FKeyEvent& InKeyEvent) const
@@ -328,7 +329,7 @@ TSharedPtr<SAvaRundownPageList> FAvaRundownEditor::GetListWidget(const FAvaRundo
 		return GetListWidget(FAvaRundownInstancedPageListTabFactory::TabID);
 	}
 
-	return GetListWidget(FAvaRundownSubListDocumentTabFactory::GetTabId(InPageListReference.SubListIndex));
+	return GetListWidget(FAvaRundownSubListDocumentTabFactory::GetTabId(InPageListReference));
 }
 
 TSharedPtr<SAvaRundownPageList> FAvaRundownEditor::GetListWidget(const FName& InTabId) const
@@ -377,7 +378,7 @@ TSharedPtr<SAvaRundownInstancedPageList> FAvaRundownEditor::GetActiveListWidget(
 	}
 	else if (ActiveList.Type == EAvaRundownPageListType::View)
 	{
-		PageList = GetListWidget(FAvaRundownSubListDocumentTabFactory::GetTabId(ActiveList.SubListIndex));
+		PageList = GetListWidget(FAvaRundownSubListDocumentTabFactory::GetTabId(ActiveList));
 	}
 
 	if (PageList.IsValid() && PageList->GetWidgetClass().GetWidgetType() == SAvaRundownInstancedPageList::StaticWidgetClass().GetWidgetType())
@@ -446,13 +447,11 @@ TConstArrayView<int32> FAvaRundownEditor::GetSelectedPagesOnFocusedWidget() cons
 
 bool FAvaRundownEditor::CanAddTemplate() const
 {
-	UAvaRundown* Rundown = AvaRundown.Get();
-
+	const UAvaRundown* Rundown = AvaRundown.Get();
 	if (IsValid(Rundown))
 	{
 		return Rundown->CanAddPage();
 	}
-
 	return false;
 }
 
@@ -752,7 +751,7 @@ void FAvaRundownEditor::CreateInstancesFromSelectedTemplates()
 
 		if (PageListWidget.IsValid() && PageListWidget->GetPageListReference().Type == EAvaRundownPageListType::View)
 		{
-			Rundown->AddPagesToSubList(PageListWidget->GetPageListReference().SubListIndex, AddedPages);
+			Rundown->AddPagesToSubList(PageListWidget->GetPageListReference(), AddedPages);
 		}
 	}
 	else
@@ -800,6 +799,18 @@ void FAvaRundownEditor::RemoveSelectedPages()
 	}
 }
 
+void FAvaRundownEditor::RefreshSubListTabs()
+{
+	const UAvaRundown* Rundown = AvaRundown.Get();
+	if (IsValid(Rundown))
+	{
+		for (const FAvaRundownSubList& SubList : Rundown->GetSubLists())
+		{
+			RefreshSubListTab(UAvaRundown::CreateSubListReference(SubList), /*bInSetActive*/false);
+		}
+	}
+}
+
 FName FAvaRundownEditor::GetToolkitFName() const
 {
 	return TEXT("AvaRundownEditor");
@@ -921,36 +932,22 @@ void FAvaRundownEditor::MarkAsModified()
 	}
 }
 
-void FAvaRundownEditor::OnActiveSubListChanged()
+void FAvaRundownEditor::OnPageListChanged(const FAvaRundownPageListChangeParams& InParams)
 {
-	UAvaRundown* Rundown = AvaRundown.Get();
+	// Undo Support - Create missing Page View Tabs.
+	RefreshSubListTab(InParams.PageListReference);
+}
 
+void FAvaRundownEditor::OnActiveListChanged()
+{
+	const UAvaRundown* Rundown = AvaRundown.Get();
 	if (!IsValid(Rundown))
 	{
 		return;
 	}
 
 	const FAvaRundownPageListReference& ActiveList = Rundown->GetActivePageListReference();
-
-	if (ActiveList.Type != EAvaRundownPageListType::View)
-	{
-		return;
-	}
-
-	const FName TabId = FAvaRundownSubListDocumentTabFactory::GetTabId(ActiveList.SubListIndex);
-
-	TSharedPtr<SDockTab> SubListTab = TabManager->FindExistingLiveTab(TabId);
-
-	if (!SubListTab.IsValid())
-	{
-		SubListTab = CreateSubListTab(ActiveList.SubListIndex);
-	}
-
-	if (SubListTab.IsValid())
-	{
-		SubListTab->ActivateInParent(ETabActivationCause::SetDirectly);
-		SubListTab->DrawAttention();
-	}
+	RefreshSubListTab(ActiveList, /*bInSetActive*/true);
 }
 
 void FAvaRundownEditor::HandleOnPagePlayerAdded(UAvaRundown* InRundown, UAvaRundownPagePlayer* InPagePlayer)
@@ -979,26 +976,26 @@ void FAvaRundownEditor::OnCanClosePlaybackContext(const UAvaRundown* InRundown,
 	}
 }
 
-TSharedPtr<SDockTab> FAvaRundownEditor::CreateSubListTab(int32 InSubListIndex)
+TSharedPtr<SDockTab> FAvaRundownEditor::CreateSubListTab(const FAvaRundownPageListReference& InSubListReference)
 {
-	TSharedPtr<FApplicationMode> AppMode = GetCurrentModePtr();
+	const TSharedPtr<FApplicationMode> AppMode = GetCurrentModePtr();
 
 	if (AppMode.IsValid())
 	{
-		const FName TabId = FAvaRundownSubListDocumentTabFactory::GetTabId(InSubListIndex);
-		TSharedRef<FAvaRundownAppMode> RundownAppMode = StaticCastSharedRef<FAvaRundownAppMode>(AppMode.ToSharedRef());
-		TSharedPtr<FDocumentTabFactory> DocTabFactory = RundownAppMode->GetDocumentTabFactory(FAvaRundownSubListDocumentTabFactory::FactoryId);
+		const FName TabId = FAvaRundownSubListDocumentTabFactory::GetTabId(InSubListReference);
+		const TSharedRef<FAvaRundownAppMode> RundownAppMode = StaticCastSharedRef<FAvaRundownAppMode>(AppMode.ToSharedRef());
+		const TSharedPtr<FDocumentTabFactory> DocTabFactory = RundownAppMode->GetDocumentTabFactory(FAvaRundownSubListDocumentTabFactory::FactoryId);
 
 		if (DocTabFactory.IsValid())
 		{
-			TSharedRef<FAvaRundownSubListDocumentTabFactory> SubListTabFactory = StaticCastSharedRef<FAvaRundownSubListDocumentTabFactory>(DocTabFactory.ToSharedRef());
+			const TSharedRef<FAvaRundownSubListDocumentTabFactory> SubListTabFactory = StaticCastSharedRef<FAvaRundownSubListDocumentTabFactory>(DocTabFactory.ToSharedRef());
 
 			FWorkflowTabSpawnInfo Info;
 			Info.TabManager = TabManager;
 			Info.Payload = nullptr;
 			Info.TabInfo = nullptr;
 
-			TSharedPtr<SDockTab> SubListTab = SubListTabFactory->SpawnSubListTab(Info, InSubListIndex);
+			TSharedPtr<SDockTab> SubListTab = SubListTabFactory->SpawnSubListTab(Info, InSubListReference);
 
 			if (SubListTab.IsValid())
 			{
@@ -1011,18 +1008,42 @@ TSharedPtr<SDockTab> FAvaRundownEditor::CreateSubListTab(int32 InSubListIndex)
 	return nullptr;
 }
 
-void FAvaRundownEditor::CreateSubListTabs()
+void FAvaRundownEditor::RefreshSubListTab(const FAvaRundownPageListReference& InSubListReference, bool bInSetActive)
 {
-	UAvaRundown* Rundown = AvaRundown.Get();
+	if (InSubListReference.Type != EAvaRundownPageListType::View)
+	{
+		return;
+	}
 
+	const UAvaRundown* Rundown = AvaRundown.Get();
 	if (!IsValid(Rundown))
 	{
 		return;
 	}
 
-	for (int32 SubListIndex = 0; SubListIndex < Rundown->GetSubLists().Num(); ++SubListIndex)
+	const FName TabId = FAvaRundownSubListDocumentTabFactory::GetTabId(InSubListReference);
+
+	TSharedPtr<SDockTab> SubListTab = TabManager->FindExistingLiveTab(TabId);
+
+	if (Rundown->IsValidSubList(InSubListReference))
+	{
+		if (!SubListTab.IsValid())
+		{
+			SubListTab = CreateSubListTab(InSubListReference);
+		}
+
+		if (SubListTab.IsValid() && bInSetActive)
+		{
+			SubListTab->ActivateInParent(ETabActivationCause::SetDirectly);
+			SubListTab->DrawAttention();
+		}
+	}
+	else
 	{
-		CreateSubListTab(SubListIndex);
+		if (SubListTab.IsValid())
+		{
+			SubListTab->RequestCloseTab();
+		}
 	}
 }
 
@@ -1127,7 +1148,7 @@ TSharedRef<SWidget> FAvaRundownEditor::MakeProfileComboButton()
 	for (FName ProfileName : ProfileNames)
 	{
 		MenuBuilder.AddMenuEntry(FText::FromName(ProfileName)
-			, FText()
+			, FText::GetEmpty()
 			, FSlateIcon()
 			, FUIAction(FExecuteAction::CreateLambda([ProfileName](){ UAvaBroadcast::Get().SetCurrentProfile(ProfileName);}))
 		);
diff --git a/Engine/Plugins/Experimental/Avalanche/Source/AvalancheMediaEditor/Private/Rundown/AvaRundownEditorUtils.cpp b/Engine/Plugins/Experimental/Avalanche/Source/AvalancheMediaEditor/Private/Rundown/AvaRundownEditorUtils.cpp
--- a/Engine/Plugins/Experimental/Avalanche/Source/AvalancheMediaEditor/Private/Rundown/AvaRundownEditorUtils.cpp
+++ b/Engine/Plugins/Experimental/Avalanche/Source/AvalancheMediaEditor/Private/Rundown/AvaRundownEditorUtils.cpp
@@ -526,7 +526,7 @@ TArray<int32> UE::AvaRundownEditor::Utils::ImportInstancedPages(
 		check(InstancedPageId != FAvaRundownPage::InvalidPageId);
 
 		// Now we have our instance page reference, add it to the sublist
-		if (InRundown->AddPageToSubList(InPageListReference.SubListIndex, InstancedPageId))
+		if (InRundown->AddPageToSubList(InPageListReference, InstancedPageId))
 		{
 			OutPageIds.Add(InstancedPageId);
 		}
diff --git a/Engine/Plugins/Experimental/Avalanche/Source/AvalancheMediaEditor/Private/Rundown/Pages/Slate/SAvaRundownInstancedPageList.cpp b/Engine/Plugins/Experimental/Avalanche/Source/AvalancheMediaEditor/Private/Rundown/Pages/Slate/SAvaRundownInstancedPageList.cpp
--- a/Engine/Plugins/Experimental/Avalanche/Source/AvalancheMediaEditor/Private/Rundown/Pages/Slate/SAvaRundownInstancedPageList.cpp
+++ b/Engine/Plugins/Experimental/Avalanche/Source/AvalancheMediaEditor/Private/Rundown/Pages/Slate/SAvaRundownInstancedPageList.cpp
@@ -19,6 +19,7 @@
 #include "Rundown/Pages/Columns/AvaRundownPageStatusColumn.h"
 #include "Rundown/Pages/Columns/AvaRundownPageTemplateNameColumn.h"
 #include "Rundown/Pages/PageViews/AvaRundownInstancedPageViewImpl.h"
+#include "Rundown/TabFactories/AvaRundownSubListDocumentTabFactory.h"
 #include "SAvaRundownPageList.h"
 #include "ScopedTransaction.h"
 #include "Styling/StyleColors.h"
@@ -45,6 +46,7 @@ void SAvaRundownInstancedPageList::Construct(const FArguments& InArgs, TSharedPt
 	check(InPageListReference.Type == EAvaRundownPageListType::Instance || Rundown->IsValidSubList(PageListReference));
 
 	Rundown->GetOnPageListChanged().AddSP(this, &SAvaRundownInstancedPageList::OnPageListChanged);
+	Rundown->GetOnActiveListChanged().AddSP(this, &SAvaRundownInstancedPageList::OnActiveListChanged);
 
 	if (Rundown->IsValidSubList(PageListReference))
 	{
@@ -68,6 +70,23 @@ void SAvaRundownInstancedPageList::Construct(const FArguments& InArgs, TSharedPt
 				.Text(this, &SAvaRundownInstancedPageList::GetPageViewName)
 				.MinDesiredWidth(150.f)
 			];
+
+		SearchBar->InsertSlot(2)
+			.AutoWidth()
+			.VAlign(EVerticalAlignment::VAlign_Center)
+			.Padding(5.f, 0.f)
+			[
+				SNew(SButton)
+				.IsFocusable(false)
+				.ToolTipText(LOCTEXT("PageViewDeleteToolTip", "Delete Current Page View"))
+				.OnClicked(this, &SAvaRundownInstancedPageList::OnDeletePageView)
+				.Content()
+				[
+					SNew(SImage)
+					.Image(FAppStyle::Get().GetBrush("Icons.Delete"))
+					.ColorAndOpacity(FSlateColor::UseForeground())
+				]
+			];
 	}
 
 	SearchBar->AddSlot()
@@ -107,6 +126,7 @@ SAvaRundownInstancedPageList::~SAvaRundownInstancedPageList()
 	if (UAvaRundown* const Rundown = GetValidRundown())
 	{
 		Rundown->GetOnPageListChanged().RemoveAll(this);
+		Rundown->GetOnActiveListChanged().RemoveAll(this);
 	}
 }
 
@@ -133,7 +153,7 @@ void SAvaRundownInstancedPageList::Refresh()
 		}
 		else if (Rundown->IsValidSubList(PageListReference))
 		{
-			const FAvaRundownSubList& SubList = Rundown->GetSubList(PageListReference.SubListIndex);
+			const FAvaRundownSubList& SubList = Rundown->GetSubList(PageListReference);
 			PageViews.Reset(SubList.PageIds.Num());
 			
 			for (const int32 PageId : SubList.PageIds)
@@ -143,7 +163,7 @@ void SAvaRundownInstancedPageList::Refresh()
 					PageViews.Emplace(MakeShared<FAvaRundownInstancedPageViewImpl>(PageId, Rundown, SharedThis(this)));
 				}
 			}
-		}	
+		}
 
 		PageListView->RequestListRefresh();
 	}
@@ -349,7 +369,7 @@ bool SAvaRundownInstancedPageList::HandleDropAssets(const TArray<FSoftObjectPath
 
 		if (PageListReference.Type == EAvaRundownPageListType::View)
 		{
-			Rundown->AddPageToSubList(PageListReference.SubListIndex, NewInstanceId, InsertAt);
+			Rundown->AddPageToSubList(PageListReference, NewInstanceId, InsertAt);
 		}
 
 		InsertAt.ConditionalUpdateAdjacentId(NewInstanceId);
@@ -478,7 +498,7 @@ bool SAvaRundownInstancedPageList::HandleDropPageIds(const FAvaRundownPageListRe
 		}
 		else if (bFromSubList)
 		{
-			bDidSomething = HandleDropPageIdsOnSubListFromSubList(InPageListReference.SubListIndex, InPageIds, InDropZone, InItem);
+			bDidSomething = HandleDropPageIdsOnSubListFromSubList(InPageListReference, InPageIds, InDropZone, InItem);
 		}
 	}
 
@@ -750,7 +770,7 @@ bool SAvaRundownInstancedPageList::HandleDropPageIdsOnSubListFromTemplates(const
 	if (InsertAt.IsValid())
 	{
 		// Make sure the sublist has the dropped on item. (It should)
-		const FAvaRundownSubList& ToSubList = Rundown->GetSubList(PageListReference.SubListIndex);
+		const FAvaRundownSubList& ToSubList = Rundown->GetSubList(PageListReference);
 		if (!ToSubList.PageIds.Contains(InsertAt.AdjacentId))
 		{
 			InsertAt.AdjacentId = FAvaRundownPage::InvalidPageId;
@@ -783,7 +803,7 @@ bool SAvaRundownInstancedPageList::HandleDropPageIdsOnSubListFromTemplates(const
 
 	for (const int32 PageId : NewPageIds)
 	{
-		if (Rundown->AddPageToSubList(PageListReference.SubListIndex, PageId, InsertAt))
+		if (Rundown->AddPageToSubList(PageListReference, PageId, InsertAt))
 		{
 			InsertAt.ConditionalUpdateAdjacentId(PageId);
 		}
@@ -814,7 +834,7 @@ bool SAvaRundownInstancedPageList::HandleDropPageIdsOnSubListFromMainList(const
 		return false;
 	}
 
-	FAvaRundownSubList& ToSubList = Rundown->GetSubList(PageListReference.SubListIndex);
+	FAvaRundownSubList& ToSubList = Rundown->GetSubList(PageListReference);
 
 	FAvaRundownPageInsertPosition InsertAt = MakeInsertPosition(InDropZone, InItem);
 	
@@ -896,7 +916,7 @@ bool SAvaRundownInstancedPageList::HandleDropPageIdsOnSubListFromMainList(const
 
 	for (const int32 PageId : ActualPageIds)
 	{
-		Rundown->AddPageToSubList(PageListReference.SubListIndex, PageId, InsertAt);
+		Rundown->AddPageToSubList(PageListReference, PageId, InsertAt);
 		InsertAt.ConditionalUpdateAdjacentId(PageId);
 	}
 
@@ -905,7 +925,7 @@ bool SAvaRundownInstancedPageList::HandleDropPageIdsOnSubListFromMainList(const
 	return true;
 }
 
-bool SAvaRundownInstancedPageList::HandleDropPageIdsOnSubListFromSubList(int32 InFromList, const TArray<int32>& InPageIds, EItemDropZone InDropZone, const FAvaRundownPageViewPtr& InItem)
+bool SAvaRundownInstancedPageList::HandleDropPageIdsOnSubListFromSubList(const FAvaRundownPageListReference& InFromList, const TArray<int32>& InPageIds, EItemDropZone InDropZone, const FAvaRundownPageViewPtr& InItem)
 {
 	if (PageListReference.Type != EAvaRundownPageListType::View)
 	{
@@ -913,12 +933,12 @@ bool SAvaRundownInstancedPageList::HandleDropPageIdsOnSubListFromSubList(int32 I
 	}
 
 	UAvaRundown* Rundown = GetValidRundown();
-	if (!Rundown || !Rundown->IsValidSubList(PageListReference) || !Rundown->IsValidSubListIndex(InFromList))
+	if (!Rundown || !Rundown->IsValidSubList(PageListReference) || !Rundown->IsValidSubList(InFromList))
 	{
 		return false;
 	}
 
-	if (PageListReference.SubListIndex == InFromList)
+	if (PageListReference == InFromList)
 	{
 		if (!Rundown->CanChangePageOrder())
 		{
@@ -933,7 +953,7 @@ bool SAvaRundownInstancedPageList::HandleDropPageIdsOnSubListFromSubList(int32 I
 		}
 	}
 
-	FAvaRundownSubList& ToSubList = Rundown->GetSubList(PageListReference.SubListIndex);
+	FAvaRundownSubList& ToSubList = Rundown->GetSubList(PageListReference);
 
 	FAvaRundownPageInsertPosition InsertAt = MakeInsertPosition(InDropZone, InItem);
 	
@@ -953,7 +973,7 @@ bool SAvaRundownInstancedPageList::HandleDropPageIdsOnSubListFromSubList(int32 I
 	ActualPageIds.Reserve(InPageIds.Num());
 
 	// Nothing to do
-	if (PageListReference.SubListIndex == InFromList)
+	if (PageListReference == InFromList)
 	{
 		if (ToSubList.PageIds.Num() <= 1)
 		{
@@ -1038,7 +1058,7 @@ bool SAvaRundownInstancedPageList::HandleDropPageIdsOnSubListFromSubList(int32 I
 
 	for (const int32 PageId : ActualPageIds)
 	{
-		Rundown->AddPageToSubList(PageListReference.SubListIndex, PageId, InsertAt);
+		Rundown->AddPageToSubList(PageListReference, PageId, InsertAt);
 		InsertAt.ConditionalUpdateAdjacentId(PageId);
 	}
 
@@ -1146,11 +1166,56 @@ void SAvaRundownInstancedPageList::OnPageListChanged(const FAvaRundownPageListCh
 	{
 		return;
 	}
-	
+
+	if (PageListReference.Type == EAvaRundownPageListType::View)
+	{
+		if (const UAvaRundown* Rundown = GetValidRundown())
+		{
+			// If the sublist has been removed, close the tab.
+			if (!Rundown->IsValidSubList(PageListReference))
+			{
+				RequestCloseTab();
+				return;
+			}
+			
+			if (EnumHasAnyFlags(InParams.ChangeType, EAvaRundownPageListChange::SubListRenamed))
+			{
+				if (const TSharedPtr<SDockTab> MyTab = MyTabWeak.Pin())
+				{
+					// Note: Handling undo or renaming to empty -> use factory naming scheme.
+					MyTab->SetLabel(FAvaRundownSubListDocumentTabFactory::GetTabLabel(PageListReference, Rundown));
+				}
+			}
+		}
+	}
+
 	RefreshPagesVisibility();
 	Refresh();
 }
 
+void SAvaRundownInstancedPageList::OnActiveListChanged()
+{
+	// If the sublist has been removed, close the tab.
+	if (PageListReference.Type == EAvaRundownPageListType::View)
+	{
+		if (const UAvaRundown* Rundown = GetValidRundown())
+		{
+			if (!Rundown->IsValidSubList(PageListReference))
+			{
+				RequestCloseTab();
+			}
+		}
+	}
+}
+
+void SAvaRundownInstancedPageList::RequestCloseTab()
+{
+	if (const TSharedPtr<SDockTab> MyTab = MyTabWeak.Pin())
+	{
+		MyTab->RequestCloseTab();
+	}
+}
+
 FReply SAvaRundownInstancedPageList::MakeActive()
 {
 	if (CanMakeActive())
@@ -1193,14 +1258,14 @@ FSlateColor SAvaRundownInstancedPageList::GetMakeActiveButtonColor() const
 
 FText SAvaRundownInstancedPageList::GetPageViewName() const
 {
-	if (UAvaRundown* const Rundown = GetValidRundown())
+	if (const UAvaRundown* const Rundown = GetValidRundown())
 	{
-		if (Rundown->IsValidSubList(PageListReference))
+		const FAvaRundownSubList& SubList = Rundown->GetSubList(PageListReference);		
+		if (SubList.IsValid())
 		{
-			return Rundown->GetSubList(PageListReference.SubListIndex).Name;
+			return SubList.Name;
 		}
 	}
-
 	return FText::GetEmpty();
 }
 
@@ -1215,15 +1280,30 @@ void SAvaRundownInstancedPageList::OnPageViewNameCommitted(const FText& InNewTex
 	{
 		if (Rundown->IsValidSubList(PageListReference))
 		{
-			Rundown->GetSubList(PageListReference.SubListIndex).Name = InNewText;
-			Rundown->GetOnPageListChanged().Broadcast({Rundown, PageListReference, EAvaRundownPageListChange::RenamedPageView, {}});
+			FScopedTransaction Transaction(LOCTEXT("RenamePageView", "Rename Page View"));
+			Rundown->Modify();
+			Rundown->RenameSubList(PageListReference, InNewText);
+		}
+	}
+}
 
-			if (TSharedPtr<SDockTab> MyTab = MyTabWeak.Pin())
+FReply SAvaRundownInstancedPageList::OnDeletePageView()
+{
+	if (UAvaRundown* const Rundown = GetValidRundown())
+	{
+		if (Rundown->IsValidSubList(PageListReference))
+		{
+			FScopedTransaction Transaction(LOCTEXT("DeletePageView", "Delete Page View"));
+			Rundown->Modify();
+			if (Rundown->RemoveSubList(PageListReference))
 			{
-				MyTab->SetLabel(InNewText);
+				return FReply::Handled();
 			}
+
+			Transaction.Cancel();
 		}
 	}
+	return FReply::Unhandled();
 }
 
 TArray<int32> SAvaRundownInstancedPageList::AddPastedPages(const TArray<FAvaRundownPage>& InPages)
diff --git a/Engine/Plugins/Experimental/Avalanche/Source/AvalancheMediaEditor/Private/Rundown/Pages/Slate/SAvaRundownInstancedPageList.h b/Engine/Plugins/Experimental/Avalanche/Source/AvalancheMediaEditor/Private/Rundown/Pages/Slate/SAvaRundownInstancedPageList.h
--- a/Engine/Plugins/Experimental/Avalanche/Source/AvalancheMediaEditor/Private/Rundown/Pages/Slate/SAvaRundownInstancedPageList.h
+++ b/Engine/Plugins/Experimental/Avalanche/Source/AvalancheMediaEditor/Private/Rundown/Pages/Slate/SAvaRundownInstancedPageList.h
@@ -35,7 +35,7 @@ public:
 	bool HandleDropPageIdsOnMainListFromMainList(const TArray<int32>& InPageIds, EItemDropZone InDropZone, const FAvaRundownPageViewPtr& InItem);
 	bool HandleDropPageIdsOnSubListFromTemplates(const TArray<int32>& InPageIds, EItemDropZone InDropZone, const FAvaRundownPageViewPtr& InItem);
 	bool HandleDropPageIdsOnSubListFromMainList(const TArray<int32>& InPageIds, EItemDropZone InDropZone, const FAvaRundownPageViewPtr& InItem);
-	bool HandleDropPageIdsOnSubListFromSubList(int32 InFromList, const TArray<int32>& InPageIds, EItemDropZone InDropZone, const FAvaRundownPageViewPtr& InItem);
+	bool HandleDropPageIdsOnSubListFromSubList(const FAvaRundownPageListReference& InFromList, const TArray<int32>& InPageIds, EItemDropZone InDropZone, const FAvaRundownPageViewPtr& InItem);
 
 	void OnTabActivated(TSharedRef<SDockTab> InDockTab, ETabActivationCause InActivationCause);
 
@@ -70,6 +70,7 @@ public:
 	
 private:
 	void OnPageListChanged(const FAvaRundownPageListChangeParams& InParams);
+	void OnActiveListChanged();
 
 	TArray<int32> FilterPlayingPages(FFilterPageFunctionRef InFilterPageFunction) const;
 	TArray<int32> FilterSelectedOrPlayingPages(FFilterPageFunctionRef InFilterPageFunction, const bool bInAllowFallback) const;
@@ -84,12 +85,15 @@ private:
 protected:
 	TWeakPtr<SDockTab> MyTabWeak;
 
+	void RequestCloseTab();
+	
 	FReply MakeActive();
 	bool CanMakeActive() const;
 	FSlateColor GetMakeActiveButtonColor() const;
 
 	FText GetPageViewName() const;
 	void OnPageViewNameCommitted(const FText& InNewText, ETextCommit::Type InCommitType);
+	FReply OnDeletePageView();
 
 	void PlayNextPageNoReturn() const { PlayNextPage(); }
 
diff --git a/Engine/Plugins/Experimental/Avalanche/Source/AvalancheMediaEditor/Private/Rundown/Pages/Slate/SAvaRundownPageList.cpp b/Engine/Plugins/Experimental/Avalanche/Source/AvalancheMediaEditor/Private/Rundown/Pages/Slate/SAvaRundownPageList.cpp
--- a/Engine/Plugins/Experimental/Avalanche/Source/AvalancheMediaEditor/Private/Rundown/Pages/Slate/SAvaRundownPageList.cpp
+++ b/Engine/Plugins/Experimental/Avalanche/Source/AvalancheMediaEditor/Private/Rundown/Pages/Slate/SAvaRundownPageList.cpp
@@ -607,7 +607,7 @@ void SAvaRundownPageList::RemoveSelectedPages()
 			}
 			else if (Rundown->IsValidSubList(PageListReference))
 			{
-				RemovedCount = Rundown->RemovePagesFromSubList(PageListReference.SubListIndex, SelectedPageIds);
+				RemovedCount = Rundown->RemovePagesFromSubList(PageListReference, SelectedPageIds);
 			}
 
 			if (RemovedCount == 0)
@@ -1164,7 +1164,7 @@ void SAvaRundownPageList::RefreshPagesVisibility()
 
 	if (Rundown->IsValidSubList(PageListReference))
 	{
-		const FAvaRundownSubList& SubList = Rundown->GetSubList(PageListReference.SubListIndex);
+		const FAvaRundownSubList& SubList = Rundown->GetSubList(PageListReference);
 		for (const int32 PageId : SubList.PageIds)
 		{
 			const FAvaRundownPage& Page = Rundown->GetPage(PageId);
@@ -1240,7 +1240,7 @@ void SAvaRundownPageList::OnSearchBoxSuggestionFilter(const FText& InSearchText,
 
 			if (Rundown->IsValidSubList(PageListReference))
 			{
-				for (const int32 PageId : Rundown->GetSubList(PageListReference.SubListIndex).PageIds)
+				for (const int32 PageId : Rundown->GetSubList(PageListReference).PageIds)
 				{
 					for (const TSharedPtr<IAvaRundownFilterSuggestionFactory>& ComplexSuggestion : AvaMediaEditorModule.GetComplexSuggestions(PageListType))
 					{
diff --git a/Engine/Plugins/Experimental/Avalanche/Source/AvalancheMediaEditor/Private/Rundown/Pages/Slate/SAvaRundownSubListStartPage.cpp b/Engine/Plugins/Experimental/Avalanche/Source/AvalancheMediaEditor/Private/Rundown/Pages/Slate/SAvaRundownSubListStartPage.cpp
--- a/Engine/Plugins/Experimental/Avalanche/Source/AvalancheMediaEditor/Private/Rundown/Pages/Slate/SAvaRundownSubListStartPage.cpp
+++ b/Engine/Plugins/Experimental/Avalanche/Source/AvalancheMediaEditor/Private/Rundown/Pages/Slate/SAvaRundownSubListStartPage.cpp
@@ -3,9 +3,10 @@
 #include "SAvaRundownSubListStartPage.h"
 
 #include "Rundown/AvaRundownEditor.h"
+#include "ScopedTransaction.h"
 #include "Styling/AppStyle.h"
 #include "Widgets/Input/SButton.h"
-#include "Widgets/Layout/SBox.h"
+#include "Widgets/SBoxPanel.h"
 #include "Widgets/Text/STextBlock.h"
 
 #define LOCTEXT_NAMESPACE "SAvaRundownSubListStartPage"
@@ -16,10 +17,12 @@ void SAvaRundownSubListStartPage::Construct(const FArguments& InArgs, TSharedPtr
 
 	ChildSlot
 	[
-		SNew(SBox)
+		SNew(SVerticalBox)
+		+ SVerticalBox::Slot()
 		.Padding(10.f, 10.f)
 		.HAlign(EHorizontalAlignment::HAlign_Center)
 		.VAlign(EVerticalAlignment::VAlign_Top)
+		.AutoHeight()
 		[
 			SNew(SButton)
 			.ToolTipText(LOCTEXT("AddSubListTooltip", "Add Page View"))
@@ -29,25 +32,46 @@ void SAvaRundownSubListStartPage::Construct(const FArguments& InArgs, TSharedPtr
 				.Text(LOCTEXT("AddSubList", "Add Page View"))
 			]
 		]
+		+ SVerticalBox::Slot()
+		.Padding(10.f, 10.f)
+		.HAlign(EHorizontalAlignment::HAlign_Center)
+		.VAlign(EVerticalAlignment::VAlign_Top)
+		.AutoHeight()
+		[
+			SNew(SButton)
+			.ToolTipText(LOCTEXT("ShowAllSubListTooltip", "Show All Page Views"))
+			.OnClicked(this, &SAvaRundownSubListStartPage::OnShowAllSubListsClicked)
+			[
+				SNew(STextBlock)
+				.Text(LOCTEXT("ShowAllSubList", "Show All Page Views"))
+			]
+		]
 	];
 }
 
 FReply SAvaRundownSubListStartPage::OnCreateSubListClicked()
 {
-	TSharedPtr<FAvaRundownEditor> RundownEditor = RundownEditorWeak.Pin();
-
-	if (RundownEditor.IsValid())
-	{
-		UAvaRundown* Rundown = RundownEditor->GetRundown();
-
-		if (IsValid(Rundown))
+	if (const TSharedPtr<FAvaRundownEditor> RundownEditor = RundownEditorWeak.Pin())
+	{		
+		if (UAvaRundown* Rundown = RundownEditor->GetRundown())
 		{
-			Rundown->AddSubList();
-
+			FScopedTransaction Transaction(LOCTEXT("AddPageView", "Add PageView"));
+			Rundown->Modify();
+			const FAvaRundownPageListReference CreatedSubListReference = Rundown->AddSubList();
+			Rundown->SetActivePageList(CreatedSubListReference);
 			return FReply::Handled();
 		}
 	}
+	return FReply::Unhandled();
+}
 
+FReply SAvaRundownSubListStartPage::OnShowAllSubListsClicked()
+{
+	if (const TSharedPtr<FAvaRundownEditor> RundownEditor = RundownEditorWeak.Pin())
+	{
+		RundownEditor->RefreshSubListTabs();
+		return FReply::Handled();
+	}
 	return FReply::Unhandled();
 }
 
diff --git a/Engine/Plugins/Experimental/Avalanche/Source/AvalancheMediaEditor/Private/Rundown/Pages/Slate/SAvaRundownSubListStartPage.h b/Engine/Plugins/Experimental/Avalanche/Source/AvalancheMediaEditor/Private/Rundown/Pages/Slate/SAvaRundownSubListStartPage.h
--- a/Engine/Plugins/Experimental/Avalanche/Source/AvalancheMediaEditor/Private/Rundown/Pages/Slate/SAvaRundownSubListStartPage.h
+++ b/Engine/Plugins/Experimental/Avalanche/Source/AvalancheMediaEditor/Private/Rundown/Pages/Slate/SAvaRundownSubListStartPage.h
@@ -19,4 +19,5 @@ protected:
 	TWeakPtr<FAvaRundownEditor> RundownEditorWeak;
 
 	FReply OnCreateSubListClicked();
+	FReply OnShowAllSubListsClicked();
 };
diff --git a/Engine/Plugins/Experimental/Avalanche/Source/AvalancheMediaEditor/Private/Rundown/ShowControl/SAvaRundownShowControl.cpp b/Engine/Plugins/Experimental/Avalanche/Source/AvalancheMediaEditor/Private/Rundown/ShowControl/SAvaRundownShowControl.cpp
--- a/Engine/Plugins/Experimental/Avalanche/Source/AvalancheMediaEditor/Private/Rundown/ShowControl/SAvaRundownShowControl.cpp
+++ b/Engine/Plugins/Experimental/Avalanche/Source/AvalancheMediaEditor/Private/Rundown/ShowControl/SAvaRundownShowControl.cpp
@@ -9,6 +9,7 @@
 #include "Rundown/AvaRundownEditor.h"
 #include "Rundown/Pages/Slate/SAvaRundownInstancedPageList.h"
 #include "Rundown/Pages/Slate/SAvaRundownPageList.h"
+#include "Rundown/TabFactories/AvaRundownSubListDocumentTabFactory.h"
 #include "Widgets/SWidget.h"
 
 #define LOCTEXT_NAMESPACE "SAvaRundownShowControl"
@@ -86,9 +87,9 @@ TSharedRef<SWidget> SAvaRundownShowControl::CreateActiveListWidget()
 
 FText SAvaRundownShowControl::GetActiveListName() const
 {
-	if (TSharedPtr<FAvaRundownEditor> RundownEditor = RundownEditorWeak.Pin())
+	if (const TSharedPtr<FAvaRundownEditor> RundownEditor = RundownEditorWeak.Pin())
 	{
-		UAvaRundown* Rundown = RundownEditor->GetRundown();
+		const UAvaRundown* Rundown = RundownEditor->GetRundown();
 
 		if (IsValid(Rundown))
 		{
@@ -107,7 +108,7 @@ FText SAvaRundownShowControl::GetActiveListName() const
 				return AllPages;
 			}
 
-			return FText::Format(LOCTEXT("RundownSubListDocument_TabLabel", "Page View {0}"), FText::AsNumber(ActiveList.SubListIndex + 1));
+			return FAvaRundownSubListDocumentTabFactory::GetTabLabel(ActiveList, Rundown);
 		}
 	}
 
diff --git a/Engine/Plugins/Experimental/Avalanche/Source/AvalancheMediaEditor/Private/Rundown/TabFactories/AvaRundownSubListDocumentTabFactory.cpp b/Engine/Plugins/Experimental/Avalanche/Source/AvalancheMediaEditor/Private/Rundown/TabFactories/AvaRundownSubListDocumentTabFactory.cpp
--- a/Engine/Plugins/Experimental/Avalanche/Source/AvalancheMediaEditor/Private/Rundown/TabFactories/AvaRundownSubListDocumentTabFactory.cpp
+++ b/Engine/Plugins/Experimental/Avalanche/Source/AvalancheMediaEditor/Private/Rundown/TabFactories/AvaRundownSubListDocumentTabFactory.cpp
@@ -11,54 +11,87 @@
 const FName FAvaRundownSubListDocumentTabFactory::FactoryId = "AvaSubListTabFactory";
 const FString FAvaRundownSubListDocumentTabFactory::BaseTabName(TEXT("AvaSubListDocument"));
 
-FName FAvaRundownSubListDocumentTabFactory::GetTabId(int32 InSubListIndex)
+FName FAvaRundownSubListDocumentTabFactory::GetTabId(const FAvaRundownPageListReference& InSubListReference)
 {
-	return FName(BaseTabName + "_" + FString::FromInt(InSubListIndex));
+	return FName(BaseTabName + "_" + InSubListReference.SubListId.ToString());
+}
+
+FText FAvaRundownSubListDocumentTabFactory::GetTabLabel(const FAvaRundownPageListReference& InSubListReference, const UAvaRundown* InRundown)
+{
+	FText SubListIdAsText;
+	
+	if (IsValid(InRundown))
+	{
+		const FAvaRundownSubList& SubList = InRundown->GetSubList(InSubListReference);
+		if (SubList.IsValid())
+		{
+			if (!SubList.Name.IsEmpty())
+			{
+				return SubList.Name;
+			}
+			
+			// If the sublist has no display name, we will use it's index in the sublists.
+			SubListIdAsText = FText::AsNumber(InRundown->GetSubListIndex(SubList) + 1);
+		}
+	}
+
+	// If the rundown is not available for some reason, use the id.
+	if (SubListIdAsText.IsEmpty())
+	{
+		SubListIdAsText = FText::FromString(InSubListReference.SubListId.ToString());
+	}
+	
+	return FText::Format(LOCTEXT("RundownSubListDocument_TabLabel", "Page View {0}"), SubListIdAsText);
+}
+
+FText FAvaRundownSubListDocumentTabFactory::GetTabDescription(const FAvaRundownPageListReference& InSubListReference, const UAvaRundown* InRundown)
+{
+	const FText SubListId = FText::FromString(InSubListReference.SubListId.ToString());
+	const FText SubListLabel = GetTabLabel(InSubListReference, InRundown);
+	return FText::Format(LOCTEXT("RundownSubListDocument_ViewMenu_Desc", "{0} Id: {1}"), SubListLabel, SubListId);
+}
+
+FText FAvaRundownSubListDocumentTabFactory::GetTabTooltip(const FAvaRundownPageListReference& InSubListReference, const UAvaRundown* InRundown)
+{
+	const FText SubListId = FText::FromString(InSubListReference.SubListId.ToString());
+	const FText SubListLabel = GetTabLabel(InSubListReference, InRundown);
+	return FText::Format(LOCTEXT("RundownSubListDocument_ViewMenu_ToolTip", "{0} Id: {1}"), SubListLabel, SubListId);
 }
 
 FAvaRundownSubListDocumentTabFactory::FAvaRundownSubListDocumentTabFactory(const TSharedPtr<FAvaRundownEditor>& InRundownEditor)
 	: FDocumentTabFactory(FactoryId, InRundownEditor)
 	, RundownEditorWeak(InRundownEditor)
-	, SubListIndex(UAvaRundown::InstancePageList.SubListIndex)
 {
 	TabIcon = FSlateIcon(FAppStyle::GetAppStyleSetName(), "PlacementBrowser.Icons.All");
 }
 
 TSharedRef<SWidget> FAvaRundownSubListDocumentTabFactory::CreateTabBody(const FWorkflowTabSpawnInfo& InInfo) const
 {
-	if (SubListIndex <= UAvaRundown::InstancePageList.SubListIndex)
+	if (!SubListReference.SubListId.IsValid())
 	{
 		return SNullWidget::NullWidget;
 	}
 
-	return SNew(SAvaRundownInstancedPageList, RundownEditorWeak.Pin(), UAvaRundown::CreateSubListReference(SubListIndex));
+	return SNew(SAvaRundownInstancedPageList, RundownEditorWeak.Pin(), SubListReference);
 }
 
-TSharedRef<SDockTab> FAvaRundownSubListDocumentTabFactory::SpawnSubListTab(const FWorkflowTabSpawnInfo& InInfo, int32 InSubListIndex)
+TSharedRef<SDockTab> FAvaRundownSubListDocumentTabFactory::SpawnSubListTab(const FWorkflowTabSpawnInfo& InInfo, const FAvaRundownPageListReference& InSubListReference)
 {
-	SubListIndex = InSubListIndex;
-	TabIdentifier = GetTabId(InSubListIndex);
-	TabLabel = FText::Format(LOCTEXT("RundownSubListDocument_TabLabel", "Page View {0}"), FText::AsNumber(SubListIndex + 1));
-	ViewMenuDescription = FText::Format(LOCTEXT("RundownSubListDocument_ViewMenu_Desc", "Page View {0}"), FText::AsNumber(SubListIndex + 1));
-	ViewMenuTooltip = FText::Format(LOCTEXT("RundownSubListDocument_ViewMenu_ToolTip", "Page View {0}"), FText::AsNumber(SubListIndex + 1));
+	SubListReference = InSubListReference;
+	TabIdentifier = GetTabId(InSubListReference);
+	
+	const TSharedPtr<FAvaRundownEditor> RundownEditor = RundownEditorWeak.Pin();
+	const UAvaRundown* Rundown = RundownEditor ? RundownEditor->GetRundown() : nullptr;
+	TabLabel = GetTabLabel(InSubListReference, Rundown);
+	ViewMenuDescription = GetTabDescription(InSubListReference, Rundown);
+	ViewMenuTooltip = GetTabTooltip(InSubListReference, Rundown);
 
 	TSharedRef<SDockTab> NewTab = SpawnTab(InInfo);
 
-	TSharedRef<SAvaRundownInstancedPageList> PageList = StaticCastSharedRef<SAvaRundownInstancedPageList>(NewTab->GetContent());
+	const TSharedRef<SAvaRundownInstancedPageList> PageList = StaticCastSharedRef<SAvaRundownInstancedPageList>(NewTab->GetContent());
 	NewTab->SetOnTabActivated(SDockTab::FOnTabActivatedCallback::CreateSP(PageList, &SAvaRundownInstancedPageList::OnTabActivated));
 	PageList->SetMyTab(NewTab);
 
-	if (TSharedPtr<FAvaRundownEditor> RundownEditor = PageList->GetRundownEditor())
-	{
-		UAvaRundown* Rundown = RundownEditor->GetRundown();
-
-		if (IsValid(Rundown) && Rundown->IsValidSubListIndex(InSubListIndex)
-			&& !Rundown->GetSubList(InSubListIndex).Name.IsEmpty())
-		{
-			NewTab->SetLabel(Rundown->GetSubList(InSubListIndex).Name);
-		}
-	}
-
 	return NewTab;
 }
 
diff --git a/Engine/Plugins/Experimental/Avalanche/Source/AvalancheMediaEditor/Private/Rundown/TabFactories/AvaRundownSubListDocumentTabFactory.h b/Engine/Plugins/Experimental/Avalanche/Source/AvalancheMediaEditor/Private/Rundown/TabFactories/AvaRundownSubListDocumentTabFactory.h
--- a/Engine/Plugins/Experimental/Avalanche/Source/AvalancheMediaEditor/Private/Rundown/TabFactories/AvaRundownSubListDocumentTabFactory.h
+++ b/Engine/Plugins/Experimental/Avalanche/Source/AvalancheMediaEditor/Private/Rundown/TabFactories/AvaRundownSubListDocumentTabFactory.h
@@ -3,8 +3,10 @@
 #pragma once
 
 #include "WorkflowOrientedApp/WorkflowTabFactory.h"
+#include "AvaMediaDefines.h"
 
 class FAvaRundownEditor;
+class UAvaRundown;
 
 /** Base class for all Tab Factories in Ava SubListDocument Editor */
 class FAvaRundownSubListDocumentTabFactory : public FDocumentTabFactory
@@ -13,16 +15,20 @@ public:
 	static const FName FactoryId;
 	static const FString BaseTabName;
 
-	static FName GetTabId(int32 InSubListIndex);
+	static FName GetTabId(const FAvaRundownPageListReference& InSubListReference);
+	static FText GetTabLabel(const FAvaRundownPageListReference& InSubListReference, const UAvaRundown* InRundown);
+	static FText GetTabDescription(const FAvaRundownPageListReference& InSubListReference, const UAvaRundown* InRundown);
+	static FText GetTabTooltip(const FAvaRundownPageListReference& InSubListReference, const UAvaRundown* InRundown);
 
 	FAvaRundownSubListDocumentTabFactory(const TSharedPtr<FAvaRundownEditor>& InSubListDocumentEditor);
 
 	virtual TSharedRef<SWidget> CreateTabBody(const FWorkflowTabSpawnInfo& InInfo) const override;
 
-	TSharedRef<SDockTab> SpawnSubListTab(const FWorkflowTabSpawnInfo& InInfo, int32 InSubListIndex);
+	TSharedRef<SDockTab> SpawnSubListTab(const FWorkflowTabSpawnInfo& InInfo, const FAvaRundownPageListReference& InSubListReference);
 
 protected:
 	TWeakPtr<FAvaRundownEditor> RundownEditorWeak;
-	int32 SubListIndex;
+	FAvaRundownPageListReference SubListReference;
+	FText SubListDisplayName;
 };
 
diff --git a/Engine/Plugins/Experimental/Avalanche/Source/AvalancheMediaEditor/Public/Rundown/AvaRundownEditor.h b/Engine/Plugins/Experimental/Avalanche/Source/AvalancheMediaEditor/Public/Rundown/AvaRundownEditor.h
--- a/Engine/Plugins/Experimental/Avalanche/Source/AvalancheMediaEditor/Public/Rundown/AvaRundownEditor.h
+++ b/Engine/Plugins/Experimental/Avalanche/Source/AvalancheMediaEditor/Public/Rundown/AvaRundownEditor.h
@@ -103,6 +103,8 @@ public:
 	bool CanRemoveSelectedPages() const;
 	void RemoveSelectedPages();
 
+	void RefreshSubListTabs();
+
 protected:
 	void RegisterApplicationModes();
 
@@ -159,10 +161,12 @@ protected:
 	};
 	TUniquePtr<FAutoPlayTicker> AutoPlayTicker;
 
-	TSharedPtr<SDockTab> CreateSubListTab(int32 InSubListIndex);
-	void CreateSubListTabs();
+	TSharedPtr<SDockTab> CreateSubListTab(const FAvaRundownPageListReference& InSubListReference);
+
+	void RefreshSubListTab(const FAvaRundownPageListReference& InSubListReference, bool bInSetActive = false);
 
-	void OnActiveSubListChanged();
+	void OnPageListChanged(const FAvaRundownPageListChangeParams& InParams);
+	void OnActiveListChanged();
 	void HandleOnPagePlayerAdded(UAvaRundown* InRundown, UAvaRundownPagePlayer* InPagePlayer);
 	void OnCanClosePlaybackContext(const UAvaRundown* InRundown, bool& bOutResult) const;
 
@@ -170,8 +174,6 @@ protected:
 
 	FOnPageEvent OnPageEvent;
 
-	TMap<FName, TWeakPtr<SDockTab>> SubListTabs;
-
 private:
 	TSharedPtr<SAvaRundownReadPage> ReadPageWidget;
 
