[Motion Design] Improving Page Views Workflow: Tab Spawner registered in workspace menu.

Page View tabs are now visible in the "Windows" menu and can be opened from there.

diff --git a/Engine/Plugins/Experimental/Avalanche/Source/AvalancheMedia/Private/Rundown/AvaRundown.cpp b/Engine/Plugins/Experimental/Avalanche/Source/AvalancheMedia/Private/Rundown/AvaRundown.cpp
--- a/Engine/Plugins/Experimental/Avalanche/Source/AvalancheMedia/Private/Rundown/AvaRundown.cpp
+++ b/Engine/Plugins/Experimental/Avalanche/Source/AvalancheMedia/Private/Rundown/AvaRundown.cpp
@@ -1346,7 +1346,7 @@ bool UAvaRundown::RemoveSubList(const FAvaRundownPageListReference& InPageListRe
 			if (SubLists.IsValidIndex(*SubListIndex))
 			{
 				SubLists.RemoveAt(*SubListIndex);
-				SubListIndices.Remove(InPageListReference.SubListId);
+				RefreshSubListIndices();
 				GetOnPageListChanged().Broadcast({this, InPageListReference, EAvaRundownPageListChange::SubListAddedOrRemoved, {}});
 				return true;
 			}
diff --git a/Engine/Plugins/Experimental/Avalanche/Source/AvalancheMediaEditor/Private/Rundown/AppModes/AvaRundownAppMode.cpp b/Engine/Plugins/Experimental/Avalanche/Source/AvalancheMediaEditor/Private/Rundown/AppModes/AvaRundownAppMode.cpp
--- a/Engine/Plugins/Experimental/Avalanche/Source/AvalancheMediaEditor/Private/Rundown/AppModes/AvaRundownAppMode.cpp
+++ b/Engine/Plugins/Experimental/Avalanche/Source/AvalancheMediaEditor/Private/Rundown/AppModes/AvaRundownAppMode.cpp
@@ -16,8 +16,8 @@ FAvaRundownAppMode::FAvaRundownAppMode(const TSharedPtr<FAvaRundownEditor>& InRu
 
 void FAvaRundownAppMode::RegisterTabFactories(TSharedPtr<FTabManager> InTabManager)
 {
-	TSharedPtr<FAvaRundownEditor> PlaybackEditor = RundownEditorWeak.Pin();
-	PlaybackEditor->PushTabFactories(TabFactories);
+	const TSharedPtr<FAvaRundownEditor> RundownEditor = RundownEditorWeak.Pin();
+	RundownEditor->PushTabFactories(TabFactories);
 	FApplicationMode::RegisterTabFactories(InTabManager);
 }
 
@@ -27,10 +27,52 @@ TSharedPtr<FDocumentTabFactory> FAvaRundownAppMode::GetDocumentTabFactory(const
 	{
 		return *DocFactory;
 	}
-
 	return nullptr;
 }
 
+void FAvaRundownAppMode::RegisterDocumentTabFactory(const TSharedPtr<FDocumentTabFactory>& InDocumentTabFactory, const TSharedPtr<FTabManager>& InTabManager)
+{
+	if (!InDocumentTabFactory)
+	{
+		return;
+	}
+	
+	// Keep track of this "document" factory so we don't add it again.
+	DocumentTabFactories.Add(InDocumentTabFactory->GetIdentifier(), InDocumentTabFactory.ToSharedRef());
+
+	// Add to the tab manager as a tab spawner.
+	if (InTabManager)
+	{
+		if (InTabManager->HasTabSpawner(InDocumentTabFactory->GetIdentifier()))
+		{
+			InTabManager->UnregisterTabSpawner(InDocumentTabFactory->GetIdentifier());
+		}
+		
+		InDocumentTabFactory->RegisterTabSpawner(InTabManager.ToSharedRef(), this);
+	}
+
+	// Update "Allowed" tab factories. (May not be necessary)
+	if (!TabFactories.GetFactory(InDocumentTabFactory->GetIdentifier()).IsValid())
+	{
+		TabFactories.RegisterFactory(InDocumentTabFactory);
+	}
+}
+
+void FAvaRundownAppMode::UnregisterDocumentTabFactory(const FName& InTabId, const TSharedPtr<FTabManager>& InTabManager)
+{
+	if (InTabManager && InTabManager->HasTabSpawner(InTabId))
+	{
+		InTabManager->UnregisterTabSpawner(InTabId);
+	}
+
+	if (TabFactories.GetFactory(InTabId).IsValid())
+	{
+		TabFactories.UnregisterFactory(InTabId);
+	}
+
+	DocumentTabFactories.Remove(InTabId);
+}
+
 FText FAvaRundownAppMode::GetLocalizedMode(const FName InMode)
 {
 	static TMap<FName, FText> LocModes;
diff --git a/Engine/Plugins/Experimental/Avalanche/Source/AvalancheMediaEditor/Private/Rundown/AppModes/AvaRundownAppMode.h b/Engine/Plugins/Experimental/Avalanche/Source/AvalancheMediaEditor/Private/Rundown/AppModes/AvaRundownAppMode.h
--- a/Engine/Plugins/Experimental/Avalanche/Source/AvalancheMediaEditor/Private/Rundown/AppModes/AvaRundownAppMode.h
+++ b/Engine/Plugins/Experimental/Avalanche/Source/AvalancheMediaEditor/Private/Rundown/AppModes/AvaRundownAppMode.h
@@ -20,6 +20,9 @@ public:
 
 	TSharedPtr<FDocumentTabFactory> GetDocumentTabFactory(const FName& InName) const;
 
+	void RegisterDocumentTabFactory(const TSharedPtr<FDocumentTabFactory>& InDocumentTabFactory, const TSharedPtr<FTabManager>& InTabManager = nullptr);
+	void UnregisterDocumentTabFactory(const FName& InTabId, const TSharedPtr<FTabManager>& InTabManager);
+
 protected:
 	
 	static FText GetLocalizedMode(const FName InMode);
diff --git a/Engine/Plugins/Experimental/Avalanche/Source/AvalancheMediaEditor/Private/Rundown/AppModes/AvaRundownDefaultMode.cpp b/Engine/Plugins/Experimental/Avalanche/Source/AvalancheMediaEditor/Private/Rundown/AppModes/AvaRundownDefaultMode.cpp
--- a/Engine/Plugins/Experimental/Avalanche/Source/AvalancheMediaEditor/Private/Rundown/AppModes/AvaRundownDefaultMode.cpp
+++ b/Engine/Plugins/Experimental/Avalanche/Source/AvalancheMediaEditor/Private/Rundown/AppModes/AvaRundownDefaultMode.cpp
@@ -121,7 +121,17 @@ FAvaRundownDefaultMode::FAvaRundownDefaultMode(const TSharedPtr<FAvaRundownEdito
 	TabFactories.RegisterFactory(MakeShared<FAvaRundownChannelStatusListTabFactory>(InRundownEditor));
 	TabFactories.RegisterFactory(MakeShared<FAvaRundownChannelLayerStatusListTabFactory>(InRundownEditor));
 
-	DocumentTabFactories.Emplace(FAvaRundownSubListDocumentTabFactory::FactoryId, MakeShared<FAvaRundownSubListDocumentTabFactory>(InRundownEditor));
+	// Document factories for existing "documents" within the rundown.
+	if (const UAvaRundown* Rundown = InRundownEditor->GetRundown())
+	{
+		for (const FAvaRundownSubList& SubList : Rundown->GetSubLists())
+		{
+			TSharedPtr<FAvaRundownSubListDocumentTabFactory> PageViewDocumentFactory =
+				MakeShared<FAvaRundownSubListDocumentTabFactory>(UAvaRundown::CreateSubListReference(SubList), InRundownEditor);
+
+			RegisterDocumentTabFactory(PageViewDocumentFactory);
+		}
+	}
 
 	// Make sure we start with our existing list of extenders instead of creating a new one
 	IAvaMediaEditorModule& AvaMediaEditorModule = IAvaMediaEditorModule::Get();
diff --git a/Engine/Plugins/Experimental/Avalanche/Source/AvalancheMediaEditor/Private/Rundown/AvaRundownEditor.cpp b/Engine/Plugins/Experimental/Avalanche/Source/AvalancheMediaEditor/Private/Rundown/AvaRundownEditor.cpp
--- a/Engine/Plugins/Experimental/Avalanche/Source/AvalancheMediaEditor/Private/Rundown/AvaRundownEditor.cpp
+++ b/Engine/Plugins/Experimental/Avalanche/Source/AvalancheMediaEditor/Private/Rundown/AvaRundownEditor.cpp
@@ -255,8 +255,6 @@ void FAvaRundownEditor::InitRundownEditor(const EToolkitMode::Type InMode
 
 	RegisterApplicationModes();
 	FSharedConsoleCommands::RegisterEditor(SharedThis(this));
-
-	RefreshSubListTabs();
 }
 
 bool FAvaRundownEditor::IsKeyRelevant(const FKeyEvent& InKeyEvent) const
@@ -811,6 +809,28 @@ void FAvaRundownEditor::RefreshSubListTabs()
 	}
 }
 
+bool FAvaRundownEditor::RequestCloseDocumentTab(const FName& InDocumentTabId)
+{
+	if (TabManager)
+	{
+		if (const TSharedPtr<SDockTab> DocumentTab = TabManager->FindExistingLiveTab(InDocumentTabId))
+		{
+			return DocumentTab->RequestCloseTab();
+		}
+	}
+	return false;
+}
+
+void FAvaRundownEditor::UnregisterDocumentTabFactory(const FName& InDocumentTabId)
+{
+	const TSharedPtr<FApplicationMode> AppMode = GetCurrentModePtr();
+	if (AppMode.IsValid())
+	{
+		const TSharedRef<FAvaRundownAppMode> RundownAppMode = StaticCastSharedRef<FAvaRundownAppMode>(AppMode.ToSharedRef());
+		RundownAppMode->UnregisterDocumentTabFactory(InDocumentTabId, TabManager);
+	}
+}
+
 FName FAvaRundownEditor::GetToolkitFName() const
 {
 	return TEXT("AvaRundownEditor");
@@ -984,25 +1004,15 @@ TSharedPtr<SDockTab> FAvaRundownEditor::CreateSubListTab(const FAvaRundownPageLi
 	{
 		const FName TabId = FAvaRundownSubListDocumentTabFactory::GetTabId(InSubListReference);
 		const TSharedRef<FAvaRundownAppMode> RundownAppMode = StaticCastSharedRef<FAvaRundownAppMode>(AppMode.ToSharedRef());
-		const TSharedPtr<FDocumentTabFactory> DocTabFactory = RundownAppMode->GetDocumentTabFactory(FAvaRundownSubListDocumentTabFactory::FactoryId);
+		TSharedPtr<FDocumentTabFactory> DocTabFactory = RundownAppMode->GetDocumentTabFactory(TabId);
 
-		if (DocTabFactory.IsValid())
+		if (!DocTabFactory.IsValid())
 		{
-			const TSharedRef<FAvaRundownSubListDocumentTabFactory> SubListTabFactory = StaticCastSharedRef<FAvaRundownSubListDocumentTabFactory>(DocTabFactory.ToSharedRef());
-
-			FWorkflowTabSpawnInfo Info;
-			Info.TabManager = TabManager;
-			Info.Payload = nullptr;
-			Info.TabInfo = nullptr;
-
-			TSharedPtr<SDockTab> SubListTab = SubListTabFactory->SpawnSubListTab(Info, InSubListReference);
-
-			if (SubListTab.IsValid())
-			{
-				TabManager->InsertNewDocumentTab(FAvaRundownSubListTabFactory::TabID, TabId, FTabManager::FLiveTabSearch(TabId), SubListTab.ToSharedRef());
-				return SubListTab;
-			}
+			DocTabFactory = MakeShared<FAvaRundownSubListDocumentTabFactory>(InSubListReference, SharedThis(this));
+			RundownAppMode->RegisterDocumentTabFactory(DocTabFactory, TabManager);
 		}
+
+		return TabManager->TryInvokeTab(TabId);
 	}
 
 	return nullptr;
@@ -1043,6 +1053,7 @@ void FAvaRundownEditor::RefreshSubListTab(const FAvaRundownPageListReference& In
 		if (SubListTab.IsValid())
 		{
 			SubListTab->RequestCloseTab();
+			UnregisterDocumentTabFactory(TabId);
 		}
 	}
 }
diff --git a/Engine/Plugins/Experimental/Avalanche/Source/AvalancheMediaEditor/Private/Rundown/Pages/Slate/SAvaRundownInstancedPageList.cpp b/Engine/Plugins/Experimental/Avalanche/Source/AvalancheMediaEditor/Private/Rundown/Pages/Slate/SAvaRundownInstancedPageList.cpp
--- a/Engine/Plugins/Experimental/Avalanche/Source/AvalancheMediaEditor/Private/Rundown/Pages/Slate/SAvaRundownInstancedPageList.cpp
+++ b/Engine/Plugins/Experimental/Avalanche/Source/AvalancheMediaEditor/Private/Rundown/Pages/Slate/SAvaRundownInstancedPageList.cpp
@@ -1177,15 +1177,6 @@ void SAvaRundownInstancedPageList::OnPageListChanged(const FAvaRundownPageListCh
 				RequestCloseTab();
 				return;
 			}
-			
-			if (EnumHasAnyFlags(InParams.ChangeType, EAvaRundownPageListChange::SubListRenamed))
-			{
-				if (const TSharedPtr<SDockTab> MyTab = MyTabWeak.Pin())
-				{
-					// Note: Handling undo or renaming to empty -> use factory naming scheme.
-					MyTab->SetLabel(FAvaRundownSubListDocumentTabFactory::GetTabLabel(PageListReference, Rundown));
-				}
-			}
 		}
 	}
 
@@ -1210,9 +1201,11 @@ void SAvaRundownInstancedPageList::OnActiveListChanged()
 
 void SAvaRundownInstancedPageList::RequestCloseTab()
 {
-	if (const TSharedPtr<SDockTab> MyTab = MyTabWeak.Pin())
+	if (const TSharedPtr<FAvaRundownEditor> RundownEditor = GetRundownEditor())
 	{
-		MyTab->RequestCloseTab();
+		const FName DocumentTabId = FAvaRundownSubListDocumentTabFactory::GetTabId(PageListReference);
+		RundownEditor->RequestCloseDocumentTab(DocumentTabId);
+		RundownEditor->UnregisterDocumentTabFactory(DocumentTabId);
 	}
 }
 
diff --git a/Engine/Plugins/Experimental/Avalanche/Source/AvalancheMediaEditor/Private/Rundown/Pages/Slate/SAvaRundownInstancedPageList.h b/Engine/Plugins/Experimental/Avalanche/Source/AvalancheMediaEditor/Private/Rundown/Pages/Slate/SAvaRundownInstancedPageList.h
--- a/Engine/Plugins/Experimental/Avalanche/Source/AvalancheMediaEditor/Private/Rundown/Pages/Slate/SAvaRundownInstancedPageList.h
+++ b/Engine/Plugins/Experimental/Avalanche/Source/AvalancheMediaEditor/Private/Rundown/Pages/Slate/SAvaRundownInstancedPageList.h
@@ -39,9 +39,6 @@ public:
 
 	void OnTabActivated(TSharedRef<SDockTab> InDockTab, ETabActivationCause InActivationCause);
 
-	TSharedPtr<SDockTab> GetMyTab() const { return MyTabWeak.Pin(); }
-	void SetMyTab(TSharedRef<SDockTab> InTab) { MyTabWeak = InTab; }
-
 	/** Plays the currently selected page. */
 	void PlaySelectedPage() const;
 	bool CanPlaySelectedPage() const;
@@ -83,8 +80,6 @@ private:
 	int32 GetPageIdToTakeNext() const;
 
 protected:
-	TWeakPtr<SDockTab> MyTabWeak;
-
 	void RequestCloseTab();
 	
 	FReply MakeActive();
diff --git a/Engine/Plugins/Experimental/Avalanche/Source/AvalancheMediaEditor/Private/Rundown/TabFactories/AvaRundownSubListDocumentTabFactory.cpp b/Engine/Plugins/Experimental/Avalanche/Source/AvalancheMediaEditor/Private/Rundown/TabFactories/AvaRundownSubListDocumentTabFactory.cpp
--- a/Engine/Plugins/Experimental/Avalanche/Source/AvalancheMediaEditor/Private/Rundown/TabFactories/AvaRundownSubListDocumentTabFactory.cpp
+++ b/Engine/Plugins/Experimental/Avalanche/Source/AvalancheMediaEditor/Private/Rundown/TabFactories/AvaRundownSubListDocumentTabFactory.cpp
@@ -2,9 +2,11 @@
 
 #include "AvaRundownSubListDocumentTabFactory.h"
 
+#include "Rundown/AvaRundown.h"
 #include "Rundown/AvaRundownEditor.h"
 #include "Rundown/Pages/Slate/SAvaRundownInstancedPageList.h"
 #include "Widgets/Docking/SDockTab.h"
+#include "WorkflowOrientedApp/ApplicationMode.h"
 
 #define LOCTEXT_NAMESPACE "AvaRundownSubListDocumentTabFactory"
 
@@ -58,11 +60,33 @@ FText FAvaRundownSubListDocumentTabFactory::GetTabTooltip(const FAvaRundownPageL
 	return FText::Format(LOCTEXT("RundownSubListDocument_ViewMenu_ToolTip", "{0} Id: {1}"), SubListLabel, SubListId);
 }
 
-FAvaRundownSubListDocumentTabFactory::FAvaRundownSubListDocumentTabFactory(const TSharedPtr<FAvaRundownEditor>& InRundownEditor)
-	: FDocumentTabFactory(FactoryId, InRundownEditor)
+FAvaRundownSubListDocumentTabFactory::FAvaRundownSubListDocumentTabFactory(const FAvaRundownPageListReference& InSubListReference, const TSharedPtr<FAvaRundownEditor>& InRundownEditor)
+	: FDocumentTabFactory(GetTabId(InSubListReference), InRundownEditor)
 	, RundownEditorWeak(InRundownEditor)
 {
 	TabIcon = FSlateIcon(FAppStyle::GetAppStyleSetName(), "PlacementBrowser.Icons.All");
+	SubListReference = InSubListReference;
+	TabIdentifier = GetTabId(InSubListReference);
+	
+	UAvaRundown* Rundown = InRundownEditor ? InRundownEditor->GetRundown() : nullptr;
+	TabLabel = GetTabLabel(InSubListReference, Rundown);
+	ViewMenuDescription = GetTabDescription(InSubListReference, Rundown);
+	ViewMenuTooltip = GetTabTooltip(InSubListReference, Rundown);
+
+	if (Rundown)
+	{
+		// Allow propagation of tab label changes.
+		Rundown->GetOnPageListChanged().AddRaw(this, &FAvaRundownSubListDocumentTabFactory::OnPageListChanged);
+	}
+}
+
+FAvaRundownSubListDocumentTabFactory::~FAvaRundownSubListDocumentTabFactory()
+{
+	const TSharedPtr<FAvaRundownEditor> RundownEditor = RundownEditorWeak.Pin();
+	if (UAvaRundown* Rundown = RundownEditor ? RundownEditor->GetRundown() : nullptr)
+	{
+		Rundown->GetOnPageListChanged().RemoveAll(this);
+	}
 }
 
 TSharedRef<SWidget> FAvaRundownSubListDocumentTabFactory::CreateTabBody(const FWorkflowTabSpawnInfo& InInfo) const
@@ -75,24 +99,80 @@ TSharedRef<SWidget> FAvaRundownSubListDocumentTabFactory::CreateTabBody(const FW
 	return SNew(SAvaRundownInstancedPageList, RundownEditorWeak.Pin(), SubListReference);
 }
 
-TSharedRef<SDockTab> FAvaRundownSubListDocumentTabFactory::SpawnSubListTab(const FWorkflowTabSpawnInfo& InInfo, const FAvaRundownPageListReference& InSubListReference)
+FTabSpawnerEntry& FAvaRundownSubListDocumentTabFactory::RegisterTabSpawner(TSharedRef<FTabManager> InTabManager, const FApplicationMode* InCurrentApplicationMode) const
 {
-	SubListReference = InSubListReference;
-	TabIdentifier = GetTabId(InSubListReference);
+	const FText PageViewsGroupName = LOCTEXT("SubMenuLabel", "Page Views");
+	TSharedPtr<FWorkspaceItem> GroupItem;
+
+	// Find or Add the PageView Group Item.
+	if (InCurrentApplicationMode)
+	{
+		if (const TSharedPtr<FWorkspaceItem> WorkspaceMenu = InCurrentApplicationMode->GetWorkspaceMenuCategory())
+		{
+			for (const TSharedRef<FWorkspaceItem>& Item : WorkspaceMenu->GetChildItems())
+			{
+				if (Item->GetDisplayName().ToString() == PageViewsGroupName.ToString())
+				{
+					GroupItem = Item;
+					break;
+				}
+			}
+		
+			if (!GroupItem.IsValid())
+			{
+				GroupItem = WorkspaceMenu->AddGroup(PageViewsGroupName, LOCTEXT("SubMenuTooltip", "Motion Design Rundown Page Views"), TabIcon);
+			}
+		}
+	}
 	
-	const TSharedPtr<FAvaRundownEditor> RundownEditor = RundownEditorWeak.Pin();
-	const UAvaRundown* Rundown = RundownEditor ? RundownEditor->GetRundown() : nullptr;
-	TabLabel = GetTabLabel(InSubListReference, Rundown);
-	ViewMenuDescription = GetTabDescription(InSubListReference, Rundown);
-	ViewMenuTooltip = GetTabTooltip(InSubListReference, Rundown);
+	//Note: We don't provide the application mode to avoid setting the group.
+	FTabSpawnerEntry& SpawnerEntry = FDocumentTabFactory::RegisterTabSpawner(InTabManager, GroupItem.IsValid() ? nullptr : InCurrentApplicationMode);
 
-	TSharedRef<SDockTab> NewTab = SpawnTab(InInfo);
+	if (GroupItem.IsValid())
+	{
+		SpawnerEntry.SetGroup(GroupItem.ToSharedRef());
+	}
 
-	const TSharedRef<SAvaRundownInstancedPageList> PageList = StaticCastSharedRef<SAvaRundownInstancedPageList>(NewTab->GetContent());
-	NewTab->SetOnTabActivated(SDockTab::FOnTabActivatedCallback::CreateSP(PageList, &SAvaRundownInstancedPageList::OnTabActivated));
-	PageList->SetMyTab(NewTab);
+	// Bind the spawner entry label to the current document factory label.
+	SpawnerEntry.SetDisplayNameAttribute(TAttribute<FText>( this, &FAvaRundownSubListDocumentTabFactory::GetTabTitle));
+	
+	return SpawnerEntry;
+}
+
+TSharedRef<SDockTab> FAvaRundownSubListDocumentTabFactory::OnSpawnTab(const FSpawnTabArgs& InSpawnArgs, TWeakPtr<FTabManager> InTabManagerWeak) const
+{
+	// Intercept the spawned tab to bind our handlers.
+	TSharedRef<SDockTab> SpawnedTab = FDocumentTabFactory::OnSpawnTab(InSpawnArgs, InTabManagerWeak);
+
+	if (SpawnedTab != SNullWidget::NullWidget)
+	{
+		const TSharedRef<SAvaRundownInstancedPageList> PageList = StaticCastSharedRef<SAvaRundownInstancedPageList>(SpawnedTab->GetContent());
+		SpawnedTab->SetOnTabActivated(SDockTab::FOnTabActivatedCallback::CreateSP(PageList, &SAvaRundownInstancedPageList::OnTabActivated));
+	}
 
-	return NewTab;
+	// Bind the tab label to the current document factory label.
+	SpawnedTab->SetLabel(TAttribute<FText>( this, &FAvaRundownSubListDocumentTabFactory::GetTabTitle));
+
+	return SpawnedTab;
+}
+
+FText FAvaRundownSubListDocumentTabFactory::GetTabTitle() const
+{
+	return TabLabel;
+}
+
+void FAvaRundownSubListDocumentTabFactory::OnPageListChanged(const FAvaRundownPageListChangeParams& InParams)
+{
+	if (SubListReference != InParams.PageListReference)
+	{
+		return;
+	}
+
+	if (EnumHasAnyFlags(InParams.ChangeType, EAvaRundownPageListChange::SubListRenamed))
+	{
+		// This will automatically propagate to the SpawnerEntry (Windows menu) and the Spawned Tab Label.
+		TabLabel = GetTabLabel(InParams.PageListReference, InParams.Rundown);
+	}
 }
 
 #undef LOCTEXT_NAMESPACE
diff --git a/Engine/Plugins/Experimental/Avalanche/Source/AvalancheMediaEditor/Private/Rundown/TabFactories/AvaRundownSubListDocumentTabFactory.h b/Engine/Plugins/Experimental/Avalanche/Source/AvalancheMediaEditor/Private/Rundown/TabFactories/AvaRundownSubListDocumentTabFactory.h
--- a/Engine/Plugins/Experimental/Avalanche/Source/AvalancheMediaEditor/Private/Rundown/TabFactories/AvaRundownSubListDocumentTabFactory.h
+++ b/Engine/Plugins/Experimental/Avalanche/Source/AvalancheMediaEditor/Private/Rundown/TabFactories/AvaRundownSubListDocumentTabFactory.h
@@ -7,6 +7,7 @@
 
 class FAvaRundownEditor;
 class UAvaRundown;
+struct FAvaRundownPageListChangeParams;
 
 /** Base class for all Tab Factories in Ava SubListDocument Editor */
 class FAvaRundownSubListDocumentTabFactory : public FDocumentTabFactory
@@ -20,15 +21,24 @@ public:
 	static FText GetTabDescription(const FAvaRundownPageListReference& InSubListReference, const UAvaRundown* InRundown);
 	static FText GetTabTooltip(const FAvaRundownPageListReference& InSubListReference, const UAvaRundown* InRundown);
 
-	FAvaRundownSubListDocumentTabFactory(const TSharedPtr<FAvaRundownEditor>& InSubListDocumentEditor);
+	FAvaRundownSubListDocumentTabFactory(const FAvaRundownPageListReference& InSubListReference, const TSharedPtr<FAvaRundownEditor>& InSubListDocumentEditor);
 
+	virtual ~FAvaRundownSubListDocumentTabFactory() override;
+
+	//~ Begin FWorkflowTabFactory
 	virtual TSharedRef<SWidget> CreateTabBody(const FWorkflowTabSpawnInfo& InInfo) const override;
+	virtual FTabSpawnerEntry& RegisterTabSpawner(TSharedRef<FTabManager> InTabManager, const FApplicationMode* InCurrentApplicationMode) const override;
+protected:
+	virtual TSharedRef<SDockTab> OnSpawnTab(const FSpawnTabArgs& InSpawnArgs, TWeakPtr<FTabManager> InTabManagerWeak) const override;
+	//~ End FWorkflowTabFactory
 
-	TSharedRef<SDockTab> SpawnSubListTab(const FWorkflowTabSpawnInfo& InInfo, const FAvaRundownPageListReference& InSubListReference);
+protected:
+	FText GetTabTitle() const;
 
+	void OnPageListChanged(const FAvaRundownPageListChangeParams& InParams);
+	
 protected:
 	TWeakPtr<FAvaRundownEditor> RundownEditorWeak;
 	FAvaRundownPageListReference SubListReference;
-	FText SubListDisplayName;
 };
 
diff --git a/Engine/Plugins/Experimental/Avalanche/Source/AvalancheMediaEditor/Public/Rundown/AvaRundownEditor.h b/Engine/Plugins/Experimental/Avalanche/Source/AvalancheMediaEditor/Public/Rundown/AvaRundownEditor.h
--- a/Engine/Plugins/Experimental/Avalanche/Source/AvalancheMediaEditor/Public/Rundown/AvaRundownEditor.h
+++ b/Engine/Plugins/Experimental/Avalanche/Source/AvalancheMediaEditor/Public/Rundown/AvaRundownEditor.h
@@ -105,6 +105,10 @@ public:
 
 	void RefreshSubListTabs();
 
+	bool RequestCloseDocumentTab(const FName& InDocumentTabId);
+
+	void UnregisterDocumentTabFactory(const FName& InDocumentTabId);
+
 protected:
 	void RegisterApplicationModes();
 
