[Motion Design] Playable BP Library: Adding a function to hide the playable's primitives.

diff --git a/Engine/Plugins/Experimental/Avalanche/Source/AvalancheMedia/Private/Playable/AvaPlayableLibrary.cpp b/Engine/Plugins/Experimental/Avalanche/Source/AvalancheMedia/Private/Playable/AvaPlayableLibrary.cpp
--- a/Engine/Plugins/Experimental/Avalanche/Source/AvalancheMedia/Private/Playable/AvaPlayableLibrary.cpp
+++ b/Engine/Plugins/Experimental/Avalanche/Source/AvalancheMedia/Private/Playable/AvaPlayableLibrary.cpp
@@ -8,6 +8,7 @@
 #include "Playable/AvaPlayable.h"
 #include "Playable/AvaPlayableAssetUserData.h"
 #include "Playable/AvaPlayableGroup.h"
+#include "Playable/Playables/AvaPlayableLevelStreaming.h"
 #include "Playable/Transition/AvaPlayableTransition.h"
 
 namespace UE::AvaPlayableLibrary::Private
@@ -97,4 +98,21 @@ bool UAvaPlayableLibrary::UpdateRemoteControlValues(const UObject* InWorldContex
 		}
 	}
 	return false;
-}
\ No newline at end of file
+}
+
+bool UAvaPlayableLibrary::GetShouldBeHidden(const UObject* InWorldContextObject)
+{
+	const UAvaPlayableLevelStreaming* LevelStreamingPlayable = Cast<UAvaPlayableLevelStreaming>(GetPlayable(InWorldContextObject));
+	return LevelStreamingPlayable ? LevelStreamingPlayable->GetShouldBeHidden() : false;
+}
+
+bool UAvaPlayableLibrary::SetShouldBeHidden(const UObject* InWorldContextObject, bool bInShouldBeHidden)
+{
+	if (UAvaPlayableLevelStreaming* LevelStreamingPlayable = Cast<UAvaPlayableLevelStreaming>(GetPlayable(InWorldContextObject)))
+	{
+		LevelStreamingPlayable->SetShouldBeHidden(bInShouldBeHidden);
+		return true;
+	}
+	return false;
+}
+
diff --git a/Engine/Plugins/Experimental/Avalanche/Source/AvalancheMedia/Private/Playable/Playables/AvaPlayableLevelStreaming.cpp b/Engine/Plugins/Experimental/Avalanche/Source/AvalancheMedia/Private/Playable/Playables/AvaPlayableLevelStreaming.cpp
--- a/Engine/Plugins/Experimental/Avalanche/Source/AvalancheMedia/Private/Playable/Playables/AvaPlayableLevelStreaming.cpp
+++ b/Engine/Plugins/Experimental/Avalanche/Source/AvalancheMedia/Private/Playable/Playables/AvaPlayableLevelStreaming.cpp
@@ -440,7 +440,7 @@ void UAvaPlayableLevelStreaming::SetupView(FSceneViewFamily& InViewFamily, FScen
 	const ELevelStreamingState StreamingState = LevelStreaming->GetLevelStreamingState();
 	const bool bIsVisible = StreamingState == ELevelStreamingState::MakingVisible || StreamingState == ELevelStreamingState::LoadedVisible;
 	
-	if (bIsVisible && bShouldBeHidden)
+	if (bIsVisible && (bShouldBeHidden || bWaitingForShowPlayable))
 	{
 		TSet<FPrimitiveComponentId> HiddenPrimitives;	// Todo(opt): cache this?
 
@@ -501,7 +501,7 @@ void UAvaPlayableLevelStreaming::HandleTransitionEvent(UAvaPlayable* InPlayable,
 {
 	if (InPlayable == this && EnumHasAnyFlags(InTransitionFlags, EAvaPlayableTransitionEventFlags::ShowPlayable))
 	{
-		bShouldBeHidden = false;
+		bWaitingForShowPlayable = false;
 	}
 }
 
diff --git a/Engine/Plugins/Experimental/Avalanche/Source/AvalancheMedia/Private/Playable/Playables/AvaPlayableLevelStreaming.h b/Engine/Plugins/Experimental/Avalanche/Source/AvalancheMedia/Private/Playable/Playables/AvaPlayableLevelStreaming.h
--- a/Engine/Plugins/Experimental/Avalanche/Source/AvalancheMedia/Private/Playable/Playables/AvaPlayableLevelStreaming.h
+++ b/Engine/Plugins/Experimental/Avalanche/Source/AvalancheMedia/Private/Playable/Playables/AvaPlayableLevelStreaming.h
@@ -43,6 +43,9 @@ protected:
 public:
 	ULevelStreamingDynamic* GetLevelStreaming() const { return LevelStreaming; }
 
+	void SetShouldBeHidden(bool bInShouldBeHidden) { bShouldBeHidden = bInShouldBeHidden; }
+	bool GetShouldBeHidden() const { return bShouldBeHidden; }
+
 protected:
 	bool LoadLevel(const TSoftObjectPtr<UWorld>& InSourceLevel, bool bInInitiallyVisible);
 
@@ -118,5 +121,8 @@ protected:
 	bool bOnPlayQueued = false;
 	
 	/** Enter Playables should be hidden until the transition has started. */
-	bool bShouldBeHidden = true;
+	bool bWaitingForShowPlayable = true;
+
+	/** If true, all primitives from the playable will hidden. */
+	bool bShouldBeHidden = false;
 };
diff --git a/Engine/Plugins/Experimental/Avalanche/Source/AvalancheMedia/Public/Playable/AvaPlayableLibrary.h b/Engine/Plugins/Experimental/Avalanche/Source/AvalancheMedia/Public/Playable/AvaPlayableLibrary.h
--- a/Engine/Plugins/Experimental/Avalanche/Source/AvalancheMedia/Public/Playable/AvaPlayableLibrary.h
+++ b/Engine/Plugins/Experimental/Avalanche/Source/AvalancheMedia/Public/Playable/AvaPlayableLibrary.h
@@ -28,4 +28,20 @@ public:
 	 */
 	UFUNCTION(BlueprintCallable, Category = "Motion Design|Playable", meta=(WorldContext = "InWorldContextObject"))
 	static AVALANCHEMEDIA_API bool UpdateRemoteControlValues(const UObject* InWorldContextObject);
+
+	/**
+	 * Returns the current hidden state of this playable.
+	 * @remark This only works if current level is managed by a playable (i.e. in a rundown or playback graph).
+	 * @return true if the hidden state is set, false otherwise or if not managed by a playable.
+	 */
+	UFUNCTION(BlueprintCallable, Category = "Motion Design|Playable", meta=(WorldContext = "InWorldContextObject"))
+	static AVALANCHEMEDIA_API bool GetShouldBeHidden(const UObject* InWorldContextObject);
+
+	/**
+	 * Sets the hidden state of all primitives under this playable. Hidden primitives will not be rendered.
+	 * @remark This only works if current level is managed by a playable (i.e. in a rundown or playback graph).
+	 * @return true if the value was set, false otherwise (if not managed by a playable).
+	 */
+	UFUNCTION(BlueprintCallable, Category = "Motion Design|Playable", meta=(WorldContext = "InWorldContextObject"))
+	static AVALANCHEMEDIA_API bool SetShouldBeHidden(const UObject* InWorldContextObject, bool bInShouldBeHidden);
 };
