[Motion Design] Playable BP Library: Renaming functions to have "Playable" in the name for clarity.

- SetShouldBeHidden -> SetPlayableHidden
- GetShouldBeHidden -> IsPlayableHidden
- Update Remote Control Values -> Update Playable Remote Control Values

diff --git a/Engine/Plugins/Experimental/Avalanche/Source/AvalancheMedia/Private/Playable/AvaPlayableLibrary.cpp b/Engine/Plugins/Experimental/Avalanche/Source/AvalancheMedia/Private/Playable/AvaPlayableLibrary.cpp
--- a/Engine/Plugins/Experimental/Avalanche/Source/AvalancheMedia/Private/Playable/AvaPlayableLibrary.cpp
+++ b/Engine/Plugins/Experimental/Avalanche/Source/AvalancheMedia/Private/Playable/AvaPlayableLibrary.cpp
@@ -82,7 +82,7 @@ UAvaPlayableTransition* UAvaPlayableLibrary::GetPlayableTransition(const UAvaPla
 	return FoundTransition;
 }
 
-bool UAvaPlayableLibrary::UpdateRemoteControlValues(const UObject* InWorldContextObject)
+bool UAvaPlayableLibrary::UpdatePlayableRemoteControlValues(const UObject* InWorldContextObject)
 {
 	if (UAvaPlayable* Playable = GetPlayable(InWorldContextObject))
 	{
@@ -100,13 +100,13 @@ bool UAvaPlayableLibrary::UpdateRemoteControlValues(const UObject* InWorldContex
 	return false;
 }
 
-bool UAvaPlayableLibrary::GetShouldBeHidden(const UObject* InWorldContextObject)
+bool UAvaPlayableLibrary::IsPlayableHidden(const UObject* InWorldContextObject)
 {
 	const UAvaPlayableLevelStreaming* LevelStreamingPlayable = Cast<UAvaPlayableLevelStreaming>(GetPlayable(InWorldContextObject));
 	return LevelStreamingPlayable ? LevelStreamingPlayable->GetShouldBeHidden() : false;
 }
 
-bool UAvaPlayableLibrary::SetShouldBeHidden(const UObject* InWorldContextObject, bool bInShouldBeHidden)
+bool UAvaPlayableLibrary::SetPlayableHidden(const UObject* InWorldContextObject, bool bInShouldBeHidden)
 {
 	if (UAvaPlayableLevelStreaming* LevelStreamingPlayable = Cast<UAvaPlayableLevelStreaming>(GetPlayable(InWorldContextObject)))
 	{
diff --git a/Engine/Plugins/Experimental/Avalanche/Source/AvalancheMedia/Public/Playable/AvaPlayableLibrary.h b/Engine/Plugins/Experimental/Avalanche/Source/AvalancheMedia/Public/Playable/AvaPlayableLibrary.h
--- a/Engine/Plugins/Experimental/Avalanche/Source/AvalancheMedia/Public/Playable/AvaPlayableLibrary.h
+++ b/Engine/Plugins/Experimental/Avalanche/Source/AvalancheMedia/Public/Playable/AvaPlayableLibrary.h
@@ -27,7 +27,7 @@ public:
 	 * @return true if the values have been injected, false otherwise.
 	 */
 	UFUNCTION(BlueprintCallable, Category = "Motion Design|Playable", meta=(WorldContext = "InWorldContextObject"))
-	static AVALANCHEMEDIA_API bool UpdateRemoteControlValues(const UObject* InWorldContextObject);
+	static AVALANCHEMEDIA_API bool UpdatePlayableRemoteControlValues(const UObject* InWorldContextObject);
 
 	/**
 	 * Returns the current hidden state of this playable.
@@ -35,7 +35,7 @@ public:
 	 * @return true if the hidden state is set, false otherwise or if not managed by a playable.
 	 */
 	UFUNCTION(BlueprintCallable, Category = "Motion Design|Playable", meta=(WorldContext = "InWorldContextObject"))
-	static AVALANCHEMEDIA_API bool GetShouldBeHidden(const UObject* InWorldContextObject);
+	static AVALANCHEMEDIA_API bool IsPlayableHidden(const UObject* InWorldContextObject);
 
 	/**
 	 * Sets the hidden state of all primitives under this playable. Hidden primitives will not be rendered.
@@ -43,5 +43,5 @@ public:
 	 * @return true if the value was set, false otherwise (if not managed by a playable).
 	 */
 	UFUNCTION(BlueprintCallable, Category = "Motion Design|Playable", meta=(WorldContext = "InWorldContextObject"))
-	static AVALANCHEMEDIA_API bool SetShouldBeHidden(const UObject* InWorldContextObject, bool bInShouldBeHidden);
+	static AVALANCHEMEDIA_API bool SetPlayableHidden(const UObject* InWorldContextObject, bool bInShouldBeHidden);
 };
