[Motion Design] Fix Rundown Remove Page broadcast page ids so it sends the correct list of page ids for each page lists.

diff --git a/Engine/Plugins/Experimental/Avalanche/Source/AvalancheMedia/Private/Rundown/AvaRundown.cpp b/Engine/Plugins/Experimental/Avalanche/Source/AvalancheMedia/Private/Rundown/AvaRundown.cpp
--- a/Engine/Plugins/Experimental/Avalanche/Source/AvalancheMedia/Private/Rundown/AvaRundown.cpp
+++ b/Engine/Plugins/Experimental/Avalanche/Source/AvalancheMedia/Private/Rundown/AvaRundown.cpp
@@ -418,21 +418,13 @@ int32 UAvaRundown::RemovePages(const TArray<int32>& InPageIds)
 	{
 		return 0;
 	}
-	
-	TArray<int32> SortedPageIds(InPageIds);
-	SortedPageIds.Sort();
-
-	int32 RemovedCount = 0;
-	bool bRemovedTemplate = false;
-	bool bRemovedInstanced = false;
 
 	// Find the instanced page ids to remove
 	TArray<int32> TemplatesIndicesToRemove;
 	TSet<int32> InstancesToRemove;
 
-	for (int32 PageIdx = SortedPageIds.Num() - 1; PageIdx >= 0; --PageIdx)
+	for (int32 PageId : InPageIds)
 	{
-		const int32 PageId = SortedPageIds[PageIdx];
 		const int32* TemplateIdx = TemplatePages.PageIndices.Find(PageId);
 		const int32* InstanceIdx = InstancedPages.PageIndices.Find(PageId);
 
@@ -468,15 +460,17 @@ int32 UAvaRundown::RemovePages(const TArray<int32>& InPageIds)
 		}
 	}
 
+	TArray<int32> RemovedTemplateIds;
+
 	if (TemplatesIndicesToRemove.IsEmpty() == false)
 	{
 		TemplatesIndicesToRemove.Sort();
+		RemovedTemplateIds.Reserve(TemplatesIndicesToRemove.Num());
 
 		for (int32 TemplateIdx = TemplatesIndicesToRemove.Num() - 1; TemplateIdx >= 0; --TemplateIdx)
 		{
+			RemovedTemplateIds.Add(TemplatePages.Pages[TemplatesIndicesToRemove[TemplateIdx]].PageId);
 			TemplatePages.Pages.RemoveAt(TemplatesIndicesToRemove[TemplateIdx]);
-			++RemovedCount;
-			bRemovedTemplate = true;
 		}
 	}
 
@@ -501,48 +495,47 @@ int32 UAvaRundown::RemovePages(const TArray<int32>& InPageIds)
 		{
 			const int32 InstanceToRemoveIdx = InstancesToRemoveIndices[RemoveIdx];
 			InstancedPages.Pages.RemoveAt(InstanceToRemoveIdx);
-			++RemovedCount;
-			bRemovedInstanced = true;
 		}
 	}
 		
-	if (RemovedCount == 0)
+	if (RemovedTemplateIds.IsEmpty() && InstancesToRemove.IsEmpty())
 	{
 		return 0;
 	}
 
 	RefreshPageIndices();
 
-	if (bRemovedTemplate)
+	if (!RemovedTemplateIds.IsEmpty())
 	{
-		GetOnPageListChanged().Broadcast({this, TemplatePageList, EAvaRundownPageListChange::RemovedPages, {InPageIds}});
+		GetOnPageListChanged().Broadcast({this, TemplatePageList, EAvaRundownPageListChange::RemovedPages, RemovedTemplateIds});
 	}
 
-	if (bRemovedInstanced)
+	if (!InstancesToRemove.IsEmpty())
 	{
-		GetOnPageListChanged().Broadcast({this, InstancePageList, EAvaRundownPageListChange::RemovedPages, {InPageIds}});
+		GetOnPageListChanged().Broadcast({this, InstancePageList, EAvaRundownPageListChange::RemovedPages, InstancesToRemove.Array()});
 	}
 
 	for (FAvaRundownSubList& SubList : SubLists)
 	{
-		bool bFoundInstance = false;
+		TArray<int32> RemovedInstanceIds;
+		RemovedInstanceIds.Reserve(InstancesToRemove.Num());
 
 		for (TArray<int32>::TIterator Iter(SubList.PageIds); Iter; ++Iter)
 		{
 			if (InstancesToRemove.Contains(*Iter))
 			{
+				RemovedInstanceIds.Add(*Iter);
 				Iter.RemoveCurrent();
-				bFoundInstance = true;
 			}
 		}
 
-		if (bFoundInstance)
+		if (!RemovedInstanceIds.IsEmpty())
 		{
-			GetOnPageListChanged().Broadcast({this, CreateSubListReference(SubList), EAvaRundownPageListChange::RemovedPages, {InPageIds}});
+			GetOnPageListChanged().Broadcast({this, CreateSubListReference(SubList), EAvaRundownPageListChange::RemovedPages, RemovedInstanceIds});
 		}
 	}
 
-	return RemovedCount;
+	return RemovedTemplateIds.Num() + InstancesToRemove.Num();	// Total Pages removed.
 }
 
 bool UAvaRundown::CanRemovePages(const TArray<int32>& InPageIds) const
