Motion Design
- "Add Texture to Scene" now scales the texture to fit the screen, maintaining aspect ratio.

diff --git a/Engine/Plugins/Experimental/Avalanche/Source/AvalancheEditor/Private/MaterialDesigner/AvaLevelMaterialDesignerExtension.cpp b/Engine/Plugins/Experimental/Avalanche/Source/AvalancheEditor/Private/MaterialDesigner/AvaLevelMaterialDesignerExtension.cpp
--- a/Engine/Plugins/Experimental/Avalanche/Source/AvalancheEditor/Private/MaterialDesigner/AvaLevelMaterialDesignerExtension.cpp
+++ b/Engine/Plugins/Experimental/Avalanche/Source/AvalancheEditor/Private/MaterialDesigner/AvaLevelMaterialDesignerExtension.cpp
@@ -3,10 +3,11 @@
 #include "MaterialDesigner/AvaLevelMaterialDesignerExtension.h"
 #include "AvaEditorModule.h"
 #include "AvaMaterialDesignerTextureAssetFactory.h"
+#include "AvaScreenAlignmentUtils.h"
 #include "ContentBrowserModule.h"
 #include "Delegates/IDelegateInstance.h"
-#include "Materials/Material.h"
 #include "Engine/Texture.h"
+#include "Materials/Material.h"
 #include "Styling/SlateIconFinder.h"
 #include "Viewport/AvaLevelViewportExtension.h"
 #include "ViewportClient/IAvaViewportClient.h"
@@ -134,7 +135,19 @@ void FAvaLevelMaterialDesignerExtension::AddTextureToSene(FAssetData InAssetData
 	FActorSpawnParameters SpawnParameters;
 	SpawnParameters.bNoFail = true;
 
-	AssetFactory->CreateActor(InAssetData.GetAsset(), World->PersistentLevel.Get(), FTransform::Identity, SpawnParameters);
+	AActor* TextureActor = AssetFactory->CreateActor(InAssetData.GetAsset(), World->PersistentLevel.Get(), FTransform::Identity, SpawnParameters);
+
+	if (!TextureActor)
+	{
+		return;
+	}
+
+	FAvaScreenAlignmentUtils::FitActorToScreen(
+		LevelViewportClients[0].ToSharedRef(), 
+		*TextureActor,
+		/* Stretch to fit */ false,
+		/* Align to nearest axis */ true
+	);
 }
 
 #undef LOCTEXT_NAMESPACE
