Motion Design
- Further changed "Add Texture to Scene" to only scale the texture down if it exceeds the bounds of the viewport.

diff --git a/Engine/Plugins/Experimental/Avalanche/Source/AvalancheEditor/Private/MaterialDesigner/AvaLevelMaterialDesignerExtension.cpp b/Engine/Plugins/Experimental/Avalanche/Source/AvalancheEditor/Private/MaterialDesigner/AvaLevelMaterialDesignerExtension.cpp
--- a/Engine/Plugins/Experimental/Avalanche/Source/AvalancheEditor/Private/MaterialDesigner/AvaLevelMaterialDesignerExtension.cpp
+++ b/Engine/Plugins/Experimental/Avalanche/Source/AvalancheEditor/Private/MaterialDesigner/AvaLevelMaterialDesignerExtension.cpp
@@ -108,12 +108,14 @@ void FAvaLevelMaterialDesignerExtension::AddTextureToSene(FAssetData InAssetData
 {
 	const TArray<TSharedPtr<IAvaViewportClient>> LevelViewportClients = FAvaLevelViewportExtension::GetLevelEditorViewportClients();
 
-	if (LevelViewportClients.IsEmpty())
+	if (LevelViewportClients.IsEmpty() || !LevelViewportClients[0].IsValid())
 	{
 		return;
 	}
 
-	UWorld* World = LevelViewportClients[0]->GetViewportWorld();
+	TSharedRef<IAvaViewportClient> LevelViewportClient = LevelViewportClients[0].ToSharedRef();
+
+	UWorld* World = LevelViewportClient->GetViewportWorld();
 
 	if (!World)
 	{
@@ -122,7 +124,7 @@ void FAvaLevelMaterialDesignerExtension::AddTextureToSene(FAssetData InAssetData
 
 	UAvaMaterialDesignerTextureAssetFactory* AssetFactory = NewObject<UAvaMaterialDesignerTextureAssetFactory>(GetTransientPackage());
 	check(AssetFactory);
-	AssetFactory->SetCameraRotation(LevelViewportClients[0]->GetViewportViewTransform().Rotator());
+	AssetFactory->SetCameraRotation(LevelViewportClient->GetViewportViewTransform().Rotator());
 
 	FText ErrorMsg;
 
@@ -132,22 +134,36 @@ void FAvaLevelMaterialDesignerExtension::AddTextureToSene(FAssetData InAssetData
 		return;
 	}
 
+	UTexture* Texture = Cast<UTexture>(InAssetData.GetAsset());
+
+	if (!Texture)
+	{
+		return;
+	}
+
 	FActorSpawnParameters SpawnParameters;
 	SpawnParameters.bNoFail = true;
 
-	AActor* TextureActor = AssetFactory->CreateActor(InAssetData.GetAsset(), World->PersistentLevel.Get(), FTransform::Identity, SpawnParameters);
+	AActor* TextureActor = AssetFactory->CreateActor(Texture, World->PersistentLevel.Get(), FTransform::Identity, SpawnParameters);
 
 	if (!TextureActor)
 	{
 		return;
 	}
 
-	FAvaScreenAlignmentUtils::FitActorToScreen(
-		LevelViewportClients[0].ToSharedRef(), 
-		*TextureActor,
-		/* Stretch to fit */ false,
-		/* Align to nearest axis */ true
+	const FVector2D ViewportSize = LevelViewportClient->GetFrustumSizeAtDistance(
+		(LevelViewportClient->GetViewportViewTransform().GetLocation() - TextureActor->GetActorLocation()).Size()
 	);
+
+	if (Texture->GetSurfaceWidth() > ViewportSize.X || Texture->GetSurfaceHeight() > ViewportSize.Y)
+	{
+		FAvaScreenAlignmentUtils::FitActorToScreen(
+			LevelViewportClient,
+			*TextureActor,
+			/* Stretch to fit */ false,
+			/* Align to nearest axis */ true
+		);
+	}
 }
 
 #undef LOCTEXT_NAMESPACE
diff --git a/Engine/Plugins/Experimental/Avalanche/Source/AvalancheEditor/Private/MaterialDesigner/AvaMaterialDesignerTextureAssetFactory.cpp b/Engine/Plugins/Experimental/Avalanche/Source/AvalancheEditor/Private/MaterialDesigner/AvaMaterialDesignerTextureAssetFactory.cpp
--- a/Engine/Plugins/Experimental/Avalanche/Source/AvalancheEditor/Private/MaterialDesigner/AvaMaterialDesignerTextureAssetFactory.cpp
+++ b/Engine/Plugins/Experimental/Avalanche/Source/AvalancheEditor/Private/MaterialDesigner/AvaMaterialDesignerTextureAssetFactory.cpp
@@ -58,8 +58,6 @@ AActor* UAvaMaterialDesignerTextureAssetFactory::SpawnActor(UObject* InAsset, UL
 		return InNewActor;
 	}
 
-	InNewActor->Rename(*Texture->GetName(), nullptr, REN_DontCreateRedirectors);
-
 	AAvaShapeActor* ShapeActor = Cast<AAvaShapeActor>(InNewActor);
 
 	if (!ShapeActor)
@@ -200,6 +198,11 @@ AActor* UAvaMaterialDesignerTextureAssetFactory::SpawnActor(UObject* InAsset, UL
 
 FString UAvaMaterialDesignerTextureAssetFactory::GetDefaultActorLabel(UObject* InAsset) const
 {
+	if (InAsset)
+	{
+		return InAsset->GetName();
+	}
+
 	static const FString DefaultName = TEXT("Rectangle");
 	return DefaultName;
 }
