MotionDesign : PropertyAnimator
- Improving sequence time source to display proper label and id and working exclusively with avalanche sequences

diff --git a/Engine/Plugins/Experimental/Avalanche/Source/AvalanchePropertyAnimator/AvalanchePropertyAnimator.Build.cs b/Engine/Plugins/Experimental/Avalanche/Source/AvalanchePropertyAnimator/AvalanchePropertyAnimator.Build.cs
--- a/Engine/Plugins/Experimental/Avalanche/Source/AvalanchePropertyAnimator/AvalanchePropertyAnimator.Build.cs
+++ b/Engine/Plugins/Experimental/Avalanche/Source/AvalanchePropertyAnimator/AvalanchePropertyAnimator.Build.cs
@@ -11,6 +11,7 @@ public class AvalanchePropertyAnimator : ModuleRules
         PrivateDependencyModuleNames.AddRange(
             new string[]
             {
+	            "Avalanche",
 	            "AvalancheSequence",
 	            "AvalancheText",
 	            "Core",
diff --git a/Engine/Plugins/Experimental/Avalanche/Source/AvalanchePropertyAnimator/Private/TimeSources/AvaPropertyAnimatorSequenceTimeSource.cpp b/Engine/Plugins/Experimental/Avalanche/Source/AvalanchePropertyAnimator/Private/TimeSources/AvaPropertyAnimatorSequenceTimeSource.cpp
--- a/Engine/Plugins/Experimental/Avalanche/Source/AvalanchePropertyAnimator/Private/TimeSources/AvaPropertyAnimatorSequenceTimeSource.cpp
+++ b/Engine/Plugins/Experimental/Avalanche/Source/AvalanchePropertyAnimator/Private/TimeSources/AvaPropertyAnimatorSequenceTimeSource.cpp
@@ -2,8 +2,12 @@
 
 #include "TimeSources/AvaPropertyAnimatorSequenceTimeSource.h"
 
+#include "Animators/PropertyAnimatorCoreBase.h"
+#include "AvaSceneSubsystem.h"
 #include "AvaSequence.h"
 #include "AvaSequencePlayer.h"
+#include "IAvaSceneInterface.h"
+#include "IAvaSequenceProvider.h"
 #include "MovieSceneSequence.h"
 #include "UObject/UObjectIterator.h"
 
@@ -12,6 +16,16 @@
 #include "Modules/ModuleManager.h"
 #endif // WITH_EDITOR
 
+FName UAvaPropertyAnimatorSequenceTimeSource::GetSequenceName(const UAvaSequence* InSequence)
+{
+	if (!InSequence)
+	{
+		return NAME_None;
+	}
+
+	return FName(InSequence->GetLabel().ToString() + TEXT(" (") + InSequence->GetName() + TEXT(")"));
+}
+
 #if WITH_EDITOR
 void UAvaPropertyAnimatorSequenceTimeSource::PostEditChangeProperty(FPropertyChangedEvent& InPropertyChangedEvent)
 {
@@ -28,14 +42,14 @@ void UAvaPropertyAnimatorSequenceTimeSource::PostEditChangeProperty(FPropertyCha
 }
 #endif // WITH_EDITOR
 
-void UAvaPropertyAnimatorSequenceTimeSource::SetSequenceName(const FString& InSequenceName)
+void UAvaPropertyAnimatorSequenceTimeSource::SetSequenceName(FName InSequenceName)
 {
 	if (SequenceName == InSequenceName)
 	{
 		return;
 	}
 
-	const TArray<FString> Sequences = GetSequenceNames();
+	const TArray<FName> Sequences = GetSequenceNames();
 	if (!Sequences.Contains(InSequenceName))
 	{
 		return;
@@ -93,9 +107,9 @@ double UAvaPropertyAnimatorSequenceTimeSource::GetTimeElapsed()
 	if (const TSharedPtr<ISequencer> Sequencer = GetSequencer())
 	{
 		// Scrub to global time only if root sequence is the selected sequence
-		if (const UMovieSceneSequence* RootSequence = Sequencer->GetRootMovieSceneSequence())
+		if (const UAvaSequence* RootSequence = Cast<UAvaSequence>(Sequencer->GetRootMovieSceneSequence()))
 		{
-			if (RootSequence->GetName() == SequenceName)
+			if (GetSequenceName(RootSequence) == SequenceName)
 			{
 				return Sequencer->GetGlobalTime().AsSeconds();
 			}
@@ -136,27 +150,73 @@ void UAvaPropertyAnimatorSequenceTimeSource::OnTimeSourceUnregistered()
 #endif // WITH_EDITOR
 }
 
-TArray<FString> UAvaPropertyAnimatorSequenceTimeSource::GetSequenceNames() const
+TArray<FName> UAvaPropertyAnimatorSequenceTimeSource::GetSequenceNames() const
 {
-	TArray<FString> SequenceNames {TEXT("")};
-
-	const UWorld* World = GetWorld();
+	TArray<FName> SequenceNames {NAME_None};
 
-	for (const UMovieSceneSequence* MovieSceneSequence : TObjectRange<UMovieSceneSequence>())
+	for (const UAvaSequence* Sequence : GetSequences())
 	{
-		if (MovieSceneSequence && MovieSceneSequence->GetWorld() == World)
+		const FName MovieSceneSequenceName = GetSequenceName(Sequence);
+
+		if (!MovieSceneSequenceName.IsNone())
 		{
-			SequenceNames.Add(MovieSceneSequence->GetName());
+			SequenceNames.Add(MovieSceneSequenceName);
 		}
 	}
 
 	return SequenceNames;
 }
 
+TArray<UAvaSequence*> UAvaPropertyAnimatorSequenceTimeSource::GetSequences() const
+{
+	TArray<UAvaSequence*> Sequences;
+
+	const UPropertyAnimatorCoreBase* Animator = GetAnimator();
+	if (!Animator)
+	{
+		return Sequences;
+	}
+
+	const AActor* AnimatorActor = Animator->GetAnimatorActor();
+	if (!AnimatorActor)
+	{
+		return Sequences;
+	}
+
+	ULevel* AnimatorLevel = AnimatorActor->GetLevel();
+	if (!AnimatorLevel)
+	{
+		return Sequences;
+	}
+
+	const IAvaSceneInterface* SceneInterface = UAvaSceneSubsystem::FindSceneInterface(AnimatorLevel);
+	if (!SceneInterface)
+	{
+		return Sequences;
+	}
+
+	const IAvaSequenceProvider* SequenceProvider = SceneInterface->GetSequenceProvider();
+	if (!SequenceProvider)
+	{
+		return Sequences;
+	}
+
+	// Find sequence with that name
+	for (const TWeakObjectPtr<UAvaSequence>& RootSequenceWeak : SequenceProvider->GetRootSequences())
+	{
+		if (UAvaSequence* RootSequence = RootSequenceWeak.Get())
+		{
+			Sequences.Add(RootSequence);
+		}
+	}
+
+	return Sequences;
+}
+
 void UAvaPropertyAnimatorSequenceTimeSource::OnSequenceChanged()
 {
 	// Reset sequence selection
-	if (SequenceName.IsEmpty())
+	if (SequenceName.IsNone())
 	{
 		SequenceWeak = nullptr;
 		SequenceWeak.Reset();
@@ -172,17 +232,12 @@ void UAvaPropertyAnimatorSequenceTimeSource::OnSequenceChanged()
 	}
 
 	// Find sequence with that name
-	const UWorld* World = GetWorld();
-
-	for (UMovieSceneSequence* MovieSceneSequence : TObjectRange<UMovieSceneSequence>())
+	for (UAvaSequence* Sequence : GetSequences())
 	{
-		if (MovieSceneSequence && MovieSceneSequence->GetWorld() == World)
+		if (Sequence && GetSequenceName(Sequence) == SequenceName)
 		{
-			if (MovieSceneSequence->GetName() == SequenceName)
-			{
-				SequenceWeak = MovieSceneSequence;
-				break;
-			}
+			SequenceWeak = Sequence;
+			break;
 		}
 	}
 
@@ -203,7 +258,7 @@ void UAvaPropertyAnimatorSequenceTimeSource::OnSequenceStarted(UAvaSequencePlaye
 		return;
 	}
 
-	const UMovieSceneSequence* Sequence = SequenceWeak.Get();
+	const UAvaSequence* Sequence = SequenceWeak.Get();
 	if (InSequence != Sequence)
 	{
 		return;
@@ -222,7 +277,7 @@ void UAvaPropertyAnimatorSequenceTimeSource::OnSequenceFinished(UAvaSequencePlay
 		return;
 	}
 
-	const UMovieSceneSequence* Sequence = SequenceWeak.Get();
+	const UAvaSequence* Sequence = SequenceWeak.Get();
 	if (Sequence != InSequence)
 	{
 		return;
diff --git a/Engine/Plugins/Experimental/Avalanche/Source/AvalanchePropertyAnimator/Private/TimeSources/AvaPropertyAnimatorSequenceTimeSource.h b/Engine/Plugins/Experimental/Avalanche/Source/AvalanchePropertyAnimator/Private/TimeSources/AvaPropertyAnimatorSequenceTimeSource.h
--- a/Engine/Plugins/Experimental/Avalanche/Source/AvalanchePropertyAnimator/Private/TimeSources/AvaPropertyAnimatorSequenceTimeSource.h
+++ b/Engine/Plugins/Experimental/Avalanche/Source/AvalanchePropertyAnimator/Private/TimeSources/AvaPropertyAnimatorSequenceTimeSource.h
@@ -2,25 +2,26 @@
 
 #pragma once
 
-#include "TimeSources/PropertyAnimatorCoreRangeTimeSource.h"
+#include "TimeSources/PropertyAnimatorCoreTimeSourceBase.h"
 #include "AvaPropertyAnimatorSequenceTimeSource.generated.h"
 
 class ISequencer;
 class UAvaSequence;
 class UAvaSequencePlayer;
-class UMovieSceneSequence;
 
 /**
- * Time source that follows specific sequence time
+ * Time source that follows specific avalanche sequence time
  */
 UCLASS()
-class UAvaPropertyAnimatorSequenceTimeSource : public UPropertyAnimatorCoreRangeTimeSource
+class UAvaPropertyAnimatorSequenceTimeSource : public UPropertyAnimatorCoreTimeSourceBase
 {
 	GENERATED_BODY()
 
 public:
+	static FName GetSequenceName(const UAvaSequence* InSequence);
+
 	UAvaPropertyAnimatorSequenceTimeSource()
-		: UPropertyAnimatorCoreRangeTimeSource(TEXT("Sequence"))
+		: UPropertyAnimatorCoreTimeSourceBase(TEXT("MD_Sequence"))
 	{}
 
 	//~ Begin UObject
@@ -29,8 +30,8 @@ public:
 #endif // WITH_EDITOR
 	//~ End UObject
 
-	void SetSequenceName(const FString& InSequenceName);
-	FString GetSequenceName() const
+	void SetSequenceName(FName InSequenceName);
+	FName GetSequenceName() const
 	{
 		return SequenceName;
 	}
@@ -49,7 +50,9 @@ protected:
 	//~ End UPropertyAnimatorCoreTimeSourceBase
 
 	UFUNCTION()
-	TArray<FString> GetSequenceNames() const;
+	TArray<FName> GetSequenceNames() const;
+
+	TArray<UAvaSequence*> GetSequences() const;
 
 	void OnSequenceChanged();
 
@@ -57,10 +60,10 @@ protected:
 	void OnSequenceFinished(UAvaSequencePlayer* InPlayer, UAvaSequence* InSequence);
 
 	UPROPERTY(EditInstanceOnly, BlueprintReadWrite, Setter="SetSequenceName", Getter="GetSequenceName", Category="Animator", meta=(GetOptions="GetSequenceNames"))
-	FString SequenceName;
+	FName SequenceName;
 
 	UPROPERTY()
-	TWeakObjectPtr<UMovieSceneSequence> SequenceWeak;
+	TWeakObjectPtr<UAvaSequence> SequenceWeak;
 
 	UPROPERTY(Transient)
 	TWeakObjectPtr<UAvaSequencePlayer> SequencePlayerWeak;
diff --git a/Engine/Plugins/Experimental/PropertyAnimatorCore/Source/PropertyAnimatorCore/Public/TimeSources/PropertyAnimatorCoreTimeSourceBase.h b/Engine/Plugins/Experimental/PropertyAnimatorCore/Source/PropertyAnimatorCore/Public/TimeSources/PropertyAnimatorCoreTimeSourceBase.h
--- a/Engine/Plugins/Experimental/PropertyAnimatorCore/Source/PropertyAnimatorCore/Public/TimeSources/PropertyAnimatorCoreTimeSourceBase.h
+++ b/Engine/Plugins/Experimental/PropertyAnimatorCore/Source/PropertyAnimatorCore/Public/TimeSources/PropertyAnimatorCoreTimeSourceBase.h
@@ -41,7 +41,7 @@ public:
 	TOptional<double> GetConditionalTimeElapsed();
 
 	/** Get the animator, this time source is on */
-	UPropertyAnimatorCoreBase* GetAnimator() const;
+	PROPERTYANIMATORCORE_API UPropertyAnimatorCoreBase* GetAnimator() const;
 
 	FName GetTimeSourceName() const
 	{
