Motion Design Editor: Changed extension/editor to get scene object from current selected level rather than always from persistent

diff --git a/Engine/Plugins/Experimental/Avalanche/Source/AvalancheEditor/Private/AvaEditorSettings.h b/Engine/Plugins/Experimental/Avalanche/Source/AvalancheEditor/Private/AvaEditorSettings.h
--- a/Engine/Plugins/Experimental/Avalanche/Source/AvalancheEditor/Private/AvaEditorSettings.h
+++ b/Engine/Plugins/Experimental/Avalanche/Source/AvalancheEditor/Private/AvaEditorSettings.h
@@ -27,6 +27,10 @@ public:
 
 	static UAvaEditorSettings* Get();
 
+	/** Whether to allow the Motion Design Interface to show the current selected level rather than fixed at the persistent level */
+	UPROPERTY(Config, EditAnywhere, BlueprintReadWrite, Category = "Behavior")
+	bool bEnableLevelContextSwitching = true;
+
 	/** Whether to Automatically Include the Attached Actors when performing Edit Actions such as Cut, Copy, Duplicate. */
 	UPROPERTY(Config, EditAnywhere, BlueprintReadWrite, Category = "Behavior")
 	bool bAutoIncludeAttachedActorsInEditActions = true;
diff --git a/Engine/Plugins/Experimental/Avalanche/Source/AvalancheEditor/Private/EditorProvider/AvaEditorProvider.cpp b/Engine/Plugins/Experimental/Avalanche/Source/AvalancheEditor/Private/EditorProvider/AvaEditorProvider.cpp
--- a/Engine/Plugins/Experimental/Avalanche/Source/AvalancheEditor/Private/EditorProvider/AvaEditorProvider.cpp
+++ b/Engine/Plugins/Experimental/Avalanche/Source/AvalancheEditor/Private/EditorProvider/AvaEditorProvider.cpp
@@ -2,18 +2,28 @@
 
 #include "AvaEditorProvider.h"
 #include "AvaEditorActorUtils.h"
+#include "AvaEditorSettings.h"
 #include "AvaScene.h"
 #include "Engine/World.h"
 #include "IAvaEditor.h"
 
 UObject* FAvaEditorProvider::GetSceneObject(UWorld* InWorld, EAvaEditorObjectQueryType InQueryType)
 {
-	if (InWorld)
+	if (!InWorld)
 	{
-		const bool bCreateSceneIfNotFound = InQueryType == EAvaEditorObjectQueryType::CreateIfNotFound;
-		return AAvaScene::GetScene(InWorld->PersistentLevel, bCreateSceneIfNotFound);
+		return nullptr;
 	}
-	return nullptr;
+
+	ULevel* SceneLevel = InWorld->PersistentLevel;
+
+	const UAvaEditorSettings* EditorSettings = UAvaEditorSettings::Get();
+	if (EditorSettings && EditorSettings->bEnableLevelContextSwitching)
+	{
+		SceneLevel = InWorld->GetCurrentLevel();
+	}
+
+	const bool bCreateSceneIfNotFound = InQueryType == EAvaEditorObjectQueryType::CreateIfNotFound;
+	return AAvaScene::GetScene(SceneLevel, bCreateSceneIfNotFound);
 }
 
 void FAvaEditorProvider::GetActorsToEdit(TArray<AActor*>& InOutActorsToEdit) const
diff --git a/Engine/Plugins/Experimental/Avalanche/Source/AvalancheEditor/Private/RemoteControl/AvaRCExtension.cpp b/Engine/Plugins/Experimental/Avalanche/Source/AvalancheEditor/Private/RemoteControl/AvaRCExtension.cpp
--- a/Engine/Plugins/Experimental/Avalanche/Source/AvalancheEditor/Private/RemoteControl/AvaRCExtension.cpp
+++ b/Engine/Plugins/Experimental/Avalanche/Source/AvalancheEditor/Private/RemoteControl/AvaRCExtension.cpp
@@ -55,6 +55,46 @@ void FAvaRCExtension::ExtendToolbarMenu(UToolMenu& InMenu)
 	Entry.StyleNameOverride = "CalloutToolbar";
 }
 
+void FAvaRCExtension::OnSceneObjectChanged(UObject* InOldSceneObject, UObject* InNewSceneObject)
+{
+	TSharedPtr<IAvaEditor> Editor = GetEditor();
+	if (!Editor.IsValid())
+	{
+		return;
+	}
+
+	TSharedPtr<IToolkitHost> ToolkitHost = Editor->GetToolkitHost();
+	if (!ToolkitHost.IsValid())
+	{
+		return;
+	}
+
+	UAssetEditorSubsystem* AssetEditorSubsystem = GEditor->GetEditorSubsystem<UAssetEditorSubsystem>();
+	check(AssetEditorSubsystem);
+
+	if (const IAvaSceneInterface* OldScene = Cast<IAvaSceneInterface>(InOldSceneObject))
+	{
+		if (URemoteControlPreset* OldPreset = OldScene->GetRemoteControlPreset())
+		{
+			AssetEditorSubsystem->CloseAllEditorsForAsset(OldPreset);
+		}
+	}
+
+	if (const IAvaSceneInterface* NewScene = Cast<IAvaSceneInterface>(InNewSceneObject))
+	{
+		if (URemoteControlPreset* NewPreset = NewScene->GetRemoteControlPreset())
+		{
+			// Level-dependent bindings are resolved via the Selected World in the Preset. By default, this should be the current edited world.
+			// However, the Current Edited World does not hold ownership of the other sublevels, and so doing a "FindObject" or similar with this world
+			// will fail for any actor/subobject within these sublevels.
+			// So set the Selected World to be the true outer world of the new scene object level.
+			NewPreset->SelectedWorld = InNewSceneObject->GetTypedOuter<UWorld>();
+
+			AssetEditorSubsystem->OpenEditorForAsset(NewPreset, EToolkitMode::WorldCentric, ToolkitHost);
+		}
+	}
+}
+
 void FAvaRCExtension::OpenRemoteControlTab() const
 {
 	TSharedPtr<IAvaEditor> Editor = GetEditor();
diff --git a/Engine/Plugins/Experimental/Avalanche/Source/AvalancheEditor/Private/RemoteControl/AvaRCExtension.h b/Engine/Plugins/Experimental/Avalanche/Source/AvalancheEditor/Private/RemoteControl/AvaRCExtension.h
--- a/Engine/Plugins/Experimental/Avalanche/Source/AvalancheEditor/Private/RemoteControl/AvaRCExtension.h
+++ b/Engine/Plugins/Experimental/Avalanche/Source/AvalancheEditor/Private/RemoteControl/AvaRCExtension.h
@@ -22,6 +22,7 @@ public:
 	virtual void Deactivate() override;
 	virtual void ExtendLevelEditorLayout(FLayoutExtender& InExtender) const override;
 	virtual void ExtendToolbarMenu(UToolMenu& InMenu) override;
+	virtual void OnSceneObjectChanged(UObject* InOldSceneObject, UObject* InNewSceneObject) override;
 	//~ End IAvaEditorExtension
 
 protected:
diff --git a/Engine/Plugins/Experimental/Avalanche/Source/AvalancheEditor/Private/Sequencer/AvaSequencerExtension.cpp b/Engine/Plugins/Experimental/Avalanche/Source/AvalancheEditor/Private/Sequencer/AvaSequencerExtension.cpp
--- a/Engine/Plugins/Experimental/Avalanche/Source/AvalancheEditor/Private/Sequencer/AvaSequencerExtension.cpp
+++ b/Engine/Plugins/Experimental/Avalanche/Source/AvalancheEditor/Private/Sequencer/AvaSequencerExtension.cpp
@@ -115,6 +115,23 @@ void FAvaSequencerExtension::Cleanup()
 	}
 }
 
+void FAvaSequencerExtension::OnSceneObjectChanged(UObject* InOldSceneObject, UObject* InNewSceneObject)
+{
+	if (IAvaSequenceProvider* SequenceProvider = Cast<IAvaSequenceProvider>(InNewSceneObject))
+	{
+		SequenceProvider->RebuildSequenceTree();
+
+		const TArray<TObjectPtr<UAvaSequence>>& Sequences = SequenceProvider->GetSequences();
+
+		// Select First Sequence from the new List
+		if (!Sequences.IsEmpty() && AvaSequencer.IsValid())
+		{
+			AvaSequencer->SetViewedSequence(Sequences[0]);
+			AvaSequencer->NotifyOnSequenceTreeChanged();
+		}
+	}
+}
+
 void FAvaSequencerExtension::RegisterTabSpawners(const TSharedRef<IAvaEditor>& InEditor) const
 {
 	InEditor->AddTabSpawner<FAvaSequencerTabSpawner>(InEditor, GetSequencerTabId());
diff --git a/Engine/Plugins/Experimental/Avalanche/Source/AvalancheEditor/Private/Sequencer/AvaSequencerExtension.h b/Engine/Plugins/Experimental/Avalanche/Source/AvalancheEditor/Private/Sequencer/AvaSequencerExtension.h
--- a/Engine/Plugins/Experimental/Avalanche/Source/AvalancheEditor/Private/Sequencer/AvaSequencerExtension.h
+++ b/Engine/Plugins/Experimental/Avalanche/Source/AvalancheEditor/Private/Sequencer/AvaSequencerExtension.h
@@ -38,6 +38,7 @@ public:
 	virtual void Activate() override;
 	virtual void Deactivate() override;
 	virtual void Cleanup() override;
+	virtual void OnSceneObjectChanged(UObject* InOldSceneObject, UObject* InNewSceneObject) override;
 	virtual void RegisterTabSpawners(const TSharedRef<IAvaEditor>& InEditor) const override;
 	virtual void ExtendLevelEditorLayout(FLayoutExtender& InExtender) const override;
 	virtual void NotifyOnSelectionChanged(const FAvaEditorSelection& InSelection) override;
diff --git a/Engine/Plugins/Experimental/Avalanche/Source/AvalancheEditorCore/Private/AvaEditor.cpp b/Engine/Plugins/Experimental/Avalanche/Source/AvalancheEditorCore/Private/AvaEditor.cpp
--- a/Engine/Plugins/Experimental/Avalanche/Source/AvalancheEditorCore/Private/AvaEditor.cpp
+++ b/Engine/Plugins/Experimental/Avalanche/Source/AvalancheEditorCore/Private/AvaEditor.cpp
@@ -415,6 +415,21 @@ UWorld* FAvaEditor::GetWorld() const
 	return nullptr;
 }
 
+void FAvaEditor::OnSceneObjectChanged()
+{
+	UObject* OldSceneObject = GetSceneObject(EAvaEditorObjectQueryType::SkipSearch);
+
+	// Clear cached scene object as scene object changed
+	SceneObjectWeak.Reset();
+
+	UObject* NewSceneObject = GetSceneObject(EAvaEditorObjectQueryType::SearchOnly);
+
+	ForEachExtension([OldSceneObject, NewSceneObject](const TSharedRef<IAvaEditorExtension>& InExtension)
+	{
+		InExtension->OnSceneObjectChanged(OldSceneObject, NewSceneObject);
+	});
+}
+
 UObject* FAvaEditor::GetSceneObject(EAvaEditorObjectQueryType InQueryType) const
 {
 	// Return cached version if we are skipping search or if already valid
@@ -423,8 +438,12 @@ UObject* FAvaEditor::GetSceneObject(EAvaEditorObjectQueryType InQueryType) const
 		return SceneObjectWeak.Get();
 	}
 
-	UWorld* const SceneWorld   = GetWorld();
-	UObject* const SceneObject = Provider->GetSceneObject(SceneWorld, InQueryType);
+	UObject* SceneObject = nullptr;
+
+	if (UWorld* SceneWorld = GetWorld())
+	{
+		SceneObject = Provider->GetSceneObject(SceneWorld, InQueryType);
+	}
 
 	const_cast<FAvaEditor*>(this)->SceneObjectWeak = SceneObject;
 	return SceneObject;
diff --git a/Engine/Plugins/Experimental/Avalanche/Source/AvalancheEditorCore/Private/AvaEditor.h b/Engine/Plugins/Experimental/Avalanche/Source/AvalancheEditorCore/Private/AvaEditor.h
--- a/Engine/Plugins/Experimental/Avalanche/Source/AvalancheEditorCore/Private/AvaEditor.h
+++ b/Engine/Plugins/Experimental/Avalanche/Source/AvalancheEditorCore/Private/AvaEditor.h
@@ -53,6 +53,7 @@ protected:
 	virtual TSharedPtr<FTabManager> GetTabManager() const override;
 	virtual FEditorModeTools* GetEditorModeTools() const override;
 	virtual UWorld* GetWorld() const override;
+	virtual void OnSceneObjectChanged() override;
 	virtual UObject* GetSceneObject(EAvaEditorObjectQueryType InQueryType) const override;
 	virtual void RegisterTabSpawners() override;
 	virtual void UnregisterTabSpawners() override;
diff --git a/Engine/Plugins/Experimental/Avalanche/Source/AvalancheEditorCore/Private/LevelEditor/AvaLevelEditor.cpp b/Engine/Plugins/Experimental/Avalanche/Source/AvalancheEditorCore/Private/LevelEditor/AvaLevelEditor.cpp
--- a/Engine/Plugins/Experimental/Avalanche/Source/AvalancheEditorCore/Private/LevelEditor/AvaLevelEditor.cpp
+++ b/Engine/Plugins/Experimental/Avalanche/Source/AvalancheEditorCore/Private/LevelEditor/AvaLevelEditor.cpp
@@ -27,6 +27,9 @@ FAvaLevelEditor::~FAvaLevelEditor()
 		OnMapChangedHandle.Reset();
 		OnRegisterLayoutHandle.Reset();
 	}
+
+	FWorldDelegates::OnCurrentLevelChanged.Remove(OnCurrentLevelChangedHandle);
+	OnCurrentLevelChangedHandle.Reset();
 }
 
 void FAvaLevelEditor::Construct()
@@ -45,6 +48,8 @@ void FAvaLevelEditor::Construct()
 
 	FAvaEditor::Construct();
 
+	OnCurrentLevelChangedHandle = FWorldDelegates::OnCurrentLevelChanged.AddSP(This, &FAvaLevelEditor::OnCurrentLevelChanged);
+
 	Toolbar->Construct(This);
 
 	if (FLevelEditorModule* const LevelEditorModule = FAvaLevelEditorUtils::GetLevelEditorModule())
@@ -72,6 +77,11 @@ void FAvaLevelEditor::BindCommands(const TSharedRef<FUICommandList>& InCommandLi
 	FAvaEditor::BindCommands(InCommandList);
 }
 
+void FAvaLevelEditor::OnCurrentLevelChanged(ULevel* InNewLevel, ULevel* InOldLevel, UWorld* InWorld)
+{
+	OnSceneObjectChanged();
+}
+
 void FAvaLevelEditor::TryOpenScene(EAvaEditorObjectQueryType InQueryType)
 {
 	FLevelEditorModule* LevelEditorModule = FAvaLevelEditorUtils::LoadLevelEditorModule();
diff --git a/Engine/Plugins/Experimental/Avalanche/Source/AvalancheEditorCore/Private/LevelEditor/AvaLevelEditor.h b/Engine/Plugins/Experimental/Avalanche/Source/AvalancheEditorCore/Private/LevelEditor/AvaLevelEditor.h
--- a/Engine/Plugins/Experimental/Avalanche/Source/AvalancheEditorCore/Private/LevelEditor/AvaLevelEditor.h
+++ b/Engine/Plugins/Experimental/Avalanche/Source/AvalancheEditorCore/Private/LevelEditor/AvaLevelEditor.h
@@ -6,6 +6,7 @@
 
 class FAvaLevelEditorToolbar;
 class ILevelEditor;
+class ULevel;
 enum class EMapChangeType : uint8;
 
 class FAvaLevelEditor : public FAvaEditor
@@ -32,6 +33,8 @@ protected:
 	//~ End IAvaEditor
 
 private:
+	void OnCurrentLevelChanged(ULevel* InNewLevel, ULevel* InOldLevel, UWorld* InWorld);
+
 	void TryOpenScene(EAvaEditorObjectQueryType InQueryType = EAvaEditorObjectQueryType::SearchOnly);
 
 	void OnMapChanged(UWorld* InWorld, EMapChangeType InChangeType);
@@ -40,6 +43,8 @@ private:
 
 	TSharedRef<FAvaLevelEditorToolbar> Toolbar;
 
+	FDelegateHandle OnCurrentLevelChangedHandle;
+
 	FDelegateHandle OnRegisterLayoutHandle;
 
 	FDelegateHandle OnMapChangedHandle;
diff --git a/Engine/Plugins/Experimental/Avalanche/Source/AvalancheEditorCore/Public/IAvaEditor.h b/Engine/Plugins/Experimental/Avalanche/Source/AvalancheEditorCore/Public/IAvaEditor.h
--- a/Engine/Plugins/Experimental/Avalanche/Source/AvalancheEditorCore/Public/IAvaEditor.h
+++ b/Engine/Plugins/Experimental/Avalanche/Source/AvalancheEditorCore/Public/IAvaEditor.h
@@ -72,6 +72,8 @@ public:
 
 	virtual UWorld* GetWorld() const = 0;
 
+	virtual void OnSceneObjectChanged() = 0;
+
 	virtual UObject* GetSceneObject(EAvaEditorObjectQueryType InQueryType = EAvaEditorObjectQueryType::SearchOnly) const = 0;
 
 	virtual void RegisterTabSpawners() = 0;
diff --git a/Engine/Plugins/Experimental/Avalanche/Source/AvalancheEditorCore/Public/IAvaEditorExtension.h b/Engine/Plugins/Experimental/Avalanche/Source/AvalancheEditorCore/Public/IAvaEditorExtension.h
--- a/Engine/Plugins/Experimental/Avalanche/Source/AvalancheEditorCore/Public/IAvaEditorExtension.h
+++ b/Engine/Plugins/Experimental/Avalanche/Source/AvalancheEditorCore/Public/IAvaEditorExtension.h
@@ -51,6 +51,8 @@ public:
 
 	virtual void BindCommands(const TSharedRef<FUICommandList>& InCommandList) {}
 
+	virtual void OnSceneObjectChanged(UObject* InOldSceneObject, UObject* InNewSceneObject) {}
+
 	virtual void RegisterTabSpawners(const TSharedRef<IAvaEditor>& InEditor) const {}
 
 	/** Opportunity for an Extension to extend the Editor Toolbar */
@@ -111,7 +113,7 @@ public:
 	{
 		const TSharedPtr<IAvaEditor> Editor = GetEditor();
 		return Editor.IsValid()
-			? Cast<InSceneObjectType>(Editor->GetSceneObject(EAvaEditorObjectQueryType::SkipSearch))
+			? Cast<InSceneObjectType>(Editor->GetSceneObject(EAvaEditorObjectQueryType::SearchOnly))
 			: nullptr;
 	}
 
