Motion Design RC: Fix issue with Controller Picker not showing the controllers in user-defined order

diff --git a/Engine/Plugins/Experimental/Avalanche/Source/AvalancheRemoteControlEditor/Private/Customizations/Widgets/SAvaRCControllerPicker.cpp b/Engine/Plugins/Experimental/Avalanche/Source/AvalancheRemoteControlEditor/Private/Customizations/Widgets/SAvaRCControllerPicker.cpp
--- a/Engine/Plugins/Experimental/Avalanche/Source/AvalancheRemoteControlEditor/Private/Customizations/Widgets/SAvaRCControllerPicker.cpp
+++ b/Engine/Plugins/Experimental/Avalanche/Source/AvalancheRemoteControlEditor/Private/Customizations/Widgets/SAvaRCControllerPicker.cpp
@@ -134,13 +134,18 @@ void SAvaRCControllerPicker::RefreshOptions()
 	TSet<FName> ProcessedNames;
 	ProcessedNames.Reserve(Controllers.Num());
 
-	for (URCVirtualPropertyBase* Controller : Controllers)
-	{
-		if (!Controller)
+	Controllers.RemoveAll([](URCVirtualPropertyBase* InController)
 		{
-			continue;
-		}
+			return InController == nullptr;
+		});
 
+	Controllers.StableSort([](const URCVirtualPropertyBase& A, const URCVirtualPropertyBase& B)
+		{
+			return A.DisplayIndex < B.DisplayIndex;
+		});
+
+	for (URCVirtualPropertyBase* Controller : Controllers)
+	{
 		// Since FAvaRCControllerPickerOption only handles names directly, avoid adding the same name entry twice
 		bool bAlreadyInSet;
 		ProcessedNames.Add(Controller->DisplayName, &bAlreadyInSet);
