Transition Rendering

diff --git a/Engine/Plugins/Experimental/Avalanche/Source/AvalancheTransition/Private/AvaTransitionModule.cpp b/Engine/Plugins/Experimental/Avalanche/Source/AvalancheTransition/Private/AvaTransitionModule.cpp
--- a/Engine/Plugins/Experimental/Avalanche/Source/AvalancheTransition/Private/AvaTransitionModule.cpp
+++ b/Engine/Plugins/Experimental/Avalanche/Source/AvalancheTransition/Private/AvaTransitionModule.cpp
@@ -1,18 +1,33 @@
 // Copyright Epic Games, Inc. All Rights Reserved.
 
-#include "IAvaTransitionModule.h"
+#include "AvaTransitionModule.h"
 #include "AvaTransitionLog.h"
+#include "Rendering/AvaTransitionSceneViewExtension.h"
+#include "SceneViewExtension.h"
 
 DEFINE_LOG_CATEGORY(LogAvaTransition);
 
-class FAvaTransitionModule : public IAvaTransitionModule
+void FAvaTransitionModule::StartupModule()
 {
-	virtual FOnValidateTransitionTree& GetOnValidateTransitionTree()
-	{
-		return OnValidateStateTree;
-	}
+	OnPostEngineInitHandle = FCoreDelegates::OnPostEngineInit.AddRaw(this, &FAvaTransitionModule::PostEngineInit);
+}
 
-	FOnValidateTransitionTree OnValidateStateTree;
-};
+void FAvaTransitionModule::ShutdownModule()
+{
+	FCoreDelegates::OnPostEngineInit.Remove(OnPostEngineInitHandle);
+	OnPostEngineInitHandle.Reset();
+
+	TransitionSceneViewExtension.Reset();
+}
+
+IAvaTransitionModule::FOnValidateTransitionTree& FAvaTransitionModule::GetOnValidateTransitionTree()
+{
+	return OnValidateStateTree;
+}
+
+void FAvaTransitionModule::PostEngineInit()
+{
+	TransitionSceneViewExtension = FSceneViewExtensions::NewExtension<FAvaTransitionSceneViewExtension>();
+}
 
 IMPLEMENT_MODULE(FAvaTransitionModule, AvalancheTransition)
diff --git a/Engine/Plugins/Experimental/Avalanche/Source/AvalancheTransition/Private/AvaTransitionModule.h b/Engine/Plugins/Experimental/Avalanche/Source/AvalancheTransition/Private/AvaTransitionModule.h
new file mode 100644
--- /dev/null
+++ b/Engine/Plugins/Experimental/Avalanche/Source/AvalancheTransition/Private/AvaTransitionModule.h
@@ -0,0 +1,28 @@
+// Copyright Epic Games, Inc. All Rights Reserved.
+
+#pragma once
+
+#include "IAvaTransitionModule.h"
+#include "Templates/SharedPointer.h"
+
+class FAvaTransitionSceneViewExtension;
+
+class FAvaTransitionModule : public IAvaTransitionModule
+{
+	//~ Begin IModuleInterface
+	virtual void StartupModule() override;
+	virtual void ShutdownModule() override;
+	//~ End IModuleInterface
+
+	//~ Begin IAvaTransitionModule
+	virtual FOnValidateTransitionTree& GetOnValidateTransitionTree() override;
+	//~ End IAvaTransitionModule
+
+	void PostEngineInit();
+
+	TSharedPtr<FAvaTransitionSceneViewExtension> TransitionSceneViewExtension;
+
+	FOnValidateTransitionTree OnValidateStateTree;
+
+	FDelegateHandle OnPostEngineInitHandle;
+};
diff --git a/Engine/Plugins/Experimental/Avalanche/Source/AvalancheTransition/Private/Rendering/AvaTransitionRenderingSubsystem.cpp b/Engine/Plugins/Experimental/Avalanche/Source/AvalancheTransition/Private/Rendering/AvaTransitionRenderingSubsystem.cpp
new file mode 100644
--- /dev/null
+++ b/Engine/Plugins/Experimental/Avalanche/Source/AvalancheTransition/Private/Rendering/AvaTransitionRenderingSubsystem.cpp
@@ -0,0 +1,50 @@
+// Copyright Epic Games, Inc. All Rights Reserved.
+
+#include "Rendering/AvaTransitionRenderingSubsystem.h"
+#include "Components/PrimitiveComponent.h"
+#include "Engine/Level.h"
+#include "GameFramework/Actor.h"
+#include "SceneView.h"
+
+void UAvaTransitionRenderingSubsystem::ShowLevel(TObjectKey<ULevel> InLevel)
+{
+	HiddenLevels.Remove(InLevel);
+}
+
+void UAvaTransitionRenderingSubsystem::HideLevel(TObjectKey<ULevel> InLevel)
+{
+	HiddenLevels.Add(InLevel);
+}
+
+void UAvaTransitionRenderingSubsystem::SetupView(FSceneView& InView)
+{
+	HiddenPrimitives.Reset();
+
+	for (const TObjectKey<ULevel>& LevelKey : HiddenLevels)
+	{
+		const ULevel* Level = LevelKey.ResolveObjectPtr();
+		if (!Level)
+		{
+			continue;
+		}
+
+		for (AActor* Actor : Level->Actors)
+		{
+			if (!Actor)
+			{
+				continue;
+			}
+
+			Actor->ForEachComponent<UPrimitiveComponent>(/*bIncludeFromChildActors*/false,
+				[&HiddenPrimitives = HiddenPrimitives](UPrimitiveComponent* InComponent)
+				{
+					if (InComponent->IsRegistered())
+					{
+						HiddenPrimitives.Add(InComponent->GetPrimitiveSceneId());
+					}
+				});
+		}
+	}
+
+	InView.HiddenPrimitives.Append(HiddenPrimitives);
+}
diff --git a/Engine/Plugins/Experimental/Avalanche/Source/AvalancheTransition/Private/Rendering/AvaTransitionSceneViewExtension.cpp b/Engine/Plugins/Experimental/Avalanche/Source/AvalancheTransition/Private/Rendering/AvaTransitionSceneViewExtension.cpp
new file mode 100644
--- /dev/null
+++ b/Engine/Plugins/Experimental/Avalanche/Source/AvalancheTransition/Private/Rendering/AvaTransitionSceneViewExtension.cpp
@@ -0,0 +1,27 @@
+// Copyright Epic Games, Inc. All Rights Reserved.
+
+#include "AvaTransitionSceneViewExtension.h"
+#include "Rendering/AvaTransitionRenderingSubsystem.h"
+#include "SceneView.h"
+
+void FAvaTransitionSceneViewExtension::SetupView(FSceneViewFamily& InViewFamily, FSceneView& InView)
+{
+	if (!InViewFamily.Scene)
+	{
+		return;
+	}
+
+	const UWorld* World = InViewFamily.Scene->GetWorld();
+	if (!World)
+	{
+		return;
+	}
+
+	UAvaTransitionRenderingSubsystem* RenderingSubsystem = World->GetSubsystem<UAvaTransitionRenderingSubsystem>();
+	if (!RenderingSubsystem)
+	{
+		return;
+	}
+
+	RenderingSubsystem->SetupView(InView);
+}
diff --git a/Engine/Plugins/Experimental/Avalanche/Source/AvalancheTransition/Private/Rendering/AvaTransitionSceneViewExtension.h b/Engine/Plugins/Experimental/Avalanche/Source/AvalancheTransition/Private/Rendering/AvaTransitionSceneViewExtension.h
new file mode 100644
--- /dev/null
+++ b/Engine/Plugins/Experimental/Avalanche/Source/AvalancheTransition/Private/Rendering/AvaTransitionSceneViewExtension.h
@@ -0,0 +1,20 @@
+// Copyright Epic Games, Inc. All Rights Reserved.
+
+#pragma once
+
+#include "SceneViewExtension.h"
+
+class FAvaTransitionSceneViewExtension : public FSceneViewExtensionBase
+{
+public:
+	FAvaTransitionSceneViewExtension(const FAutoRegister& InAutoRegister)
+		: FSceneViewExtensionBase(InAutoRegister)
+	{
+	}
+
+	//~ Begin FSceneViewExtensionBase
+	void SetupViewFamily(FSceneViewFamily& InViewFamily) override {}
+	void SetupView(FSceneViewFamily& InViewFamily, FSceneView& InView) override;
+	void BeginRenderViewFamily(FSceneViewFamily& InViewFamily) override {}
+	//~ End FSceneViewExtensionBase
+};
diff --git a/Engine/Plugins/Experimental/Avalanche/Source/AvalancheTransition/Private/Tasks/AvaTransitionWaitForLayerTask.cpp b/Engine/Plugins/Experimental/Avalanche/Source/AvalancheTransition/Private/Tasks/AvaTransitionWaitForLayerTask.cpp
--- a/Engine/Plugins/Experimental/Avalanche/Source/AvalancheTransition/Private/Tasks/AvaTransitionWaitForLayerTask.cpp
+++ b/Engine/Plugins/Experimental/Avalanche/Source/AvalancheTransition/Private/Tasks/AvaTransitionWaitForLayerTask.cpp
@@ -3,7 +3,10 @@
 #include "Tasks/AvaTransitionWaitForLayerTask.h"
 #include "AvaTransitionLayerUtils.h"
 #include "Behavior/AvaTransitionBehaviorInstance.h"
+#include "Engine/Level.h"
+#include "Rendering/AvaTransitionRenderingSubsystem.h"
 #include "StateTreeExecutionContext.h"
+#include "StateTreeLinker.h"
 
 #define LOCTEXT_NAMESPACE "AvaTransitionWaitForLayerTask"
 
@@ -18,17 +21,24 @@ FText FAvaTransitionWaitForLayerTask::GetDescription(const FGuid& InId, FStateTr
 }
 #endif
 
+bool FAvaTransitionWaitForLayerTask::Link(FStateTreeLinker& InLinker)
+{
+	Super::Link(InLinker);
+	InLinker.LinkExternalData(RenderingSubsystemHandle);
+	return true;
+}
+
 EStateTreeRunStatus FAvaTransitionWaitForLayerTask::EnterState(FStateTreeExecutionContext& InContext, const FStateTreeTransitionResult& InTransition) const
 {
-	return QueryStatus(InContext);
+	return WaitForLayer(InContext);
 }
 
 EStateTreeRunStatus FAvaTransitionWaitForLayerTask::Tick(FStateTreeExecutionContext& InContext, const float InDeltaTime) const
 {
-	return QueryStatus(InContext);
+	return WaitForLayer(InContext);
 }
 
-EStateTreeRunStatus FAvaTransitionWaitForLayerTask::QueryStatus(FStateTreeExecutionContext& InContext) const
+EStateTreeRunStatus FAvaTransitionWaitForLayerTask::WaitForLayer(FStateTreeExecutionContext& InContext) const
 {
 	TArray<const FAvaTransitionBehaviorInstance*> BehaviorInstances = QueryBehaviorInstances(InContext);
 
@@ -39,11 +49,28 @@ EStateTreeRunStatus FAvaTransitionWaitForLayerTask::QueryStatus(FStateTreeExecut
 			return InInstance->IsRunning();
 		});
 
+	FAvaTransitionWaitForLayerTask::FInstanceDataType& InstanceData = InContext.GetInstanceData(*this);
+
+	UAvaTransitionRenderingSubsystem& RenderingSubsystem = InContext.GetExternalData(RenderingSubsystemHandle);
+
 	if (bIsLayerRunning)
 	{
+		// Hide Level Primitives while Waiting for other Layers
+		if (InstanceData.bHideSceneWhileWaiting && InstanceData.HiddenLevel == nullptr)
+		{
+			const FAvaTransitionContext& TransitionContext = InContext.GetExternalData(TransitionContextHandle);
+			if (const FAvaTransitionScene* TransitionScene = TransitionContext.GetTransitionScene())
+			{
+				InstanceData.HiddenLevel = TransitionScene->GetLevel();
+				RenderingSubsystem.HideLevel(InstanceData.HiddenLevel);
+			}
+		}
+
 		return EStateTreeRunStatus::Running;
 	}
 
+	// Restore Level Visibility
+	RenderingSubsystem.ShowLevel(InstanceData.HiddenLevel);
 	return EStateTreeRunStatus::Succeeded;
 }
 
diff --git a/Engine/Plugins/Experimental/Avalanche/Source/AvalancheTransition/Public/Rendering/AvaTransitionRenderingSubsystem.h b/Engine/Plugins/Experimental/Avalanche/Source/AvalancheTransition/Public/Rendering/AvaTransitionRenderingSubsystem.h
new file mode 100644
--- /dev/null
+++ b/Engine/Plugins/Experimental/Avalanche/Source/AvalancheTransition/Public/Rendering/AvaTransitionRenderingSubsystem.h
@@ -0,0 +1,29 @@
+
+// Copyright Epic Games, Inc. All Rights Reserved.
+#pragma once
+
+#include "Containers/Set.h"
+#include "PrimitiveComponentId.h"
+#include "Subsystems/WorldSubsystem.h"
+#include "UObject/ObjectKey.h"
+#include "AvaTransitionRenderingSubsystem.generated.h"
+
+class FSceneView;
+class ULevel;
+
+UCLASS(MinimalAPI)
+class UAvaTransitionRenderingSubsystem : public UWorldSubsystem
+{
+	GENERATED_BODY()
+
+public:
+	AVALANCHETRANSITION_API void ShowLevel(TObjectKey<ULevel> InLevel);
+	AVALANCHETRANSITION_API void HideLevel(TObjectKey<ULevel> InLevel);
+
+	void SetupView(FSceneView& InView);
+
+private:
+	TSet<TObjectKey<ULevel>> HiddenLevels;
+
+	TSet<FPrimitiveComponentId> HiddenPrimitives;
+};
diff --git a/Engine/Plugins/Experimental/Avalanche/Source/AvalancheTransition/Public/Tasks/AvaTransitionWaitForLayerTask.h b/Engine/Plugins/Experimental/Avalanche/Source/AvalancheTransition/Public/Tasks/AvaTransitionWaitForLayerTask.h
--- a/Engine/Plugins/Experimental/Avalanche/Source/AvalancheTransition/Public/Tasks/AvaTransitionWaitForLayerTask.h
+++ b/Engine/Plugins/Experimental/Avalanche/Source/AvalancheTransition/Public/Tasks/AvaTransitionWaitForLayerTask.h
@@ -3,12 +3,21 @@
 #pragma once
 
 #include "AvaTransitionLayerTask.h"
+#include "UObject/ObjectKey.h"
 #include "AvaTransitionWaitForLayerTask.generated.h"
 
+class UAvaTransitionRenderingSubsystem;
+class ULevel;
+
 USTRUCT()
 struct FAvaTransitionWaitForLayerTaskInstanceData : public FAvaTransitionLayerTaskInstanceData
 {
 	GENERATED_BODY()
+
+	UPROPERTY(EditAnywhere, Category="Transition Logic")
+	bool bHideSceneWhileWaiting = true;
+
+	TObjectKey<ULevel> HiddenLevel;
 };
 
 USTRUCT(DisplayName = "Wait for other Scenes in Layer to Finish", Category="Transition Logic")
@@ -23,6 +32,7 @@ struct AVALANCHETRANSITION_API FAvaTransitionWaitForLayerTask : public FAvaTrans
 	virtual FText GetDescription(const FGuid& InId, FStateTreeDataView InInstanceDataView, const IStateTreeBindingLookup& InBindingLookup, EStateTreeNodeFormatting InFormatting) const override;
 #endif
 	virtual const UStruct* GetInstanceDataType() const override { return FInstanceDataType::StaticStruct(); }
+	virtual bool Link(FStateTreeLinker& InLinker) override;
 	//~ End FStateTreeNodeBase
 
 	//~ Begin FStateTreeTaskBase
@@ -30,5 +40,7 @@ struct AVALANCHETRANSITION_API FAvaTransitionWaitForLayerTask : public FAvaTrans
 	virtual EStateTreeRunStatus Tick(FStateTreeExecutionContext& InContext, const float InDeltaTime) const override;
 	//~ End FStateTreeTaskBase
 
-	EStateTreeRunStatus QueryStatus(FStateTreeExecutionContext& InContext) const;
+	EStateTreeRunStatus WaitForLayer(FStateTreeExecutionContext& InContext) const;
+
+	TStateTreeExternalDataHandle<UAvaTransitionRenderingSubsystem> RenderingSubsystemHandle;
 };
