[Motion Design] Transition Logic: fix the controller comparison when reusing the scene.

diff --git a/Engine/Plugins/Experimental/Avalanche/Source/AvalancheMedia/Private/Playable/Transition/AvaPlayableTransitionScene.cpp b/Engine/Plugins/Experimental/Avalanche/Source/AvalancheMedia/Private/Playable/Transition/AvaPlayableTransitionScene.cpp
--- a/Engine/Plugins/Experimental/Avalanche/Source/AvalancheMedia/Private/Playable/Transition/AvaPlayableTransitionScene.cpp
+++ b/Engine/Plugins/Experimental/Avalanche/Source/AvalancheMedia/Private/Playable/Transition/AvaPlayableTransitionScene.cpp
@@ -22,7 +22,24 @@ namespace UE::AvaMedia::Private
 			}
 
 			const FAvaPlayableRemoteControlValue* MyValue = MyPlayable->GetLatestRemoteControlValues().ControllerValues.Find(InControllerId);
-			const FAvaPlayableRemoteControlValue* OtherValue = OtherPlayable->GetLatestRemoteControlValues().ControllerValues.Find(InControllerId);
+			const FAvaPlayableRemoteControlValue* OtherValue = nullptr;
+
+			if (OtherPlayable != MyPlayable)
+			{
+				OtherValue = OtherPlayable->GetLatestRemoteControlValues().ControllerValues.Find(InControllerId);
+			}
+			else
+			{
+				// In case of scene reuse, compare against RC values incoming in the transition.
+				if (UAvaPlayableTransition* Transition = static_cast<const FAvaPlayableTransitionScene&>(InOtherScene).PlayableTransitionWeak.Get())
+				{
+					if (const TSharedPtr<FAvaPlayableRemoteControlValues> OtherValues = Transition->GetValuesForPlayable(OtherPlayable))
+					{
+						OtherValue = OtherValues->ControllerValues.Find(InControllerId);
+					}
+				}
+			}
+
 			if (!MyValue || !OtherValue)
 			{
 				return EAvaTransitionComparisonResult::None;
