Motion Design: Added Remote Control / Transition Logic BP functionality. This is to aid in TL BP Tasks to be able to for example get what has changed between Scenes

diff --git a/Engine/Plugins/Experimental/Avalanche/Source/Avalanche/Private/Transition/AvaTransitionRCLibrary.cpp b/Engine/Plugins/Experimental/Avalanche/Source/Avalanche/Private/Transition/AvaTransitionRCLibrary.cpp
--- a/Engine/Plugins/Experimental/Avalanche/Source/Avalanche/Private/Transition/AvaTransitionRCLibrary.cpp
+++ b/Engine/Plugins/Experimental/Avalanche/Source/Avalanche/Private/Transition/AvaTransitionRCLibrary.cpp
@@ -131,3 +131,69 @@ bool UAvaTransitionRCLibrary::CompareRCControllerValues(UObject* InTransitionNod
 
 	return CompareRCControllerValues(*TransitionContext, InControllerId, InValueComparisonType);
 }
+
+TArray<URCVirtualPropertyBase*> UAvaTransitionRCLibrary::GetChangedRCControllers(UObject* InTransitionNode)
+{
+	IAvaTransitionNodeInterface* NodeInterface = Cast<IAvaTransitionNodeInterface>(InTransitionNode);
+	if (!NodeInterface)
+	{
+		return {};
+	}
+
+	const FAvaTransitionContext* TransitionContext = NodeInterface->GetBehaviorInstanceCache().GetTransitionContext();
+	if (!TransitionContext)
+	{
+		return {};
+	}
+
+	const FAvaTransitionScene* TransitionScene = TransitionContext->GetTransitionScene();
+	if (!TransitionScene)
+	{
+		return {};
+	}
+
+	ULevel* Level = TransitionScene->GetLevel();
+	if (!Level || !Level->OwningWorld)
+	{
+		return {};
+	}
+
+	UAvaSceneSubsystem* SceneSubsystem = Level->OwningWorld->GetSubsystem<UAvaSceneSubsystem>();
+	if (!SceneSubsystem)
+	{
+		return {};
+	}
+
+	IAvaSceneInterface* SceneInterface = SceneSubsystem->GetSceneInterface(Level);
+	if (!SceneInterface)
+	{
+		return {};
+	}
+
+	URemoteControlPreset* RemoteControlPreset = SceneInterface->GetRemoteControlPreset();
+	if (!RemoteControlPreset)
+	{
+		return {};
+	}
+
+	TArray<URCVirtualPropertyBase*> Controllers = RemoteControlPreset->GetControllers();
+
+	// Remove Invalid Controllers and Controllers that are not different/changed between scenes
+	Controllers.RemoveAll(
+		[TransitionContext](URCVirtualPropertyBase* InController)
+		{
+			if (!InController)
+			{
+				return true;
+			}
+			return !CompareRCControllerValues(*TransitionContext, FAvaRCControllerId(InController), EAvaTransitionComparisonResult::Different);
+		});
+
+	Controllers.StableSort(
+		[](const URCVirtualPropertyBase& A, const URCVirtualPropertyBase& B)
+		{
+			return A.DisplayIndex < B.DisplayIndex;
+		});
+
+	return Controllers;
+}
diff --git a/Engine/Plugins/Experimental/Avalanche/Source/Avalanche/Private/Transition/AvaTransitionRCLibrary.h b/Engine/Plugins/Experimental/Avalanche/Source/Avalanche/Private/Transition/AvaTransitionRCLibrary.h
--- a/Engine/Plugins/Experimental/Avalanche/Source/Avalanche/Private/Transition/AvaTransitionRCLibrary.h
+++ b/Engine/Plugins/Experimental/Avalanche/Source/Avalanche/Private/Transition/AvaTransitionRCLibrary.h
@@ -21,6 +21,9 @@ class UAvaTransitionRCLibrary : public UBlueprintFunctionLibrary
 public:
 	static bool CompareRCControllerValues(const FAvaTransitionContext& InTransitionContext, const FAvaRCControllerId& InControllerId, EAvaTransitionComparisonResult InValueComparisonType);
 
-	UFUNCTION(BlueprintPure, DisplayName="Compare RC Controller Values", Category="Remote Control", meta=(DefaultToSelf="InTransitionNode"))
+	UFUNCTION(BlueprintPure, DisplayName="Compare RC Controller Values", Category="Transition Logic|Remote Control", meta=(DefaultToSelf="InTransitionNode"))
 	static bool CompareRCControllerValues(UObject* InTransitionNode, const FAvaRCControllerId& InControllerId, EAvaTransitionComparisonResult InValueComparisonType = EAvaTransitionComparisonResult::Different);
+
+	UFUNCTION(BlueprintCallable, DisplayName="Get Changed RC Controllers", Category="Transition Logic|Remote Control", meta=(DefaultToSelf="InTransitionNode"))
+	static TArray<URCVirtualPropertyBase*> GetChangedRCControllers(UObject* InTransitionNode);
 };
diff --git a/Engine/Plugins/Experimental/Avalanche/Source/AvalancheRemoteControl/AvalancheRemoteControl.Build.cs b/Engine/Plugins/Experimental/Avalanche/Source/AvalancheRemoteControl/AvalancheRemoteControl.Build.cs
--- a/Engine/Plugins/Experimental/Avalanche/Source/AvalancheRemoteControl/AvalancheRemoteControl.Build.cs
+++ b/Engine/Plugins/Experimental/Avalanche/Source/AvalancheRemoteControl/AvalancheRemoteControl.Build.cs
@@ -22,8 +22,9 @@ public class AvalancheRemoteControl : ModuleRules
             new string[]
             {
                 "RemoteControlComponents", 
+                "RemoteControlLogic",
                 "Slate",
-                "SlateCore", 
+                "SlateCore",
             }
         );
 
diff --git a/Engine/Plugins/Experimental/Avalanche/Source/AvalancheRemoteControl/Private/AvaRCControllerId.cpp b/Engine/Plugins/Experimental/Avalanche/Source/AvalancheRemoteControl/Private/AvaRCControllerId.cpp
--- a/Engine/Plugins/Experimental/Avalanche/Source/AvalancheRemoteControl/Private/AvaRCControllerId.cpp
+++ b/Engine/Plugins/Experimental/Avalanche/Source/AvalancheRemoteControl/Private/AvaRCControllerId.cpp
@@ -1,8 +1,14 @@
 // Copyright Epic Games, Inc. All Rights Reserved.
 
 #include "AvaRCControllerId.h"
+#include "RCVirtualProperty.h"
 #include "RemoteControlPreset.h"
 
+FAvaRCControllerId::FAvaRCControllerId(URCVirtualPropertyBase* InController)
+	: Name(InController ? InController->DisplayName : NAME_None)
+{
+}
+
 URCVirtualPropertyBase* FAvaRCControllerId::FindController(URemoteControlPreset* InPreset) const
 {
 	if (InPreset)
diff --git a/Engine/Plugins/Experimental/Avalanche/Source/AvalancheRemoteControl/Private/AvaRCLibrary.cpp b/Engine/Plugins/Experimental/Avalanche/Source/AvalancheRemoteControl/Private/AvaRCLibrary.cpp
new file mode 100644
--- /dev/null
+++ b/Engine/Plugins/Experimental/Avalanche/Source/AvalancheRemoteControl/Private/AvaRCLibrary.cpp
@@ -0,0 +1,78 @@
+// Copyright Epic Games, Inc. All Rights Reserved.
+
+#include "AvaRCLibrary.h"
+#include "Action/RCAction.h"
+#include "Behaviour/RCBehaviour.h"
+#include "Controller/RCController.h"
+#include "Engine/World.h"
+#include "RemoteControlPreset.h"
+
+TArray<AActor*> UAvaRCLibrary::GetControlledActors(UObject* InWorldContextObject, URCVirtualPropertyBase* InController)
+{
+	if (!InWorldContextObject)
+	{
+		return {};
+	}
+
+	/**
+	 * UObject::GetWorld will return the OwningWorld for Actors/Components/Levels. Instead, RC Bindings for streamed in actors
+	 * Will be based on the Streamed World (i.e. world outer), instead of the owning World.
+	 */
+	const UWorld* World = InWorldContextObject->IsA<UWorld>()
+		? CastChecked<UWorld>(InWorldContextObject)
+		: InWorldContextObject->GetTypedOuter<UWorld>();
+
+	if (!World)
+	{
+		return {};
+	}
+
+	URCController* Controller = Cast<URCController>(InController);
+	if (!Controller)
+	{
+		return {};
+	}
+
+	URemoteControlPreset* Preset = Controller->GetTypedOuter<URemoteControlPreset>();
+	if (!Preset)
+	{
+		return {};
+	}
+
+	TArray<AActor*> ControlledActors;
+
+	for (URCBehaviour* Behavior : Controller->Behaviours)
+	{
+		if (!Behavior || !Behavior->ActionContainer)
+		{
+			continue;
+		}
+
+		Behavior->ActionContainer->ForEachAction([Preset, World, &ControlledActors](URCAction* InAction)
+			{
+				TSharedPtr<FRemoteControlEntity> Entity = Preset->GetExposedEntity(InAction->ExposedFieldId).Pin();
+				if (!Entity.IsValid())
+				{
+					return;
+				}
+
+				UObject* Object = Entity->GetBoundObjectForWorld(World);
+				if (!Object)
+				{
+					return;
+				}
+
+				if (AActor* Actor = Cast<AActor>(Object))
+				{
+					ControlledActors.Add(Actor);
+				}
+				else if (AActor* OuterActor = Object->GetTypedOuter<AActor>())
+				{
+					ControlledActors.Add(OuterActor);
+				}
+			}
+			, /*bRecursive*/true);
+	}
+
+	return ControlledActors;
+}
diff --git a/Engine/Plugins/Experimental/Avalanche/Source/AvalancheRemoteControl/Public/AvaRCControllerId.h b/Engine/Plugins/Experimental/Avalanche/Source/AvalancheRemoteControl/Public/AvaRCControllerId.h
--- a/Engine/Plugins/Experimental/Avalanche/Source/AvalancheRemoteControl/Public/AvaRCControllerId.h
+++ b/Engine/Plugins/Experimental/Avalanche/Source/AvalancheRemoteControl/Public/AvaRCControllerId.h
@@ -14,6 +14,10 @@ struct FAvaRCControllerId
 {
 	GENERATED_BODY()
 
+	FAvaRCControllerId() = default;
+
+	AVALANCHEREMOTECONTROL_API explicit FAvaRCControllerId(URCVirtualPropertyBase* InController); 
+
 	AVALANCHEREMOTECONTROL_API URCVirtualPropertyBase* FindController(URemoteControlPreset* InPreset) const;
 
 	AVALANCHEREMOTECONTROL_API FText ToText() const;
diff --git a/Engine/Plugins/Experimental/Avalanche/Source/AvalancheRemoteControl/Public/AvaRCLibrary.h b/Engine/Plugins/Experimental/Avalanche/Source/AvalancheRemoteControl/Public/AvaRCLibrary.h
new file mode 100644
--- /dev/null
+++ b/Engine/Plugins/Experimental/Avalanche/Source/AvalancheRemoteControl/Public/AvaRCLibrary.h
@@ -0,0 +1,18 @@
+// Copyright Epic Games, Inc. All Rights Reserved.
+
+#pragma once
+
+#include "Kismet/BlueprintFunctionLibrary.h"
+#include "AvaRCLibrary.generated.h"
+
+class URCVirtualPropertyBase;
+
+UCLASS()
+class UAvaRCLibrary : public UBlueprintFunctionLibrary
+{
+	GENERATED_BODY()
+
+public:
+	UFUNCTION(BlueprintCallable, Category="Remote Control", meta=(WorldContext="InWorldContextObject"))
+	static AVALANCHEREMOTECONTROL_API TArray<AActor*> GetControlledActors(UObject* InWorldContextObject, URCVirtualPropertyBase* InController);
+};
