Motion Design: Fixed mask dirtying the level

diff --git a/Engine/Plugins/Experimental/Avalanche/Source/AvalancheMask/Private/AvaMaskUtilities.cpp b/Engine/Plugins/Experimental/Avalanche/Source/AvalancheMask/Private/AvaMaskUtilities.cpp
--- a/Engine/Plugins/Experimental/Avalanche/Source/AvalancheMask/Private/AvaMaskUtilities.cpp
+++ b/Engine/Plugins/Experimental/Avalanche/Source/AvalancheMask/Private/AvaMaskUtilities.cpp
@@ -76,7 +76,7 @@ namespace UE::AvaMask::Internal
 		if (!Component)
 		{
 #if WITH_EDITOR
-			InActor->Modify();
+			InActor->Modify(/*bAlwaysMarkDirty*/false);
 #endif
 
 			// Construct the new component and attach as needed
@@ -85,8 +85,10 @@ namespace UE::AvaMask::Internal
 				, MakeUniqueObjectName(InActor, ComponentClass)
 				, RF_Transactional);
 
+			Component->CreationMethod = EComponentCreationMethod::Instance;
+
 			// Add to SerializedComponents array so it gets saved
-			InActor->AddInstanceComponent(Component);
+			InActor->AddOwnedComponent(Component);
 			Component->OnComponentCreated();
 			Component->RegisterComponent();
 
