MotionDesign : AvaOutliner
- Add delete chips capability, item can now implement custom behavior
- Support delete, move, copy operations for animators/modifiers (Holding ALT for copy else move)
- Refreshing operator stack on item deletion

diff --git a/Engine/Plugins/Experimental/Avalanche/Source/AvalancheEditor/Private/OperatorStack/Widgets/SAvaOperatorStackTab.cpp b/Engine/Plugins/Experimental/Avalanche/Source/AvalancheEditor/Private/OperatorStack/Widgets/SAvaOperatorStackTab.cpp
--- a/Engine/Plugins/Experimental/Avalanche/Source/AvalancheEditor/Private/OperatorStack/Widgets/SAvaOperatorStackTab.cpp
+++ b/Engine/Plugins/Experimental/Avalanche/Source/AvalancheEditor/Private/OperatorStack/Widgets/SAvaOperatorStackTab.cpp
@@ -20,7 +20,7 @@ void SAvaOperatorStackTab::Construct(const FArguments& InArgs
 	, const TSharedPtr<IAvaDetailsProvider>& InProvider)
 {
 	DetailsProviderWeak = InProvider;
-	
+
 	UOperatorStackEditorSubsystem* OperatorStackSubsystem = UOperatorStackEditorSubsystem::Get();
 
 	check(InProvider.IsValid() && OperatorStackSubsystem);
@@ -33,16 +33,16 @@ void SAvaOperatorStackTab::Construct(const FArguments& InArgs
 	UActorModifierCoreStack::OnModifierRemovedDelegate.AddSP(this, &SAvaOperatorStackTab::OnModifierUpdated);
 
 	// Property controllers delegates
-	UPropertyAnimatorCoreBase::OnAnimatorCreatedDelegate.AddSP(this, &SAvaOperatorStackTab::OnControllerUpdated);
-	UPropertyAnimatorCoreBase::OnAnimatorRemovedDelegate.AddSP(this, &SAvaOperatorStackTab::OnControllerUpdated);
-	UPropertyAnimatorCoreBase::OnAnimatorRenamedDelegate.AddSP(this, &SAvaOperatorStackTab::OnControllerUpdated);
+	UPropertyAnimatorCoreBase::OnAnimatorCreatedDelegate.AddSP(this, &SAvaOperatorStackTab::OnAnimatorUpdated);
+	UPropertyAnimatorCoreBase::OnAnimatorRemovedDelegate.AddSP(this, &SAvaOperatorStackTab::OnAnimatorUpdated);
+	UPropertyAnimatorCoreBase::OnAnimatorRenamedDelegate.AddSP(this, &SAvaOperatorStackTab::OnAnimatorUpdated);
 
 	const TSharedPtr<IDetailKeyframeHandler> KeyframeHandler = InProvider->GetDetailsKeyframeHandler();
-	
+
 	OperatorStack = OperatorStackSubsystem->GenerateWidget();
 	OperatorStack->SetKeyframeHandler(KeyframeHandler);
 	OperatorStack->SetPanelTag(SAvaOperatorStackTab::PanelTag);
-	
+
 	ChildSlot
 	[
 		OperatorStack.ToSharedRef()
@@ -54,11 +54,11 @@ void SAvaOperatorStackTab::Construct(const FArguments& InArgs
 SAvaOperatorStackTab::~SAvaOperatorStackTab()
 {
 	USelection::SelectionChangedEvent.RemoveAll(this);
-	
+
 	UActorModifierCoreStack::OnModifierAddedDelegate.RemoveAll(this);
 	UActorModifierCoreStack::OnModifierMovedDelegate.RemoveAll(this);
 	UActorModifierCoreStack::OnModifierRemovedDelegate.RemoveAll(this);
-	
+
 	UPropertyAnimatorCoreBase::OnAnimatorCreatedDelegate.RemoveAll(this);
 	UPropertyAnimatorCoreBase::OnAnimatorRemovedDelegate.RemoveAll(this);
 	UPropertyAnimatorCoreBase::OnAnimatorRenamedDelegate.RemoveAll(this);
@@ -71,7 +71,7 @@ void SAvaOperatorStackTab::RefreshSelection(UObject* InSelectionObject) const
 	{
 		return;
 	}
-	
+
 	FEditorModeTools* ModeTools = DetailsProvider->GetDetailsModeTools();
 	if (!ModeTools)
 	{
@@ -98,12 +98,15 @@ void SAvaOperatorStackTab::RefreshSelection(UObject* InSelectionObject) const
 
 void SAvaOperatorStackTab::OnModifierUpdated(UActorModifierCoreBase* InUpdatedItem) const
 {
-	RefreshCurrentSelection(InUpdatedItem);
+	if (InUpdatedItem)
+	{
+		RefreshCurrentSelection(InUpdatedItem->GetModifierStack());
+	}
 }
 
-void SAvaOperatorStackTab::OnControllerUpdated(UPropertyAnimatorCoreBase* InController) const
+void SAvaOperatorStackTab::OnAnimatorUpdated(UPropertyAnimatorCoreBase* InUpdatedItem) const
 {
-	RefreshCurrentSelection(InController);
+	RefreshCurrentSelection(InUpdatedItem);
 }
 
 void SAvaOperatorStackTab::RefreshCurrentSelection(const UObject* InObject) const
@@ -113,14 +116,14 @@ void SAvaOperatorStackTab::RefreshCurrentSelection(const UObject* InObject) cons
 	{
 		return;
 	}
-	
+
 	FEditorModeTools* ModeTools = DetailsProvider->GetDetailsModeTools();
 	if (!InObject || !OperatorStack.IsValid() || !ModeTools)
 	{
 		return;
 	}
 
-	const AActor* OwningActor = InObject->GetTypedOuter<AActor>();
+	AActor* OwningActor = InObject->GetTypedOuter<AActor>();
 	if (!OwningActor)
 	{
 		return;
@@ -138,8 +141,8 @@ void SAvaOperatorStackTab::RefreshCurrentSelection(const UObject* InObject) cons
 		return;
 	}
 
-	const TArray<AActor*> SelectedActors = EditorSelection.GetSelectedObjects<AActor, EAvaSelectionSource::All>();
-	if (!SelectedActors.Contains(OwningActor))
+	const TArray<UObject*> SelectedObjects = EditorSelection.GetSelectedObjects<UObject, EAvaSelectionSource::All>();
+	if (!SelectedObjects.Contains(OwningActor) && !SelectedObjects.Contains(InObject))
 	{
 		return;
 	}
diff --git a/Engine/Plugins/Experimental/Avalanche/Source/AvalancheEditor/Private/OperatorStack/Widgets/SAvaOperatorStackTab.h b/Engine/Plugins/Experimental/Avalanche/Source/AvalancheEditor/Private/OperatorStack/Widgets/SAvaOperatorStackTab.h
--- a/Engine/Plugins/Experimental/Avalanche/Source/AvalancheEditor/Private/OperatorStack/Widgets/SAvaOperatorStackTab.h
+++ b/Engine/Plugins/Experimental/Avalanche/Source/AvalancheEditor/Private/OperatorStack/Widgets/SAvaOperatorStackTab.h
@@ -18,7 +18,7 @@ class SAvaOperatorStackTab : public SCompoundWidget
 {
 public:
 	static inline const FName PanelTag = TEXT("AvaOperatorStackTab");
-	
+
 	SLATE_BEGIN_ARGS(SAvaOperatorStackTab) {}
 	SLATE_END_ARGS()
 
@@ -29,9 +29,9 @@ public:
 private:
 	void RefreshSelection(UObject* InSelectionObject) const;
 	void OnModifierUpdated(UActorModifierCoreBase* InUpdatedItem) const;
-	void OnControllerUpdated(UPropertyAnimatorCoreBase* InController) const;
+	void OnAnimatorUpdated(UPropertyAnimatorCoreBase* InUpdatedItem) const;
 	void RefreshCurrentSelection(const UObject* InObject) const;
-	
+
 	TWeakPtr<IAvaDetailsProvider> DetailsProviderWeak;
 
 	TSharedPtr<SOperatorStackEditorWidget> OperatorStack;
diff --git a/Engine/Plugins/Experimental/Avalanche/Source/AvalancheEditor/Private/Outliner/AvaOutlinerExtension.cpp b/Engine/Plugins/Experimental/Avalanche/Source/AvalancheEditor/Private/Outliner/AvaOutlinerExtension.cpp
--- a/Engine/Plugins/Experimental/Avalanche/Source/AvalancheEditor/Private/Outliner/AvaOutlinerExtension.cpp
+++ b/Engine/Plugins/Experimental/Avalanche/Source/AvalancheEditor/Private/Outliner/AvaOutlinerExtension.cpp
@@ -247,6 +247,15 @@ void FAvaOutlinerExtension::OutlinerDuplicateActors(const TArray<AActor*>& InTem
 	}
 }
 
+void FAvaOutlinerExtension::OutlinerDeleteActors(const TArray<AActor*>& InDeleteActors)
+{
+	if (const TSharedPtr<IAvaEditor> Editor = GetEditor())
+	{
+		// TODO: Set Selection to Delete Actors if different
+		Editor->EditDelete();
+	}
+}
+
 FEditorModeTools* FAvaOutlinerExtension::GetOutlinerModeTools() const
 {
 	return GetEditorModeTools();
diff --git a/Engine/Plugins/Experimental/Avalanche/Source/AvalancheEditor/Private/Outliner/AvaOutlinerExtension.h b/Engine/Plugins/Experimental/Avalanche/Source/AvalancheEditor/Private/Outliner/AvaOutlinerExtension.h
--- a/Engine/Plugins/Experimental/Avalanche/Source/AvalancheEditor/Private/Outliner/AvaOutlinerExtension.h
+++ b/Engine/Plugins/Experimental/Avalanche/Source/AvalancheEditor/Private/Outliner/AvaOutlinerExtension.h
@@ -47,6 +47,7 @@ public:
 	virtual bool CanOutlinerProcessActorSpawn(AActor* InActor) const override;
 	virtual bool ShouldHideItem(const FAvaOutlinerItemPtr& InItem) const override;
 	virtual void OutlinerDuplicateActors(const TArray<AActor*>& InTemplateActors) override;
+	virtual void OutlinerDeleteActors(const TArray<AActor*>& InDeleteActors) override;
 	virtual FEditorModeTools* GetOutlinerModeTools() const override;
 	virtual FAvaSceneTree* GetSceneTree() const override;
 	virtual UWorld* GetOutlinerWorld() const override;
diff --git a/Engine/Plugins/Experimental/Avalanche/Source/AvalancheModifiersEditor/Private/Outliner/AvaOutlinerModifier.cpp b/Engine/Plugins/Experimental/Avalanche/Source/AvalancheModifiersEditor/Private/Outliner/AvaOutlinerModifier.cpp
--- a/Engine/Plugins/Experimental/Avalanche/Source/AvalancheModifiersEditor/Private/Outliner/AvaOutlinerModifier.cpp
+++ b/Engine/Plugins/Experimental/Avalanche/Source/AvalancheModifiersEditor/Private/Outliner/AvaOutlinerModifier.cpp
@@ -85,6 +85,26 @@ void FAvaOutlinerModifier::OnVisibilityChanged(EAvaOutlinerVisibilityType InVisi
 	}
 }
 
+bool FAvaOutlinerModifier::CanDelete() const
+{
+	return Modifier.IsValid();
+}
+
+bool FAvaOutlinerModifier::Delete()
+{
+	if (const UActorModifierCoreSubsystem* ModifierSubsystem = UActorModifierCoreSubsystem::Get())
+	{
+		FText FailReason;
+		FActorModifierCoreStackRemoveOp RemoveOp;
+		RemoveOp.FailReason = &FailReason;
+		RemoveOp.bShouldTransact = false;
+
+		return ModifierSubsystem->RemoveModifiers({Modifier.Get()}, RemoveOp);
+	}
+
+	return false;
+}
+
 void FAvaOutlinerModifier::SetObject_Impl(UObject* InObject)
 {
 	FAvaOutlinerObject::SetObject_Impl(InObject);
diff --git a/Engine/Plugins/Experimental/Avalanche/Source/AvalancheModifiersEditor/Private/Outliner/AvaOutlinerModifier.h b/Engine/Plugins/Experimental/Avalanche/Source/AvalancheModifiersEditor/Private/Outliner/AvaOutlinerModifier.h
--- a/Engine/Plugins/Experimental/Avalanche/Source/AvalancheModifiersEditor/Private/Outliner/AvaOutlinerModifier.h
+++ b/Engine/Plugins/Experimental/Avalanche/Source/AvalancheModifiersEditor/Private/Outliner/AvaOutlinerModifier.h
@@ -13,11 +13,11 @@ class AVALANCHEMODIFIERSEDITOR_API FAvaOutlinerModifier : public FAvaOutlinerObj
 {
 public:
 	UE_AVA_INHERITS_WITH_SUPER(FAvaOutlinerModifier, FAvaOutlinerObject);
-	
+
 	FAvaOutlinerModifier(IAvaOutliner& InOutliner, UActorModifierCoreBase* InObject);
-	
+
 	UActorModifierCoreBase* GetModifier() const { return Modifier.Get(); }
-	
+
 	//~ Begin IAvaOutlinerItem
 	virtual void Select(FAvaOutlinerScopedSelection& InSelection) const override;
 	virtual FText GetDisplayName() const override;
@@ -26,8 +26,10 @@ public:
 	virtual bool ShowVisibility(EAvaOutlinerVisibilityType InVisibilityType) const override;
 	virtual bool GetVisibility(EAvaOutlinerVisibilityType InVisibilityType) const override;
 	virtual void OnVisibilityChanged(EAvaOutlinerVisibilityType InVisibilityType, bool bInNewVisibility) override;
+	virtual bool CanDelete() const override;
+	virtual bool Delete() override;
 	//~ End IAvaOutlinerItem
-	
+
 protected:
 	//~ Begin FAvaOutlinerObjectItem
 	virtual void SetObject_Impl(UObject* InObject) override;
diff --git a/Engine/Plugins/Experimental/Avalanche/Source/AvalancheModifiersEditor/Private/Outliner/AvaOutlinerModifierDropHandler.cpp b/Engine/Plugins/Experimental/Avalanche/Source/AvalancheModifiersEditor/Private/Outliner/AvaOutlinerModifierDropHandler.cpp
--- a/Engine/Plugins/Experimental/Avalanche/Source/AvalancheModifiersEditor/Private/Outliner/AvaOutlinerModifierDropHandler.cpp
+++ b/Engine/Plugins/Experimental/Avalanche/Source/AvalancheModifiersEditor/Private/Outliner/AvaOutlinerModifierDropHandler.cpp
@@ -1,6 +1,8 @@
 ﻿// Copyright Epic Games, Inc. All Rights Reserved.
 
 #include "AvaOutlinerModifierDropHandler.h"
+
+#include "Framework/Application/SlateApplication.h"
 #include "Framework/Notifications/NotificationManager.h"
 #include "Item/AvaOutlinerActor.h"
 #include "Modifiers/ActorModifierCoreStack.h"
@@ -171,7 +173,7 @@ bool FAvaOutlinerModifierDropHandler::DropModifiersInActor(AActor* InActor, EIte
 		return false;
 	}
 
-	UE_LOG(LogAvaOutlinerModifierDropHandler, Log, TEXT("Dropping %i modifier(s) on actor %s")
+	UE_LOG(LogAvaOutlinerModifierDropHandler, Log, TEXT("Dropping %i modifier(s) on actor %s (Clone)")
 		, CloneModifiers.Num()
 		, *InActor->GetActorNameOrLabel());
 
@@ -186,16 +188,24 @@ bool FAvaOutlinerModifierDropHandler::DropModifiersInActor(AActor* InActor, EIte
 
 		bSuccess = ModifierSubsystem->CloneModifiers(CloneModifiers, Stack, CloneOp).Num() == CloneModifiers.Num();
 
-		if (!bSuccess)
+		// When ALT is not pressed : perform a move operation between actors (copy + delete)
+		// When ALT is pressed : perform a clone operation between actors
+		if (bSuccess && !FSlateApplication::Get().GetModifierKeys().IsAltDown())
 		{
-			UE_LOG(LogAvaOutlinerModifierDropHandler, Warning, TEXT("Clone %i modifier(s) on actor %s failed : %s"),
-				CloneModifiers.Num(),
-				*InActor->GetActorNameOrLabel(),
-				*FailReason.ToString());
+			FActorModifierCoreStackRemoveOp RemoveOp;
+			RemoveOp.bShouldTransact = true;
+			RemoveOp.FailReason = &FailReason;
+
+			bSuccess = ModifierSubsystem->RemoveModifiers(TSet<UActorModifierCoreBase*>(CloneModifiers), RemoveOp);
 		}
 
 		if (!FailReason.IsEmpty())
 		{
+			UE_LOG(LogAvaOutlinerModifierDropHandler, Warning, TEXT("Dropping %i modifier(s) on actor %s (Clone) failed with reason : %s"),
+				CloneModifiers.Num(),
+				*InActor->GetActorNameOrLabel(),
+				*FailReason.ToString());
+
 			FNotificationInfo NotificationInfo(FailReason);
 			NotificationInfo.ExpireDuration = 3.f;
 			NotificationInfo.bFireAndForget = true;
@@ -258,10 +268,29 @@ bool FAvaOutlinerModifierDropHandler::DropModifiersInModifier(UActorModifierCore
 		CloneOp.ClonePositionContext = InTargetModifier;
 
 		bSuccess = ModifierSubsystem->CloneModifiers(CloneModifiers, InTargetModifier->GetModifierStack(), CloneOp).Num() == CloneModifiers.Num();
+
+		// When ALT is not pressed : perform a move operation between actors (copy + delete)
+		// When ALT is pressed : perform a clone operation between actors
+		if (bSuccess && !FSlateApplication::Get().GetModifierKeys().IsAltDown())
+		{
+			FActorModifierCoreStackRemoveOp RemoveOp;
+			RemoveOp.bShouldTransact = true;
+			RemoveOp.FailReason = &FailReason;
+
+			bSuccess = ModifierSubsystem->RemoveModifiers(TSet<UActorModifierCoreBase*>(CloneModifiers), RemoveOp);
+		}
 	}
 
 	if (!FailReason.IsEmpty())
 	{
+		UE_LOG(LogAvaOutlinerModifierDropHandler, Warning, TEXT("Dropping %i modifier(s) %s modifier %s on actor %s (%s) failed with reason : %s")
+			, MoveModifiers.IsEmpty() ? CloneModifiers.Num() : MoveModifiers.Num()
+			, Position == EActorModifierCoreStackPosition::After ? TEXT("after") : TEXT("before")
+			, *InTargetModifier->GetModifierName().ToString()
+			, *TargetActor->GetActorNameOrLabel()
+			, MoveModifiers.IsEmpty() ? TEXT("Clone") : TEXT("Move")
+			, *FailReason.ToString());
+
 		FNotificationInfo NotificationInfo(FailReason);
 		NotificationInfo.ExpireDuration = 3.f;
 		NotificationInfo.bFireAndForget = true;
diff --git a/Engine/Plugins/Experimental/Avalanche/Source/AvalancheModifiersEditor/Private/Outliner/AvaOutlinerModifierProxy.cpp b/Engine/Plugins/Experimental/Avalanche/Source/AvalancheModifiersEditor/Private/Outliner/AvaOutlinerModifierProxy.cpp
--- a/Engine/Plugins/Experimental/Avalanche/Source/AvalancheModifiersEditor/Private/Outliner/AvaOutlinerModifierProxy.cpp
+++ b/Engine/Plugins/Experimental/Avalanche/Source/AvalancheModifiersEditor/Private/Outliner/AvaOutlinerModifierProxy.cpp
@@ -89,8 +89,30 @@ FText FAvaOutlinerModifierProxy::GetIconTooltipText() const
 	return LOCTEXT("Tooltip", "Shows all the Modifiers found in the Root Stack of an Actor");
 }
 
-void FAvaOutlinerModifierProxy::GetProxiedItems(const TSharedRef<IAvaOutlinerItem>& InParent
-	, TArray<FAvaOutlinerItemPtr>& OutChildren, bool bInRecursive)
+bool FAvaOutlinerModifierProxy::CanDelete() const
+{
+	return IsValid(GetModifierStack());
+}
+
+bool FAvaOutlinerModifierProxy::Delete()
+{
+	UActorModifierCoreStack* ModifierStack = GetModifierStack();
+	const UActorModifierCoreSubsystem* ModifierSubsystem = UActorModifierCoreSubsystem::Get();
+
+	if (IsValid(ModifierStack) && ModifierSubsystem)
+	{
+		FText FailReason;
+		FActorModifierCoreStackRemoveOp RemoveOp;
+		RemoveOp.FailReason = &FailReason;
+		RemoveOp.bShouldTransact = false;
+
+		return ModifierSubsystem->RemoveModifiers(TSet<UActorModifierCoreBase*>{ModifierStack->GetModifiers()}, RemoveOp);
+	}
+
+	return false;
+}
+
+void FAvaOutlinerModifierProxy::GetProxiedItems(const TSharedRef<IAvaOutlinerItem>& InParent, TArray<FAvaOutlinerItemPtr>& OutChildren, bool bInRecursive)
 {
 	if (const UActorModifierCoreStack* const ModifierStack = GetModifierStack())
 	{
diff --git a/Engine/Plugins/Experimental/Avalanche/Source/AvalancheModifiersEditor/Private/Outliner/AvaOutlinerModifierProxy.h b/Engine/Plugins/Experimental/Avalanche/Source/AvalancheModifiersEditor/Private/Outliner/AvaOutlinerModifierProxy.h
--- a/Engine/Plugins/Experimental/Avalanche/Source/AvalancheModifiersEditor/Private/Outliner/AvaOutlinerModifierProxy.h
+++ b/Engine/Plugins/Experimental/Avalanche/Source/AvalancheModifiersEditor/Private/Outliner/AvaOutlinerModifierProxy.h
@@ -28,6 +28,8 @@ public:
 	virtual FText GetDisplayName() const override;
 	virtual FSlateIcon GetIcon() const override;
 	virtual FText GetIconTooltipText() const override;
+	virtual bool CanDelete() const override;
+	virtual bool Delete() override;
 	//~ End IAvaOutlinerItem
 
 	//~ Begin FAvaOutlinerItemProxy
diff --git a/Engine/Plugins/Experimental/Avalanche/Source/AvalancheOutliner/Private/AvaOutliner.cpp b/Engine/Plugins/Experimental/Avalanche/Source/AvalancheOutliner/Private/AvaOutliner.cpp
--- a/Engine/Plugins/Experimental/Avalanche/Source/AvalancheOutliner/Private/AvaOutliner.cpp
+++ b/Engine/Plugins/Experimental/Avalanche/Source/AvalancheOutliner/Private/AvaOutliner.cpp
@@ -832,6 +832,55 @@ void FAvaOutliner::DuplicateItems(TArray<FAvaOutlinerItemPtr> InItems
 	OutlinerProvider.OutlinerDuplicateActors(TemplateActors);
 }
 
+void FAvaOutliner::DeleteItems(TArray<FAvaOutlinerItemPtr> InItems)
+{
+	SortItems(InItems);
+
+	TArray<AActor*> DeleteActors;
+	DeleteActors.Reserve(InItems.Num());
+
+	for (TArray<FAvaOutlinerItemPtr>::TIterator It(InItems); It; ++It)
+	{
+		FAvaOutlinerItemPtr Item = *It;
+
+		if (!Item.IsValid() || !Item->CanDelete())
+		{
+			It.RemoveCurrent();
+		}
+
+		if (const FAvaOutlinerActor* const ActorItem = Item->CastTo<FAvaOutlinerActor>())
+		{
+			DeleteActors.Add(ActorItem->GetActor());
+		}
+	}
+
+	if (InItems.IsEmpty())
+	{
+		return;
+	}
+
+	FScopedTransaction DeleteTransaction(LOCTEXT("OutlinerItemDeleteAction", "Outliner Delete Item(s)"), !GIsTransacting);
+
+	if (!DeleteActors.IsEmpty())
+	{
+		OutlinerProvider.OutlinerDeleteActors(DeleteActors);
+	}
+
+	bool bRequestRefresh = false;
+	for (FAvaOutlinerItemPtr& Item : InItems)
+	{
+		if (Item && Item->Delete())
+		{
+			bRequestRefresh = true;
+		}
+	}
+
+	if (bRequestRefresh)
+	{
+		RequestRefresh();
+	}
+}
+
 void FAvaOutliner::UnregisterOutlinerView(int32 InOutlinerViewId)
 {
 	OutlinerViews.Remove(InOutlinerViewId);
diff --git a/Engine/Plugins/Experimental/Avalanche/Source/AvalancheOutliner/Private/AvaOutliner.h b/Engine/Plugins/Experimental/Avalanche/Source/AvalancheOutliner/Private/AvaOutliner.h
--- a/Engine/Plugins/Experimental/Avalanche/Source/AvalancheOutliner/Private/AvaOutliner.h
+++ b/Engine/Plugins/Experimental/Avalanche/Source/AvalancheOutliner/Private/AvaOutliner.h
@@ -128,7 +128,12 @@ public:
 	void DuplicateItems(TArray<FAvaOutlinerItemPtr> InItems
 		, FAvaOutlinerItemPtr InRelativeItem
 		, TOptional<EItemDropZone> InRelativeDropZone);
-	
+
+	/**
+	 * Delete a set of items in the outliner by calling their custom delete handler
+	 */
+	void DeleteItems(TArray<FAvaOutlinerItemPtr> InItems);
+
 	/** Unregisters the Outliner View bound to the given id */
 	void UnregisterOutlinerView(int32 InOutlinerViewId);
 
diff --git a/Engine/Plugins/Experimental/Avalanche/Source/AvalancheOutliner/Private/AvaOutlinerView.cpp b/Engine/Plugins/Experimental/Avalanche/Source/AvalancheOutliner/Private/AvaOutlinerView.cpp
--- a/Engine/Plugins/Experimental/Avalanche/Source/AvalancheOutliner/Private/AvaOutlinerView.cpp
+++ b/Engine/Plugins/Experimental/Avalanche/Source/AvalancheOutliner/Private/AvaOutlinerView.cpp
@@ -272,6 +272,10 @@ void FAvaOutlinerView::BindCommands(const TSharedPtr<FUICommandList>& InBaseComm
 		, FExecuteAction::CreateSP(this, &FAvaOutlinerView::RenameSelected)
 		, FCanExecuteAction::CreateSP(this, &FAvaOutlinerView::CanRenameSelected));
 
+	ViewCommandList->MapAction(GenericCommands.Delete
+		, FExecuteAction::CreateSP(this, &FAvaOutlinerView::DeleteSelected)
+		, FCanExecuteAction::CreateSP(this, &FAvaOutlinerView::CanDeleteSelected));
+
 	ViewCommandList->MapAction(GenericCommands.Duplicate
 		, FExecuteAction::CreateSP(this, &FAvaOutlinerView::DuplicateSelected)
 		, FCanExecuteAction::CreateSP(this, &FAvaOutlinerView::CanDuplicateSelected));
@@ -1411,6 +1415,39 @@ bool FAvaOutlinerView::CanRenameSelected() const
 	return GetViewSelectedItemCount() > 0;
 }
 
+void FAvaOutlinerView::DeleteSelected()
+{
+	if (TSharedPtr<FAvaOutliner> Outliner = OutlinerWeak.Pin())
+	{
+		TArray<FAvaOutlinerItemPtr> Items = GetViewSelectedItems();
+
+		Items.RemoveAll([](const FAvaOutlinerItemPtr& InItem)
+		{
+			return !InItem.IsValid() || !InItem->CanDelete();
+		});
+
+		if (Items.IsEmpty())
+		{
+			return;
+		}
+
+		Outliner->DeleteItems(Items);
+	}
+}
+
+bool FAvaOutlinerView::CanDeleteSelected() const
+{
+	for (const FAvaOutlinerItemPtr& Item : GetViewSelectedItems())
+	{
+		if (Item && Item->CanDelete())
+		{
+			return true;
+		}
+	}
+
+	return false;
+}
+
 void FAvaOutlinerView::DuplicateSelected()
 {
 	if (TSharedPtr<FAvaOutliner> Outliner = OutlinerWeak.Pin())
diff --git a/Engine/Plugins/Experimental/Avalanche/Source/AvalancheOutliner/Private/AvaOutlinerView.h b/Engine/Plugins/Experimental/Avalanche/Source/AvalancheOutliner/Private/AvaOutlinerView.h
--- a/Engine/Plugins/Experimental/Avalanche/Source/AvalancheOutliner/Private/AvaOutlinerView.h
+++ b/Engine/Plugins/Experimental/Avalanche/Source/AvalancheOutliner/Private/AvaOutlinerView.h
@@ -365,6 +365,9 @@ public:
 	void OnItemRenameAction(EAvaOutlinerRenameAction InRenameAction, const TSharedPtr<FAvaOutlinerView>& InOutlinerView);
 	bool CanRenameSelected() const;
 
+	void DeleteSelected();
+	bool CanDeleteSelected() const;
+
 	void DuplicateSelected();
 	bool CanDuplicateSelected() const;
 
diff --git a/Engine/Plugins/Experimental/Avalanche/Source/AvalancheOutliner/Private/Item/AvaOutlinerItem.cpp b/Engine/Plugins/Experimental/Avalanche/Source/AvalancheOutliner/Private/Item/AvaOutlinerItem.cpp
--- a/Engine/Plugins/Experimental/Avalanche/Source/AvalancheOutliner/Private/Item/AvaOutlinerItem.cpp
+++ b/Engine/Plugins/Experimental/Avalanche/Source/AvalancheOutliner/Private/Item/AvaOutlinerItem.cpp
@@ -225,6 +225,12 @@ bool FAvaOutlinerItem::Rename(const FString& InName)
 	return false;
 }
 
+bool FAvaOutlinerItem::Delete()
+{
+	Outliner.GetProvider().NotifyOutlinerItemDeleted(SharedThis(this));
+	return false;
+}
+
 void FAvaOutlinerItem::SetLocked(bool bInIsLocked)
 {
 	Outliner.GetProvider().NotifyOutlinerItemLockChanged(SharedThis(this));
diff --git a/Engine/Plugins/Experimental/Avalanche/Source/AvalancheOutliner/Private/Tests/AvaOutlinerTestUtils.cpp b/Engine/Plugins/Experimental/Avalanche/Source/AvalancheOutliner/Private/Tests/AvaOutlinerTestUtils.cpp
--- a/Engine/Plugins/Experimental/Avalanche/Source/AvalancheOutliner/Private/Tests/AvaOutlinerTestUtils.cpp
+++ b/Engine/Plugins/Experimental/Avalanche/Source/AvalancheOutliner/Private/Tests/AvaOutlinerTestUtils.cpp
@@ -11,8 +11,10 @@
 #include "Engine/SkyLight.h"
 #include "Engine/StaticMeshActor.h"
 #include "Engine/TriggerVolume.h"
+#include "Editor/UnrealEdEngine.h"
 #include "PackageTools.h"
 #include "Selection.h"
+#include "UnrealEdGlobals.h"
 
 void UE::AvaOutliner::Private::FAvaOutlinerEditorModeTools::Init(UWorld* InWorld)
 {
@@ -388,9 +390,39 @@ void UE::AvaOutliner::Private::FAvaOutlinerProviderTest::OutlinerDuplicateActors
 	Outliner->OnActorsDuplicated(DuplicateActorMap);
 }
 
-TOptional<EItemDropZone> UE::AvaOutliner::Private::FAvaOutlinerProviderTest::OnOutlinerItemCanAcceptDrop(const FDragDropEvent& DragDropEvent
-	, EItemDropZone DropZone
-	, FAvaOutlinerItemPtr TargetItem) const
+void UE::AvaOutliner::Private::FAvaOutlinerProviderTest::OutlinerDeleteActors(const TArray<AActor*>& InDeleteActors)
+{
+	if (InDeleteActors.IsEmpty())
+	{
+		return;
+	}
+
+	auto DestroyActor = [](AActor* InActor)
+	{
+		if (UWorld* ActorWorld = InActor->GetWorld())
+		{
+			if (GIsEditor && GUnrealEd)
+			{
+				GUnrealEd->DeleteActors({InActor}, ActorWorld, GUnrealEd->GetSelectedActors()->GetElementSelectionSet());
+			}
+			else
+			{
+				ActorWorld->DestroyActor(InActor);
+			}
+		}
+	};
+
+	for (AActor* Actor : InDeleteActors)
+	{
+		if (Actor)
+		{
+			Actor->Modify();
+			DestroyActor(Actor);
+		}
+	}
+}
+
+TOptional<EItemDropZone> UE::AvaOutliner::Private::FAvaOutlinerProviderTest::OnOutlinerItemCanAcceptDrop(const FDragDropEvent& DragDropEvent, EItemDropZone DropZone, FAvaOutlinerItemPtr TargetItem) const
 {
 	return TOptional<EItemDropZone>();
 }
diff --git a/Engine/Plugins/Experimental/Avalanche/Source/AvalancheOutliner/Private/Tests/AvaOutlinerTestUtils.h b/Engine/Plugins/Experimental/Avalanche/Source/AvalancheOutliner/Private/Tests/AvaOutlinerTestUtils.h
--- a/Engine/Plugins/Experimental/Avalanche/Source/AvalancheOutliner/Private/Tests/AvaOutlinerTestUtils.h
+++ b/Engine/Plugins/Experimental/Avalanche/Source/AvalancheOutliner/Private/Tests/AvaOutlinerTestUtils.h
@@ -82,6 +82,7 @@ namespace UE::AvaOutliner::Private
 		virtual bool ShouldLockOutliner() const override { return false; }
 		virtual bool ShouldHideItem(const FAvaOutlinerItemPtr& Item) const override { return false; }
 		virtual void OutlinerDuplicateActors(const TArray<AActor*>& InTemplateActors) override;
+		virtual void OutlinerDeleteActors(const TArray<AActor*>& InDeleteActors) override;
 		virtual FEditorModeTools* GetOutlinerModeTools() const override { return &ModeTools.Get(); }
 		virtual FAvaSceneTree* GetSceneTree() const override { return nullptr; }
 		virtual UWorld* GetOutlinerWorld() const override { return World; }
diff --git a/Engine/Plugins/Experimental/Avalanche/Source/AvalancheOutliner/Public/IAvaOutlinerProvider.h b/Engine/Plugins/Experimental/Avalanche/Source/AvalancheOutliner/Public/IAvaOutlinerProvider.h
--- a/Engine/Plugins/Experimental/Avalanche/Source/AvalancheOutliner/Public/IAvaOutlinerProvider.h
+++ b/Engine/Plugins/Experimental/Avalanche/Source/AvalancheOutliner/Public/IAvaOutlinerProvider.h
@@ -45,6 +45,11 @@ public:
 	 */
 	virtual void OutlinerDuplicateActors(const TArray<AActor*>& InTemplateActors) = 0;
 
+	/**
+	 * Called when the Outliner requests the deletion of the given actors	
+	 */
+	virtual void OutlinerDeleteActors(const TArray<AActor*>& InDeleteActors) = 0;
+
 	/** Returns the Mode Tools used (e.g. for Selections) */
 	virtual FEditorModeTools* GetOutlinerModeTools() const = 0;
 
@@ -79,6 +84,9 @@ public:
 	/** Called when an Item has been Renamed */
 	virtual void NotifyOutlinerItemRenamed(const FAvaOutlinerItemPtr& InItem) {}
 
+	/** Called when an Item has been deleted */
+	virtual void NotifyOutlinerItemDeleted(const FAvaOutlinerItemPtr& InItem) {}
+
 	/** Called when an Item's Lock has Changed */
 	virtual void NotifyOutlinerItemLockChanged(const FAvaOutlinerItemPtr& InItem) {}
 
diff --git a/Engine/Plugins/Experimental/Avalanche/Source/AvalancheOutliner/Public/Item/AvaOutlinerActor.h b/Engine/Plugins/Experimental/Avalanche/Source/AvalancheOutliner/Public/Item/AvaOutlinerActor.h
--- a/Engine/Plugins/Experimental/Avalanche/Source/AvalancheOutliner/Public/Item/AvaOutlinerActor.h
+++ b/Engine/Plugins/Experimental/Avalanche/Source/AvalancheOutliner/Public/Item/AvaOutlinerActor.h
@@ -31,6 +31,8 @@ public:
 	virtual void OnVisibilityChanged(EAvaOutlinerVisibilityType VisibilityType, bool bNewVisibility) override;
 	virtual bool CanRename() const override { return Actor.IsValid(); }
 	virtual bool Rename(const FString& InName) override;
+	virtual bool CanDelete() const override { return Actor.IsValid(); }
+	virtual bool Delete() override { return !Actor.IsValid(); }
 	virtual bool CanLock() const override { return true; }
 	virtual void SetLocked(bool bInIsLocked) override;
 	virtual bool IsLocked() const override;
diff --git a/Engine/Plugins/Experimental/Avalanche/Source/AvalancheOutliner/Public/Item/AvaOutlinerItem.h b/Engine/Plugins/Experimental/Avalanche/Source/AvalancheOutliner/Public/Item/AvaOutlinerItem.h
--- a/Engine/Plugins/Experimental/Avalanche/Source/AvalancheOutliner/Public/Item/AvaOutlinerItem.h
+++ b/Engine/Plugins/Experimental/Avalanche/Source/AvalancheOutliner/Public/Item/AvaOutlinerItem.h
@@ -46,6 +46,8 @@ public:
 	virtual bool CanAutoExpand() const override { return true; }
 	virtual bool CanRename() const override { return false; }
 	virtual bool Rename(const FString& InName) override;
+	virtual bool CanDelete() const override { return false; }
+	virtual bool Delete() override;
 	virtual bool CanLock() const override { return false; }
 	virtual void SetLocked(bool bInIsLocked) override;
 	virtual bool IsLocked() const override { return false; }
diff --git a/Engine/Plugins/Experimental/Avalanche/Source/AvalancheOutliner/Public/Item/IAvaOutlinerItem.h b/Engine/Plugins/Experimental/Avalanche/Source/AvalancheOutliner/Public/Item/IAvaOutlinerItem.h
--- a/Engine/Plugins/Experimental/Avalanche/Source/AvalancheOutliner/Public/Item/IAvaOutlinerItem.h
+++ b/Engine/Plugins/Experimental/Avalanche/Source/AvalancheOutliner/Public/Item/IAvaOutlinerItem.h
@@ -265,6 +265,12 @@ public:
 	/** The implementation to rename the item (e.g. for AActor it will be the Actor Label that is changed) */
 	virtual bool Rename(const FString& InName) = 0;
 
+	/** Whether this Item can be deleted or not */
+	virtual bool CanDelete() const = 0;
+
+	/** The implementation to delete the item */
+	virtual bool Delete() = 0;
+
 	/** Whether an Item can be Locked in any kind. E.g. Actors can be location-locked */
 	virtual bool CanLock() const = 0;
 
diff --git a/Engine/Plugins/Experimental/Avalanche/Source/AvalanchePropertyAnimatorEditor/Private/Outliner/AvaPropertyAnimatorEditorOutliner.cpp b/Engine/Plugins/Experimental/Avalanche/Source/AvalanchePropertyAnimatorEditor/Private/Outliner/AvaPropertyAnimatorEditorOutliner.cpp
--- a/Engine/Plugins/Experimental/Avalanche/Source/AvalanchePropertyAnimatorEditor/Private/Outliner/AvaPropertyAnimatorEditorOutliner.cpp
+++ b/Engine/Plugins/Experimental/Avalanche/Source/AvalanchePropertyAnimatorEditor/Private/Outliner/AvaPropertyAnimatorEditorOutliner.cpp
@@ -6,6 +6,7 @@
 #include "Components/PropertyAnimatorCoreComponent.h"
 #include "Selection/AvaOutlinerScopedSelection.h"
 #include "Styling/SlateIconFinder.h"
+#include "Subsystems/PropertyAnimatorCoreSubsystem.h"
 
 FAvaPropertyAnimatorEditorOutliner::FAvaPropertyAnimatorEditorOutliner(IAvaOutliner& InOutliner, UPropertyAnimatorCoreBase* InAnimator)
 	: FAvaOutlinerObject(InOutliner, InAnimator)
@@ -68,6 +69,21 @@ void FAvaPropertyAnimatorEditorOutliner::OnVisibilityChanged(EAvaOutlinerVisibil
 	}
 }
 
+bool FAvaPropertyAnimatorEditorOutliner::CanDelete() const
+{
+	return PropertyAnimator.IsValid();
+}
+
+bool FAvaPropertyAnimatorEditorOutliner::Delete()
+{
+	if (const UPropertyAnimatorCoreSubsystem* AnimatorSubsystem = UPropertyAnimatorCoreSubsystem::Get())
+	{
+		return AnimatorSubsystem->RemoveAnimator(PropertyAnimator.Get(), /** Transact */false);
+	}
+
+	return false;
+}
+
 void FAvaPropertyAnimatorEditorOutliner::SetObject_Impl(UObject* InObject)
 {
 	FAvaOutlinerObject::SetObject_Impl(InObject);
diff --git a/Engine/Plugins/Experimental/Avalanche/Source/AvalanchePropertyAnimatorEditor/Private/Outliner/AvaPropertyAnimatorEditorOutliner.h b/Engine/Plugins/Experimental/Avalanche/Source/AvalanchePropertyAnimatorEditor/Private/Outliner/AvaPropertyAnimatorEditorOutliner.h
--- a/Engine/Plugins/Experimental/Avalanche/Source/AvalanchePropertyAnimatorEditor/Private/Outliner/AvaPropertyAnimatorEditorOutliner.h
+++ b/Engine/Plugins/Experimental/Avalanche/Source/AvalanchePropertyAnimatorEditor/Private/Outliner/AvaPropertyAnimatorEditorOutliner.h
@@ -29,6 +29,8 @@ public:
 	virtual bool ShowVisibility(EAvaOutlinerVisibilityType InVisibilityType) const override;
 	virtual bool GetVisibility(EAvaOutlinerVisibilityType InVisibilityType) const override;
 	virtual void OnVisibilityChanged(EAvaOutlinerVisibilityType InVisibilityType, bool bInNewVisibility) override;
+	virtual bool CanDelete() const override;
+	virtual bool Delete() override;
 	//~ End IAvaOutlinerItem
 
 protected:
diff --git a/Engine/Plugins/Experimental/Avalanche/Source/AvalanchePropertyAnimatorEditor/Private/Outliner/AvaPropertyAnimatorEditorOutlinerDropHandler.cpp b/Engine/Plugins/Experimental/Avalanche/Source/AvalanchePropertyAnimatorEditor/Private/Outliner/AvaPropertyAnimatorEditorOutlinerDropHandler.cpp
--- a/Engine/Plugins/Experimental/Avalanche/Source/AvalanchePropertyAnimatorEditor/Private/Outliner/AvaPropertyAnimatorEditorOutlinerDropHandler.cpp
+++ b/Engine/Plugins/Experimental/Avalanche/Source/AvalanchePropertyAnimatorEditor/Private/Outliner/AvaPropertyAnimatorEditorOutlinerDropHandler.cpp
@@ -3,6 +3,7 @@
 #include "Outliner/AvaPropertyAnimatorEditorOutlinerDropHandler.h"
 
 #include "Components/PropertyAnimatorCoreComponent.h"
+#include "Framework/Application/SlateApplication.h"
 #include "Framework/Notifications/NotificationManager.h"
 #include "Item/AvaOutlinerActor.h"
 #include "Outliner/AvaPropertyAnimatorEditorOutliner.h"
@@ -193,12 +194,29 @@ bool FAvaPropertyAnimatorEditorOutlinerDropHandler::DropAnimatorsOnActor(AActor*
 
 	if (NewAnimators.Num() != CloneAnimators.Num())
 	{
+		UE_LOG(LogAvaPropertyAnimatorEditorOutlinerDropHandler, Warning, TEXT("%s : Could not clone all %i animators to target actor"), *InActor->GetActorNameOrLabel(), CloneAnimators.Num());
+
 		FNotificationInfo NotificationInfo(LOCTEXT("CloneAnimatorsFail", "An issue occured while cloning animators on an actor"));
 		NotificationInfo.ExpireDuration = 3.f;
 		NotificationInfo.bFireAndForget = true;
 		FSlateNotificationManager::Get().AddNotification(NotificationInfo);
 	}
 
+	// When ALT is pressed : copy animators on target actor
+	// When ALT is not pressed : move animators on target actor (copy + delete)
+	if (!FSlateApplication::Get().GetModifierKeys().IsAltDown())
+	{
+		if (!AnimatorSubsystem->RemoveAnimators(CloneAnimators, /** Transact */true))
+		{
+			UE_LOG(LogAvaPropertyAnimatorEditorOutlinerDropHandler, Warning, TEXT("Could not remove the %i cloned animators on source actor"), CloneAnimators.Num());
+
+			FNotificationInfo NotificationInfo(LOCTEXT("RemoveAnimatorsFail", "An issue occured while removing animators on an actor"));
+			NotificationInfo.ExpireDuration = 3.f;
+			NotificationInfo.bFireAndForget = true;
+			FSlateNotificationManager::Get().AddNotification(NotificationInfo);
+		}
+	}
+
 	return NewAnimators.Num() > 0;
 }
 
diff --git a/Engine/Plugins/Experimental/Avalanche/Source/AvalanchePropertyAnimatorEditor/Private/Outliner/AvaPropertyAnimatorEditorOutlinerProxy.cpp b/Engine/Plugins/Experimental/Avalanche/Source/AvalanchePropertyAnimatorEditor/Private/Outliner/AvaPropertyAnimatorEditorOutlinerProxy.cpp
--- a/Engine/Plugins/Experimental/Avalanche/Source/AvalanchePropertyAnimatorEditor/Private/Outliner/AvaPropertyAnimatorEditorOutlinerProxy.cpp
+++ b/Engine/Plugins/Experimental/Avalanche/Source/AvalanchePropertyAnimatorEditor/Private/Outliner/AvaPropertyAnimatorEditorOutlinerProxy.cpp
@@ -9,6 +9,7 @@
 #include "Outliner/AvaPropertyAnimatorEditorOutliner.h"
 #include "Selection/AvaOutlinerScopedSelection.h"
 #include "Styling/SlateIconFinder.h"
+#include "Subsystems/PropertyAnimatorCoreSubsystem.h"
 
 #define LOCTEXT_NAMESPACE "AvaPropertyAnimatorEditorOutlinerProxy"
 
@@ -79,6 +80,24 @@ FText FAvaPropertyAnimatorEditorOutlinerProxy::GetIconTooltipText() const
 	return LOCTEXT("Tooltip", "Shows all the animators found in the property animator component of an actor");
 }
 
+bool FAvaPropertyAnimatorEditorOutlinerProxy::CanDelete() const
+{
+	return IsValid(GetPropertyAnimatorComponent());
+}
+
+bool FAvaPropertyAnimatorEditorOutlinerProxy::Delete()
+{
+	const UPropertyAnimatorCoreComponent* AnimatorComponent = GetPropertyAnimatorComponent();
+	const UPropertyAnimatorCoreSubsystem* AnimatorSubsystem = UPropertyAnimatorCoreSubsystem::Get();
+
+	if (IsValid(AnimatorComponent) && AnimatorSubsystem)
+	{
+		return AnimatorSubsystem->RemoveAnimators(TSet<UPropertyAnimatorCoreBase*>{AnimatorComponent->GetAnimators()}, /** Transact */false);
+	}
+
+	return false;
+}
+
 void FAvaPropertyAnimatorEditorOutlinerProxy::GetProxiedItems(const TSharedRef<IAvaOutlinerItem>& InParent
 	, TArray<FAvaOutlinerItemPtr>& OutChildren, bool bInRecursive)
 {
diff --git a/Engine/Plugins/Experimental/Avalanche/Source/AvalanchePropertyAnimatorEditor/Private/Outliner/AvaPropertyAnimatorEditorOutlinerProxy.h b/Engine/Plugins/Experimental/Avalanche/Source/AvalanchePropertyAnimatorEditor/Private/Outliner/AvaPropertyAnimatorEditorOutlinerProxy.h
--- a/Engine/Plugins/Experimental/Avalanche/Source/AvalanchePropertyAnimatorEditor/Private/Outliner/AvaPropertyAnimatorEditorOutlinerProxy.h
+++ b/Engine/Plugins/Experimental/Avalanche/Source/AvalanchePropertyAnimatorEditor/Private/Outliner/AvaPropertyAnimatorEditorOutlinerProxy.h
@@ -25,6 +25,8 @@ public:
 	virtual FText GetDisplayName() const override;
 	virtual FSlateIcon GetIcon() const override;
 	virtual FText GetIconTooltipText() const override;
+	virtual bool CanDelete() const override;
+	virtual bool Delete() override;
 	//~ End IAvaOutlinerItem
 
 	//~ Begin FAvaOutlinerItemProxy
