Deprecating AActor::FindComponentByInterface<UMyInterface> in favor of AActor::FindComponentByInterface<IMyInterface>.
The previous function would incorrectly return the UInterface type used for reflection and could lead to undefined behavior.

diff --git a/Engine/Plugins/Experimental/Avalanche/Source/AvalancheMask/Private/Mask2D/AvaMask2DBaseModifier.cpp b/Engine/Plugins/Experimental/Avalanche/Source/AvalancheMask/Private/Mask2D/AvaMask2DBaseModifier.cpp
--- a/Engine/Plugins/Experimental/Avalanche/Source/AvalancheMask/Private/Mask2D/AvaMask2DBaseModifier.cpp
+++ b/Engine/Plugins/Experimental/Avalanche/Source/AvalancheMask/Private/Mask2D/AvaMask2DBaseModifier.cpp
@@ -522,7 +522,7 @@ bool UAvaMask2DBaseModifier::TryResolveParentChannel()
 		const IGeometryMaskReadInterface* ReadComponent = nullptr;
 		while (ReadComponent == nullptr && Parent)
 		{
-			ReadComponent = Cast<IGeometryMaskReadInterface>(Parent->FindComponentByInterface<UGeometryMaskReadInterface>());
+			ReadComponent = Parent->FindComponentByInterface<IGeometryMaskReadInterface>();
 			Parent = Parent->GetAttachParentActor();
 		}
 
diff --git a/Engine/Plugins/Experimental/Avalanche/Source/AvalancheMaskEditor/Private/AvaMaskEditorMode.cpp b/Engine/Plugins/Experimental/Avalanche/Source/AvalancheMaskEditor/Private/AvaMaskEditorMode.cpp
--- a/Engine/Plugins/Experimental/Avalanche/Source/AvalancheMaskEditor/Private/AvaMaskEditorMode.cpp
+++ b/Engine/Plugins/Experimental/Avalanche/Source/AvalancheMaskEditor/Private/AvaMaskEditorMode.cpp
@@ -62,7 +62,7 @@ void UAvaMaskEditorMode::Enter()
 
 	for (AActor* Actor : TActorRange<AActor>(GetWorld()))
 	{
-		if (Actor->FindComponentByInterface<UGeometryMaskWriteInterface>())
+		if (Actor->FindComponentByInterface(UGeometryMaskWriteInterface::StaticClass()))
 		{
 			WeakMaskWriterActors.Add(Actor);
 			if (UAvaGizmoComponent* GizmoComponent = Actor->GetComponentByClass<UAvaGizmoComponent>())
@@ -287,7 +287,7 @@ void UAvaMaskEditorMode::OnSelectionChanged(const UTypedElementSelectionSet* InS
 					{
 						if (const AActor* WriterActor = WeakWriterActor.Get())
 						{
-							if (const IGeometryMaskWriteInterface* WriterObject = Cast<IGeometryMaskWriteInterface>(WriterActor->FindComponentByInterface<UGeometryMaskWriteInterface>()))
+							if (const IGeometryMaskWriteInterface* WriterObject = WriterActor->FindComponentByInterface<IGeometryMaskWriteInterface>())
 							{
 								if (UAvaGizmoComponent* GizmoComponent = WriterActor->GetComponentByClass<UAvaGizmoComponent>())
 								{
@@ -318,12 +318,12 @@ void UAvaMaskEditorMode::OnSelectionChanged(const UTypedElementSelectionSet* InS
 UGeometryMaskCanvas* UAvaMaskEditorMode::GetCanvasReferencedByActor(const AActor* InActor)
 {
 	FName CanvasName = NAME_None;
-	if (const IGeometryMaskWriteInterface* WriteComponent = Cast<IGeometryMaskWriteInterface>(InActor->FindComponentByInterface<UGeometryMaskWriteInterface>()))
+	if (const IGeometryMaskWriteInterface* WriteComponent = InActor->FindComponentByInterface<IGeometryMaskWriteInterface>())
 	{
 		const FGeometryMaskWriteParameters WriteComponentParameters = WriteComponent->GetParameters();
 		CanvasName = WriteComponentParameters.CanvasName;
 	}
-	else if (const IGeometryMaskReadInterface* ReadComponent = Cast<IGeometryMaskReadInterface>(InActor->FindComponentByInterface<UGeometryMaskReadInterface>()))
+	else if (const IGeometryMaskReadInterface* ReadComponent = InActor->FindComponentByInterface<IGeometryMaskReadInterface>())
 	{
 		const FGeometryMaskReadParameters ReadComponentParameters = ReadComponent->GetParameters();
 		CanvasName = ReadComponentParameters.CanvasName;
diff --git a/Engine/Source/Runtime/AIModule/Private/Navigation/PathFollowingComponent.cpp b/Engine/Source/Runtime/AIModule/Private/Navigation/PathFollowingComponent.cpp
--- a/Engine/Source/Runtime/AIModule/Private/Navigation/PathFollowingComponent.cpp
+++ b/Engine/Source/Runtime/AIModule/Private/Navigation/PathFollowingComponent.cpp
@@ -1483,7 +1483,7 @@ bool UPathFollowingComponent::UpdateMovementComponent(bool bForce)
 		{
 			MyPawn->OnActorHit.AddUniqueDynamic(this, &UPathFollowingComponent::OnActorBump);
 
-			SetNavMovementInterface(Cast<INavMovementInterface>(MyPawn->FindComponentByInterface<UNavMovementInterface>()));
+			SetNavMovementInterface(MyPawn->FindComponentByInterface<INavMovementInterface>());
 		}
 	}
 
diff --git a/Engine/Source/Runtime/Engine/Classes/GameFramework/Actor.h b/Engine/Source/Runtime/Engine/Classes/GameFramework/Actor.h
--- a/Engine/Source/Runtime/Engine/Classes/GameFramework/Actor.h
+++ b/Engine/Source/Runtime/Engine/Classes/GameFramework/Actor.h
@@ -7,6 +7,7 @@
 #include "UObject/UObjectBaseUtility.h"
 #include "UObject/Object.h"
 #include "InputCoreTypes.h"
+#include "Templates/Requires.h"
 #include "Templates/SubclassOf.h"
 #include "Engine/EngineTypes.h"
 #include "Engine/EngineBaseTypes.h"
@@ -3709,14 +3710,20 @@ public:
 	}
 
 	/** Templatized version of FindComponentByInterface that handles casting for you */
-	template<class T>
+	template<class T UE_REQUIRES(TPointerIsConvertibleFromTo<T, UInterface>::Value)>
+	UE_DEPRECATED(5.5, "This version incorrectly casts to the UInterface type used for reflection. Use FindComponentByInterface<IMyInterface>() instead")
 	T* FindComponentByInterface() const
 	{
-		static_assert(TPointerIsConvertibleFromTo<T, const UInterface>::Value, "'T' template parameter to FindComponentByInterface must be derived from UInterface");
-
 		return (T*)FindComponentByInterface(T::StaticClass());
 	}
 
+	/** Templatized version of FindComponentByInterface that handles casting for you */
+	template<class T UE_REQUIRES(TIsIInterface<T>::Value)>
+	T* FindComponentByInterface() const
+	{
+		return Cast<T>(FindComponentByInterface(T::UClassType::StaticClass()));
+	}
+
 private:
 
 	/**
