Fixed strange Clang-only error.

diff --git a/Engine/Plugins/Experimental/Avalanche/Source/AvalancheMedia/Private/Playable/Playables/AvaPlayableLevelStreaming.cpp b/Engine/Plugins/Experimental/Avalanche/Source/AvalancheMedia/Private/Playable/Playables/AvaPlayableLevelStreaming.cpp
--- a/Engine/Plugins/Experimental/Avalanche/Source/AvalancheMedia/Private/Playable/Playables/AvaPlayableLevelStreaming.cpp
+++ b/Engine/Plugins/Experimental/Avalanche/Source/AvalancheMedia/Private/Playable/Playables/AvaPlayableLevelStreaming.cpp
@@ -904,7 +904,7 @@ void UAvaPlayableLevelStreaming::GetOrLoadSubPlayable(const ULevelStreaming* InL
 	if (UAvaPlayableLevelStreaming* NewPlayable = CreateSubPlayable(PlayableGroup, SourceAssetPath))
 	{
 		// TODO: Propagate more stuff from InLevelStreaming. Needs to reach LoadLevel.
-		const FAvaSoftAssetPtr AssetPtr = { UWorld::StaticClass(), TSoftObjectPtr(SourceAssetPath)};
+		const FAvaSoftAssetPtr AssetPtr = { UWorld::StaticClass(), TSoftObjectPtr<UObject>(SourceAssetPath)};
 		if (NewPlayable->LoadAsset(AssetPtr, GetShouldBeVisible()))
 		{
 			AddSubPlayable(NewPlayable);
