Another Clang-only fix.

diff --git a/Engine/Plugins/Experimental/Avalanche/Source/AvalancheMedia/Private/Playable/Playables/AvaPlayableLevelStreaming.cpp b/Engine/Plugins/Experimental/Avalanche/Source/AvalancheMedia/Private/Playable/Playables/AvaPlayableLevelStreaming.cpp
--- a/Engine/Plugins/Experimental/Avalanche/Source/AvalancheMedia/Private/Playable/Playables/AvaPlayableLevelStreaming.cpp
+++ b/Engine/Plugins/Experimental/Avalanche/Source/AvalancheMedia/Private/Playable/Playables/AvaPlayableLevelStreaming.cpp
@@ -928,7 +928,7 @@ UAvaPlayableLevelStreaming* UAvaPlayableLevelStreaming::CreateSubPlayable(UAvaPl
 	const FPlayableCreationInfo PlayableCreationInfo =
 		{
 			InPlayableGroup->GetPlayableGroupManager(),
-			{ UWorld::StaticClass(), TSoftObjectPtr(InSourceAssetPath) },
+			{ UWorld::StaticClass(), TSoftObjectPtr<UObject>(InSourceAssetPath) },
 			FName(),
 			InPlayableGroup
 		};
