[Motion Design] Fix ViewState reference collection in UAvaGameViewportClient.

The FSceneViewState has a pool of Material Instance Dynamic (MID) use to blend the post process materials. It was getting garbage collected because the
UAvaGameViewportClient didn't add the ViewState(s) to referenced objects during GC's reference collection phase.

diff --git a/Engine/Plugins/Experimental/Avalanche/Source/Avalanche/Private/Framework/AvaGameViewportClient.cpp b/Engine/Plugins/Experimental/Avalanche/Source/Avalanche/Private/Framework/AvaGameViewportClient.cpp
--- a/Engine/Plugins/Experimental/Avalanche/Source/Avalanche/Private/Framework/AvaGameViewportClient.cpp
+++ b/Engine/Plugins/Experimental/Avalanche/Source/Avalanche/Private/Framework/AvaGameViewportClient.cpp
@@ -1,7 +1,7 @@
 // Copyright Epic Games, Inc. All Rights Reserved.
 
 #include "Framework/AvaGameViewportClient.h"
-#include "Camera/CameraActor.h"
+
 #include "Camera/CameraComponent.h"
 #include "Camera/CameraPhotography.h"
 #include "Components/LineBatchComponent.h"
@@ -20,6 +20,7 @@
 #include "IXRTrackingSystem.h"
 #include "LegacyScreenPercentageDriver.h"
 #include "Misc/EngineVersionComparison.h"
+#include "SceneManagement.h"
 #include "SceneViewExtension.h"
 #include "TextureResource.h"
 #include "UObject/Package.h"
@@ -290,6 +291,19 @@ bool UAvaGameViewportClient::IsStatEnabled(const FString& InName) const
 	return IAvaModule::Get().IsRuntimeStatEnabled(InName);
 }
 
+void UAvaGameViewportClient::AddReferencedObjects(UObject* InThis, FReferenceCollector& InCollector)
+{
+	UAvaGameViewportClient* This = CastChecked<UAvaGameViewportClient>(InThis);
+	for (FSceneViewStateReference& ViewState : This->ViewStates)
+	{
+		if (ViewState.GetReference())
+		{
+			ViewState.GetReference()->AddReferencedObjects(InCollector);
+		}
+	}
+	Super::AddReferencedObjects(InThis, InCollector);
+}
+
 void UAvaGameViewportClient::SetCameraCutThisFrame()
 {
 	bCameraCutThisFrame = true;
diff --git a/Engine/Plugins/Experimental/Avalanche/Source/Avalanche/Public/Framework/AvaGameViewportClient.h b/Engine/Plugins/Experimental/Avalanche/Source/Avalanche/Public/Framework/AvaGameViewportClient.h
--- a/Engine/Plugins/Experimental/Avalanche/Source/Avalanche/Public/Framework/AvaGameViewportClient.h
+++ b/Engine/Plugins/Experimental/Avalanche/Source/Avalanche/Public/Framework/AvaGameViewportClient.h
@@ -28,6 +28,8 @@ public:
 	virtual bool IsStatEnabled(const FString& InName) const override;
 	//~ End FViewportClient
 
+	static void AddReferencedObjects(UObject* InThis, FReferenceCollector& InCollector);
+	
 	void SetCameraCutThisFrame();
 
 	void SetRenderTarget(UTextureRenderTarget2D* InRenderTarget);
