MotionDesign : Shapes
- Fixed place motion design shape actors not working anymore within PlaceActors tab

diff --git a/Engine/Plugins/Experimental/Avalanche/Source/AvalancheInteractiveTools/Private/AvalancheInteractiveToolsModule.cpp b/Engine/Plugins/Experimental/Avalanche/Source/AvalancheInteractiveTools/Private/AvalancheInteractiveToolsModule.cpp
--- a/Engine/Plugins/Experimental/Avalanche/Source/AvalancheInteractiveTools/Private/AvalancheInteractiveToolsModule.cpp
+++ b/Engine/Plugins/Experimental/Avalanche/Source/AvalancheInteractiveTools/Private/AvalancheInteractiveToolsModule.cpp
@@ -260,7 +260,7 @@ void FAvalancheInteractiveToolsModule::OnPlacementCategoryRefreshed(FName InCate
 		{
 			PlaceableItem = MakeShared<FPlaceableItem>(
 				Tool.Factory.Get(),
-				FAssetData(Tool.Factory->NewActorClass),
+				FAssetData(Tool.Factory->NewActorClass.GetDefaultObject()),
 				Tool.Priority
 			);
 		}
@@ -268,7 +268,7 @@ void FAvalancheInteractiveToolsModule::OnPlacementCategoryRefreshed(FName InCate
 		{
 			PlaceableItem = MakeShared<FPlaceableItem>(
 				*Tool.FactoryClass.Get(),
-				FAssetData(Tool.FactoryClass.Get()),
+				FAssetData(Tool.FactoryClass.Get()->GetDefaultObject()),
 				NAME_None,
 				NAME_None,
 				TOptional<FLinearColor>(),
diff --git a/Engine/Plugins/Experimental/Avalanche/Source/AvalancheShapesEditor/Private/AvaShapeFactory.cpp b/Engine/Plugins/Experimental/Avalanche/Source/AvalancheShapesEditor/Private/AvaShapeFactory.cpp
--- a/Engine/Plugins/Experimental/Avalanche/Source/AvalancheShapesEditor/Private/AvaShapeFactory.cpp
+++ b/Engine/Plugins/Experimental/Avalanche/Source/AvalancheShapesEditor/Private/AvaShapeFactory.cpp
@@ -28,9 +28,9 @@ bool UAvaShapeFactory::CanCreateActorFrom(const FAssetData& AssetData, FText& Ou
 		return false;
 	}
 
-	if (UClass* AssetClass = AssetData.GetClass())
+	if (const UClass* AssetClass = AssetData.GetClass(EResolveClass::Yes))
 	{
-		return AssetClass->IsChildOf<AAvaShapeActor>();
+		return AssetClass->IsChildOf(NewActorClass);
 	}
 
 	return false;
@@ -84,6 +84,16 @@ AActor* UAvaShapeFactory::SpawnActor(UObject* InAsset, ULevel* InLevel, const FT
 	return ShapeActor;
 }
 
+void UAvaShapeFactory::PostSpawnActor(UObject* InAsset, AActor* InNewActor)
+{
+	Super::PostSpawnActor(InAsset, InNewActor);
+
+	if (AAvaShapeActor* NewShapeActor = Cast<AAvaShapeActor>(InNewActor))
+	{
+		FActorLabelUtilities::RenameExistingActor(NewShapeActor, NewShapeActor->GetDefaultActorLabel(), /** Unique*/true);
+	}
+}
+
 void UAvaShapeFactory::PostPlaceAsset(TArrayView<const FTypedElementHandle> InHandle, const FAssetPlacementInfo& InPlacementInfo, const FPlacementOptions& InPlacementOptions)
 {
 	Super::PostPlaceAsset(InHandle, InPlacementInfo, InPlacementOptions);
@@ -93,7 +103,7 @@ void UAvaShapeFactory::PostPlaceAsset(TArrayView<const FTypedElementHandle> InHa
 		TArray<FAnalyticsEventAttribute> Attributes;
 		if (MeshClass)
 		{
-			Attributes.Emplace(TEXT("MeshClass"), MeshClass->GetName());	
+			Attributes.Emplace(TEXT("MeshClass"), MeshClass->GetName());
 		}
 		FEngineAnalytics::GetProvider().RecordEvent(TEXT("Editor.Usage.MotionDesign.PlaceShape"), Attributes);
 	}
diff --git a/Engine/Plugins/Experimental/Avalanche/Source/AvalancheShapesEditor/Private/AvaShapeFactory.h b/Engine/Plugins/Experimental/Avalanche/Source/AvalancheShapesEditor/Private/AvaShapeFactory.h
--- a/Engine/Plugins/Experimental/Avalanche/Source/AvalancheShapesEditor/Private/AvaShapeFactory.h
+++ b/Engine/Plugins/Experimental/Avalanche/Source/AvalancheShapesEditor/Private/AvaShapeFactory.h
@@ -25,6 +25,7 @@ protected:
 	virtual bool CanCreateActorFrom(const FAssetData& AssetData, FText& OutErrorMsg) override;
 	virtual AActor* GetDefaultActor(const FAssetData& AssetData) override;
 	virtual AActor* SpawnActor(UObject* InAsset, ULevel* InLevel, const FTransform& InTransform, const FActorSpawnParameters& InSpawnParams) override;
+	virtual void PostSpawnActor(UObject* InAsset, AActor* InNewActor) override;
 	virtual void PostPlaceAsset(TArrayView<const FTypedElementHandle> InHandle, const FAssetPlacementInfo& InPlacementInfo, const FPlacementOptions& InPlacementOptions) override;
 	//~ End UActorFactory
 };
