MotionDesign : Modifiers
- RadialArrange modifier adding missing option to select the plane to dispose actors onto

diff --git a/Engine/Plugins/Experimental/Avalanche/Source/AvalancheModifiers/Private/AvaModifiersActorUtils.cpp b/Engine/Plugins/Experimental/Avalanche/Source/AvalancheModifiers/Private/AvaModifiersActorUtils.cpp
--- a/Engine/Plugins/Experimental/Avalanche/Source/AvalancheModifiers/Private/AvaModifiersActorUtils.cpp
+++ b/Engine/Plugins/Experimental/Avalanche/Source/AvalancheModifiers/Private/AvaModifiersActorUtils.cpp
@@ -15,20 +15,20 @@ FBox FAvaModifiersActorUtils::GetActorsBounds(const TSet<TWeakObjectPtr<AActor>>
 {
 	FBox ActorsBounds = FBox(EForceInit::ForceInit);
 	ActorsBounds.IsValid = 0;
-	
+
 	FVector OrientedVerts[8];
-	
+
 	FTransform AccumulatedTransform = InReferenceTransform;
 
 	for (const TWeakObjectPtr<AActor>& ActorWeak : InActors)
 	{
 		const AActor* Actor = ActorWeak.Get();
-		
+
 		if (!Actor)
 		{
 			continue;
 		}
-		
+
 		if (bInSkipHidden && !IsActorVisible(Actor))
 		{
 			continue;
@@ -53,9 +53,9 @@ FBox FAvaModifiersActorUtils::GetActorsBounds(const TSet<TWeakObjectPtr<AActor>>
 			ActorsBounds.IsValid = 1;
 		}
 	}
-	
+
 	AccumulatedTransform.SetScale3D(FVector::OneVector);
-	
+
 	return ActorsBounds.TransformBy(AccumulatedTransform);
 }
 
@@ -65,7 +65,7 @@ FBox FAvaModifiersActorUtils::GetActorsBounds(AActor* InActor, bool bInIncludeCh
 	{
 		return FBox();
 	}
-	
+
 	TSet<TWeakObjectPtr<AActor>> AttachedModifyActors {InActor};
 
 	if (bInIncludeChildren)
@@ -77,20 +77,20 @@ FBox FAvaModifiersActorUtils::GetActorsBounds(AActor* InActor, bool bInIncludeCh
 			return InAttachedActor;
 		});
 	}
-	
+
 	return FAvaModifiersActorUtils::GetActorsBounds(AttachedModifyActors, InActor->GetActorTransform(), bInSkipHidden);
 }
 
 FBox FAvaModifiersActorUtils::GetActorBounds(const AActor* InActor)
-{	
+{
 	FBox Box(ForceInit);
 	Box.IsValid = 0;
-	
+
 	if (!InActor || !InActor->GetRootComponent())
 	{
 		return Box;
 	}
-	
+
 	FTransform ActorToWorld = InActor->GetTransform();
 	ActorToWorld.SetScale3D(FVector::OneVector);
 	const FTransform WorldToActor = ActorToWorld.Inverse();
@@ -101,7 +101,7 @@ FBox FAvaModifiersActorUtils::GetActorBounds(const AActor* InActor)
 		{
 			return;
 		}
-		
+
 #if WITH_EDITOR
 		// Ignore Visualization Components, but don't consider them as failed components.
 		if (InPrimitiveComponent->IsVisualizationComponent())
@@ -109,7 +109,7 @@ FBox FAvaModifiersActorUtils::GetActorBounds(const AActor* InActor)
 			return;
 		}
 #endif
-		
+
 		const FTransform ComponentToActor = InPrimitiveComponent->GetComponentTransform() * WorldToActor;
 		const FBox ComponentBox = InPrimitiveComponent->CalcBounds(ComponentToActor).GetBox();
 
@@ -138,17 +138,17 @@ FOrientedBox FAvaModifiersActorUtils::GetOrientedBox(const FBox& InLocalBox, con
 FVector FAvaModifiersActorUtils::GetVectorAxis(int32 InAxis)
 {
 	FVector FollowedAxisVector = FVector::ZeroVector;
-	
+
 	if (EnumHasAnyFlags(static_cast<EAvaModifiersAxis>(InAxis), EAvaModifiersAxis::X))
 	{
 		FollowedAxisVector.X = 1;
 	}
-	
+
 	if (EnumHasAnyFlags(static_cast<EAvaModifiersAxis>(InAxis), EAvaModifiersAxis::Y))
 	{
 		FollowedAxisVector.Y = 1;
 	}
-	
+
 	if (EnumHasAnyFlags(static_cast<EAvaModifiersAxis>(InAxis), EAvaModifiersAxis::Z))
 	{
 		FollowedAxisVector.Z = 1;
@@ -175,7 +175,7 @@ bool FAvaModifiersActorUtils::IsActorNotIsolated(const AActor* InActor)
 	{
 		return false;
 	}
-	
+
 #if WITH_EDITOR
 	bool bIsIsolatingActors = false;
 	TArray<TWeakObjectPtr<const AActor>> IsolatedActors;
@@ -195,38 +195,36 @@ bool FAvaModifiersActorUtils::IsActorNotIsolated(const AActor* InActor)
 	return false;
 }
 
-FRotator FAvaModifiersActorUtils::FindLookAtRotation(const FVector& InEyePosition, const FVector& InTargetPosition, const EAvaAxis InAxis, const bool bInFlipAxis)
+FRotator FAvaModifiersActorUtils::FindLookAtRotation(const FVector& InEyePosition, const FVector& InTargetPosition, EAvaModifiersAxis InAxis, bool bInFlipAxis)
 {
-	auto LookAtDirection = [&InEyePosition, &InTargetPosition, bInFlipAxis]() -> FVector
-	{
-		FVector OutDirection = bInFlipAxis ? (InEyePosition - InTargetPosition) : (InTargetPosition - InEyePosition);
-
-		OutDirection.Normalize();
-		if (OutDirection.SizeSquared() < 0.5f)
-		{
-			// Assert possible if OutDirection is not normalized.
-			OutDirection = FVector(1, 0, 0);
-		}
-
-		return OutDirection;
-	};
+	FVector OutDirection = bInFlipAxis ? InEyePosition - InTargetPosition : InTargetPosition - InEyePosition;
+	OutDirection = OutDirection.GetSafeNormal(UE_SMALL_NUMBER, FVector::XAxisVector);
 
 	FMatrix NewRotation = FRotationMatrix::Identity;
 
 	switch (InAxis)
 	{
-	case EAvaAxis::Horizontal:
-		NewRotation = FRotationMatrix::MakeFromX(LookAtDirection());
+	case EAvaModifiersAxis::X:
+		NewRotation = FRotationMatrix::MakeFromX(OutDirection);
 		break;
-	case EAvaAxis::Vertical:
-		NewRotation = FRotationMatrix::MakeFromY(LookAtDirection());
+	case EAvaModifiersAxis::Y:
+		NewRotation = FRotationMatrix::MakeFromY(OutDirection);
 		break;
-	case EAvaAxis::Depth:
-		NewRotation = FRotationMatrix::MakeFromZ(LookAtDirection());
+	case EAvaModifiersAxis::Z:
+		NewRotation = FRotationMatrix::MakeFromZ(OutDirection);
 		break;
+	case EAvaModifiersAxis::None:
+		return NewRotation.Rotator();
+	}
+
+	FRotator Result = NewRotation.Rotator();
+
+	if (bInFlipAxis && InAxis == EAvaModifiersAxis::X)
+	{
+		Result = Result + FRotator(180, 0, 0);
 	}
 
-	return NewRotation.Rotator();
+	return Result;
 }
 
 bool FAvaModifiersActorUtils::IsActorVisible(const AActor* InActor)
diff --git a/Engine/Plugins/Experimental/Avalanche/Source/AvalancheModifiers/Private/AvaModifiersActorUtils.h b/Engine/Plugins/Experimental/Avalanche/Source/AvalancheModifiers/Private/AvaModifiersActorUtils.h
--- a/Engine/Plugins/Experimental/Avalanche/Source/AvalancheModifiers/Private/AvaModifiersActorUtils.h
+++ b/Engine/Plugins/Experimental/Avalanche/Source/AvalancheModifiers/Private/AvaModifiersActorUtils.h
@@ -38,7 +38,7 @@ struct FAvaModifiersActorUtils
 	/** Begin Outliner */
 	static bool IsActorNotIsolated(const AActor* InActor);
 	/** End Outliner */
-	
-	static FRotator FindLookAtRotation(const FVector& InEyePosition, const FVector& InTargetPosition, const EAvaAxis InAxis, const bool bInFlipAxis);
+
+	static FRotator FindLookAtRotation(const FVector& InEyePosition, const FVector& InTargetPosition, EAvaModifiersAxis InAxis, bool bInFlipAxis);
 	static bool IsActorVisible(const AActor* InActor);
 };
\ No newline at end of file
diff --git a/Engine/Plugins/Experimental/Avalanche/Source/AvalancheModifiers/Private/Modifiers/AvaRadialArrangeModifier.cpp b/Engine/Plugins/Experimental/Avalanche/Source/AvalancheModifiers/Private/Modifiers/AvaRadialArrangeModifier.cpp
--- a/Engine/Plugins/Experimental/Avalanche/Source/AvalancheModifiers/Private/Modifiers/AvaRadialArrangeModifier.cpp
+++ b/Engine/Plugins/Experimental/Avalanche/Source/AvalancheModifiers/Private/Modifiers/AvaRadialArrangeModifier.cpp
@@ -10,6 +10,20 @@
 
 #define LOCTEXT_NAMESPACE "AvaRadialArrangeModifier"
 
+void UAvaRadialArrangeModifier::PostLoad()
+{
+	Super::PostLoad();
+
+	if (OrientationAxis == EAvaModifiersAxis::None)
+	{
+		PRAGMA_DISABLE_DEPRECATION_WARNINGS
+
+		OrientationAxis = static_cast<EAvaModifiersAxis>(1 << static_cast<int32>(OrientAxis));
+
+		PRAGMA_ENABLE_DEPRECATION_WARNINGS
+	}
+}
+
 #if WITH_EDITOR
 void UAvaRadialArrangeModifier::PostEditChangeProperty(FPropertyChangedEvent& PropertyChangedEvent)
 {
@@ -19,6 +33,7 @@ void UAvaRadialArrangeModifier::PostEditChangeProperty(FPropertyChangedEvent& Pr
 
 	static const TSet<FName> PropertiesName =
 	{
+		GET_MEMBER_NAME_CHECKED(UAvaRadialArrangeModifier, Plane),
 		GET_MEMBER_NAME_CHECKED(UAvaRadialArrangeModifier, Count),
 		GET_MEMBER_NAME_CHECKED(UAvaRadialArrangeModifier, Rings),
 		GET_MEMBER_NAME_CHECKED(UAvaRadialArrangeModifier, InnerRadius),
@@ -28,7 +43,7 @@ void UAvaRadialArrangeModifier::PostEditChangeProperty(FPropertyChangedEvent& Pr
 		GET_MEMBER_NAME_CHECKED(UAvaRadialArrangeModifier, Arrangement),
 		GET_MEMBER_NAME_CHECKED(UAvaRadialArrangeModifier, bStartFromOuterRadius),
 		GET_MEMBER_NAME_CHECKED(UAvaRadialArrangeModifier, bOrient),
-		GET_MEMBER_NAME_CHECKED(UAvaRadialArrangeModifier, OrientAxis),
+		GET_MEMBER_NAME_CHECKED(UAvaRadialArrangeModifier, OrientationAxis),
 		GET_MEMBER_NAME_CHECKED(UAvaRadialArrangeModifier, bFlipOrient)
 	};
 
@@ -39,80 +54,135 @@ void UAvaRadialArrangeModifier::PostEditChangeProperty(FPropertyChangedEvent& Pr
 }
 #endif // WITH_EDITOR
 
+void UAvaRadialArrangeModifier::SetPlane(EAvaRadialArrangePlane InPlane)
+{
+	if (Plane == InPlane)
+	{
+		return;
+	}
+
+	Plane = InPlane;
+	MarkModifierDirty();
+}
+
 void UAvaRadialArrangeModifier::SetCount(const int32 InCount)
 {
-	Count = InCount;
+	if (Count == InCount)
+	{
+		return;
+	}
 
+	Count = InCount;
 	MarkModifierDirty();
 }
 
 void UAvaRadialArrangeModifier::SetRings(const int32 InRings)
 {
-	Rings = InRings;
+	if (Rings == InRings)
+	{
+		return;
+	}
 
+	Rings = InRings;
 	MarkModifierDirty();
 }
 
 void UAvaRadialArrangeModifier::SetInnerRadius(const float InInnerRadius)
 {
-	InnerRadius = InInnerRadius;
+	if (FMath::IsNearlyEqual(InnerRadius, InInnerRadius))
+	{
+		return;
+	}
 
+	InnerRadius = InInnerRadius;
 	MarkModifierDirty();
 }
 
 void UAvaRadialArrangeModifier::SetOuterRadius(const float InOuterRadius)
 {
-	OuterRadius = InOuterRadius;
+	if (FMath::IsNearlyEqual(OuterRadius, InOuterRadius))
+	{
+		return;
+	}
 
+	OuterRadius = InOuterRadius;
 	MarkModifierDirty();
 }
 
 void UAvaRadialArrangeModifier::SetStartAngle(const float InStartAngle)
 {
-	StartAngle = InStartAngle;
+	if (FMath::IsNearlyEqual(StartAngle, InStartAngle))
+	{
+		return;
+	}
 
+	StartAngle = InStartAngle;
 	MarkModifierDirty();
 }
 
 void UAvaRadialArrangeModifier::SetEndAngle(const float InEndAngle)
 {
-	EndAngle = InEndAngle;
+	if (FMath::IsNearlyEqual(EndAngle, InEndAngle))
+	{
+		return;
+	}
 
+	EndAngle = InEndAngle;
 	MarkModifierDirty();
 }
 
 void UAvaRadialArrangeModifier::SetArrangement(const EAvaRadialArrangement InArrangement)
 {
-	Arrangement = InArrangement;
+	if (Arrangement == InArrangement)
+	{
+		return;
+	}
 
+	Arrangement = InArrangement;
 	MarkModifierDirty();
 }
 
 void UAvaRadialArrangeModifier::SetStartFromOuterRadius(const bool bInStartFromOuterRadius)
 {
-	bStartFromOuterRadius = bInStartFromOuterRadius;
+	if (bStartFromOuterRadius == bInStartFromOuterRadius)
+	{
+		return;
+	}
 
+	bStartFromOuterRadius = bInStartFromOuterRadius;
 	MarkModifierDirty();
 }
 
 void UAvaRadialArrangeModifier::SetOrient(const bool bInOrient)
 {
-	bOrient = bInOrient;
+	if (bOrient == bInOrient)
+	{
+		return;
+	}
 
+	bOrient = bInOrient;
 	MarkModifierDirty();
 }
 
-void UAvaRadialArrangeModifier::SetOrientAxis(const EAvaAxis InOrientAxis)
+void UAvaRadialArrangeModifier::SetOrientationAxis(EAvaModifiersAxis InAxis)
 {
-	OrientAxis = InOrientAxis;
+	if (OrientationAxis == InAxis)
+	{
+		return;
+	}
 
+	OrientationAxis = InAxis;
 	MarkModifierDirty();
 }
 
 void UAvaRadialArrangeModifier::SetFlipOrient(const bool bInFlipOrient)
 {
-	bFlipOrient = bInFlipOrient;
+	if (bFlipOrient == bInFlipOrient)
+	{
+		return;
+	}
 
+	bFlipOrient = bInFlipOrient;
 	MarkModifierDirty();
 }
 
@@ -137,7 +207,7 @@ void UAvaRadialArrangeModifier::Apply()
 		Next();
 		return;
 	}
-	
+
 	const FAvaSceneTreeUpdateModifierExtension* SceneExtension = GetExtension<FAvaSceneTreeUpdateModifierExtension>();
 	if (!SceneExtension)
 	{
@@ -165,7 +235,22 @@ void UAvaRadialArrangeModifier::Apply()
 		const float RingStartOffset = RadiusDistancePerRing * InRingIndex;
 		const float ChildRadius = InnerRadius + RingStartOffset;
 
-		return FVector(0.0f, ChildRadius * SlotCos, ChildRadius * SlotSin);
+		FVector Offset = FVector::ZeroVector;
+
+		switch (Plane)
+		{
+			case EAvaRadialArrangePlane::XY:
+				Offset = FVector(ChildRadius * SlotCos, ChildRadius * SlotSin, 0);
+				break;
+			case EAvaRadialArrangePlane::YZ:
+				Offset = FVector(0, ChildRadius * SlotCos, ChildRadius * SlotSin);
+				break;
+			case EAvaRadialArrangePlane::XZ:
+				Offset = FVector(ChildRadius * SlotCos, 0, ChildRadius * SlotSin);
+				break;
+		}
+
+		return Offset;
 	};
 
 	UAvaTransformModifierShared* LayoutShared = GetShared<UAvaTransformModifierShared>(true);
@@ -190,13 +275,13 @@ void UAvaRadialArrangeModifier::Apply()
 				NewChildrenActorsWeak.Add(ChildActor);
 			}
 		}
-		
+
 		// No need to handle nested children actor, only direct children, visibility will propagate
 		if (AttachedActor->GetAttachParentActor() != ModifyActor)
 		{
 			continue;
 		}
-		
+
 		// Track this actor visibility state
 		const bool bIsVisible = ChildIndex < TotalSlotCount;
 		VisibilityShared->SetActorVisibility(this, AttachedActor, !bIsVisible, true);
@@ -257,22 +342,28 @@ void UAvaRadialArrangeModifier::Apply()
 			RingIndex = Rings - (RingIndex - 1);
 		}
 
-		const FVector RelativeOffset = CalculateRelativeOffset(SlotAngle, RingIndex);
-
 		// Track this actor layout state
 		LayoutShared->SaveActorState(this, AttachedActor);
+
+		const FVector RelativeOffset = CalculateRelativeOffset(SlotAngle, RingIndex);
 		AttachedActor->SetActorRelativeLocation(RelativeOffset);
 
 		if (bOrient)
 		{
 			const FVector EyePosition = AttachedActor->GetActorLocation();
 			const FVector TargetPosition = ModifyActor->GetActorLocation();
-			const FRotator NewRotation = FAvaModifiersActorUtils::FindLookAtRotation(EyePosition, TargetPosition, OrientAxis, bFlipOrient);
+
+			const FRotator NewRotation = FAvaModifiersActorUtils::FindLookAtRotation(EyePosition, TargetPosition, OrientationAxis, bFlipOrient);
 
 			AttachedActor->SetActorRotation(NewRotation);
 		}
+		// Restore original rotation
+		else if (FAvaTransformSharedActorState* ActorState = LayoutShared->FindActorState(AttachedActor))
+		{
+			AttachedActor->SetActorRotation(ActorState->ActorTransform.GetRotation());
+		}
 	}
-	
+
 	// Untrack previous actors that are not attached anymore
 	const TSet<TWeakObjectPtr<AActor>> UntrackActors = ChildrenActorsWeak.Difference(NewChildrenActorsWeak);
 	LayoutShared->RestoreActorsState(this, UntrackActors);
diff --git a/Engine/Plugins/Experimental/Avalanche/Source/AvalancheModifiers/Public/Modifiers/AvaRadialArrangeModifier.h b/Engine/Plugins/Experimental/Avalanche/Source/AvalancheModifiers/Public/Modifiers/AvaRadialArrangeModifier.h
--- a/Engine/Plugins/Experimental/Avalanche/Source/AvalancheModifiers/Public/Modifiers/AvaRadialArrangeModifier.h
+++ b/Engine/Plugins/Experimental/Avalanche/Source/AvalancheModifiers/Public/Modifiers/AvaRadialArrangeModifier.h
@@ -3,6 +3,7 @@
 #pragma once
 
 #include "AvaArrangeBaseModifier.h"
+#include "AvaModifiersActorUtils.h"
 #include "AvaRadialArrangeModifier.generated.h"
 
 /** Specifies how child elements will be arranged radially. */
@@ -21,6 +22,15 @@ enum class EAvaRadialArrangement : uint8
 	 Equal
 };
 
+/** Enumerates how to layout the ring */
+UENUM(BlueprintType)
+enum class EAvaRadialArrangePlane : uint8
+{
+	XY,
+	YZ,
+	XZ,
+};
+
 /**
  * Arranges child actors in a circular rings around its center
  */
@@ -30,6 +40,12 @@ class UAvaRadialArrangeModifier : public UAvaArrangeBaseModifier
 	GENERATED_BODY()
 
 public:
+	AVALANCHEMODIFIERS_API void SetPlane(EAvaRadialArrangePlane InPlane);
+	EAvaRadialArrangePlane GetPlane() const
+	{
+		return Plane;
+	}
+
 	/** Sets the number of child elements to use in the arrangement. Children whose index is greater than or equal to this value will be hidden. */
 	UFUNCTION(BlueprintCallable, Category="Motion Design|Modifiers|RadialArrange")
 	AVALANCHEMODIFIERS_API void SetCount(const int32 InCount);
@@ -131,13 +147,13 @@ public:
 
 	/** Sets the axis to look at the center. */
 	UFUNCTION(BlueprintCallable, Category="Motion Design|Modifiers|RadialArrange")
-	AVALANCHEMODIFIERS_API void SetOrientAxis(const EAvaAxis InOrientAxis);
+	AVALANCHEMODIFIERS_API void SetOrientationAxis(EAvaModifiersAxis InAxis);
 
 	/** Gets the axis to look at the center. */
 	UFUNCTION(BlueprintPure, Category="Motion Design|Modifiers|RadialArrange")
-	EAvaAxis GetOrientAxis() const
+	EAvaModifiersAxis GetOrientationAxis() const
 	{
-		return OrientAxis;
+		return OrientationAxis;
 	}
 
 	/** If true, will flip the center orientation to face outwards. */
@@ -153,6 +169,7 @@ public:
 
 protected:
 	//~ Begin UObject
+	virtual void PostLoad() override;
 #if WITH_EDITOR
 	virtual void PostEditChangeProperty(FPropertyChangedEvent& PropertyChangedEvent) override;
 #endif
@@ -163,32 +180,36 @@ protected:
 	virtual void Apply() override;
 	//~ End UActorModifierCoreBase
 
+	/** Base plane for the radial arrangement */
+	UPROPERTY(EditInstanceOnly, Setter, Getter, Category="RadialArrange", meta=(AllowPrivateAccess="true"))
+	EAvaRadialArrangePlane Plane = EAvaRadialArrangePlane::XY;
+
 	/** The number of child elements to limit in the arrangement, or -1 if unlimited. Children whose index is greater than or equal to this value will be hidden. */
-	UPROPERTY(EditInstanceOnly, Setter="SetCount", Getter="GetCount", Category="RadialArrange", meta=(ClampMin="-1", UIMin="-1", AllowPrivateAccess="true"))
+	UPROPERTY(EditInstanceOnly, Setter, Getter, Category="RadialArrange", meta=(ClampMin="-1", UIMin="-1", AllowPrivateAccess="true"))
 	int32 Count = -1;
 
 	/** The number of rings. */
-	UPROPERTY(EditInstanceOnly, Setter="SetRings", Getter="GetRings", Category="RadialArrange", meta=(ClampMin="1", UIMin="1", AllowPrivateAccess="true"))
+	UPROPERTY(EditInstanceOnly, Setter, Getter, Category="RadialArrange", meta=(ClampMin="1", UIMin="1", AllowPrivateAccess="true"))
 	int32 Rings = 1;
 
 	/** The radius from the center to the first inner ring. */
-	UPROPERTY(EditInstanceOnly, Setter="SetInnerRadius", Getter="GetInnerRadius", Interp, Category="RadialArrange", meta=(ClampMin="0.0", UIMin="0.0", AllowPrivateAccess="true"))
+	UPROPERTY(EditInstanceOnly, Setter, Getter, Interp, Category="RadialArrange", meta=(ClampMin="0.0", UIMin="0.0", AllowPrivateAccess="true"))
 	float InnerRadius = 70.0f;
 
 	/** The radius from the center to the last outer ring. */
-	UPROPERTY(EditInstanceOnly, Setter="SetOuterRadius", Getter="GetOuterRadius", Interp, Category="RadialArrange", meta=(EditCondition="Rings > 1", EditConditionHides, ClampMin="0.0", UIMin="0.0", AllowPrivateAccess="true"))
+	UPROPERTY(EditInstanceOnly, Setter, Getter, Interp, Category="RadialArrange", meta=(EditCondition="Rings > 1", EditConditionHides, ClampMin="0.0", UIMin="0.0", AllowPrivateAccess="true"))
 	float OuterRadius = 200.0f;
 
 	/** The start angle for the arrangement space and moving clockwise. 0 = Up, -90 = Left, 90 = Right */
-	UPROPERTY(EditInstanceOnly, Setter="SetStartAngle", Getter="GetStartAngle", Interp, Category="RadialArrange", meta=(ClampMin="-180.0", UIMin="-180.0", ClampMax="180.0", UIMax="180.0", AllowPrivateAccess="true"))
+	UPROPERTY(EditInstanceOnly, Setter, Getter, Interp, Category="RadialArrange", meta=(ClampMin="-180.0", UIMin="-180.0", ClampMax="180.0", UIMax="180.0", AllowPrivateAccess="true"))
 	float StartAngle = -180.0f;
 
 	/** The end angle for the arrangement space. 0 = Up, -90 = Left, 90 = Right */
-	UPROPERTY(EditInstanceOnly, Setter="SetEndAngle", Getter="GetEndAngle", Interp, Category="RadialArrange", meta=(ClampMin = "-180.0", UIMin = "-180.0", ClampMax = "180.0", UIMax = "180.0", AllowPrivateAccess="true"))
+	UPROPERTY(EditInstanceOnly, Setter, Getter, Interp, Category="RadialArrange", meta=(ClampMin = "-180.0", UIMin = "-180.0", ClampMax = "180.0", UIMax = "180.0", AllowPrivateAccess="true"))
 	float EndAngle = 180.0f;
 
 	/** Defines how to arrange the child elements around the center. */
-	UPROPERTY(EditInstanceOnly, Setter="SetArrangement", Getter="GetArrangement", Category="RadialArrange", meta=(AllowPrivateAccess="true"))
+	UPROPERTY(EditInstanceOnly, Setter, Getter, Category="RadialArrange", meta=(AllowPrivateAccess="true"))
 	EAvaRadialArrangement Arrangement = EAvaRadialArrangement::Equal;
 
 	/** If true, will arrange the child elements starting from the outer radius and moving to the inner radius. Has no effect if only using one ring. */
@@ -200,7 +221,11 @@ protected:
 	bool bOrient;
 
 	/** The axis to look at the center. */
-	UPROPERTY(EditInstanceOnly, Setter="SetOrientAxis", Getter="GetOrientAxis", Category="RadialArrange", meta=(EditCondition="bOrient", EditConditionHides, AllowPrivateAccess="true"))
+	UPROPERTY(EditInstanceOnly, Setter, Getter, Category="RadialArrange", meta=(EditCondition="bOrient", EditConditionHides, AllowPrivateAccess="true"))
+	EAvaModifiersAxis OrientationAxis = EAvaModifiersAxis::None;
+
+	UE_DEPRECATED(5.5, TEXT("Use OrientationAxis instead"))
+	UPROPERTY(meta=(DeprecatedProperty, DeprecationMessage="Use OrientationAxis instead"))
 	EAvaAxis OrientAxis;
 
 	/** If true, will flip the orientation axis to the opposite direction. */
