MotionDesign : Modifiers
- Fixed deprecated message warning format

diff --git a/Engine/Plugins/Experimental/Avalanche/Source/AvalancheModifiers/Public/Modifiers/AvaRadialArrangeModifier.h b/Engine/Plugins/Experimental/Avalanche/Source/AvalancheModifiers/Public/Modifiers/AvaRadialArrangeModifier.h
--- a/Engine/Plugins/Experimental/Avalanche/Source/AvalancheModifiers/Public/Modifiers/AvaRadialArrangeModifier.h
+++ b/Engine/Plugins/Experimental/Avalanche/Source/AvalancheModifiers/Public/Modifiers/AvaRadialArrangeModifier.h
@@ -224,7 +224,7 @@ protected:
 	UPROPERTY(EditInstanceOnly, Setter, Getter, Category="RadialArrange", meta=(EditCondition="bOrient", EditConditionHides, AllowPrivateAccess="true"))
 	EAvaModifiersAxis OrientationAxis = EAvaModifiersAxis::None;
 
-	UE_DEPRECATED(5.5, TEXT("Use OrientationAxis instead"))
+	UE_DEPRECATED(5.5, "Use OrientationAxis instead")
 	UPROPERTY(meta=(DeprecatedProperty, DeprecationMessage="Use OrientationAxis instead"))
 	EAvaAxis OrientAxis;
 
