Outliner Level items

diff --git a/Engine/Plugins/Experimental/Avalanche/Source/AvalancheOutliner/Private/AvaOutliner.cpp b/Engine/Plugins/Experimental/Avalanche/Source/AvalancheOutliner/Private/AvaOutliner.cpp
--- a/Engine/Plugins/Experimental/Avalanche/Source/AvalancheOutliner/Private/AvaOutliner.cpp
+++ b/Engine/Plugins/Experimental/Avalanche/Source/AvalancheOutliner/Private/AvaOutliner.cpp
@@ -452,8 +452,8 @@ void FAvaOutliner::Refresh()
 		InOutlinerView->Refresh();
 	});
 
-	// Reset Tree and Save so that the Tree is updated to the latest Outliner State
-	SaveState->SaveSceneTree(*this, /*bInResetTree*/true);
+	// Save so that the Tree is updated to the latest Outliner State
+	SaveState->SaveSceneTree(*this, /*bInResetTree*/false);
 }
 
 TSharedRef<FAvaOutlinerItem> FAvaOutliner::GetTreeRoot() const
diff --git a/Engine/Plugins/Experimental/Avalanche/Source/AvalancheOutliner/Private/Item/AvaOutlinerTreeRoot.cpp b/Engine/Plugins/Experimental/Avalanche/Source/AvalancheOutliner/Private/Item/AvaOutlinerTreeRoot.cpp
--- a/Engine/Plugins/Experimental/Avalanche/Source/AvalancheOutliner/Private/Item/AvaOutlinerTreeRoot.cpp
+++ b/Engine/Plugins/Experimental/Avalanche/Source/AvalancheOutliner/Private/Item/AvaOutlinerTreeRoot.cpp
@@ -1,7 +1,6 @@
 // Copyright Epic Games, Inc. All Rights Reserved.
 
 #include "Item/AvaOutlinerTreeRoot.h"
-
 #include "ActorFactories/ActorFactory.h"
 #include "AssetSelection.h"
 #include "AvaOutliner.h"
@@ -10,6 +9,7 @@
 #include "DragDropOps/AvaOutlinerItemDragDropOp.h"
 #include "EngineUtils.h"
 #include "Item/AvaOutlinerActor.h"
+#include "Item/AvaOutlinerLevel.h"
 #include "ItemActions/AvaOutlinerAddItem.h"
 #include "Textures/SlateIcon.h"
 #include "Widgets/SNullWidget.h"
@@ -25,26 +25,24 @@ void FAvaOutlinerTreeRoot::FindChildren(TArray<FAvaOutlinerItemPtr>& OutChildren
 		return;
 	}
 
-	FAvaOutliner& OutlinerPrivate = static_cast<FAvaOutliner&>(Outliner);
-
-	// All the Actors that don't have an Outliner Parent will be at the Root
-	const TArray<TWeakObjectPtr<AActor>> RootChildren = OutlinerPrivate.GetActorSceneOutlinerChildren(nullptr);
+	const TArray<ULevel*>& Levels = World->GetLevels();
 
-	// Worst case and most likely case: all Outliner Children are valid.
-	// Note: if recursive, re-allocations will still need to be done past this reserve count, as it's unknown at this point how many items are going to be added
-	OutChildren.Reserve(OutChildren.Num() + RootChildren.Num());
+	OutChildren.Reserve(OutChildren.Num() + Levels.Num());
 
-	for (const TWeakObjectPtr<AActor>& ChildWeak : RootChildren)
+	for (ULevel* Level : Levels)
 	{
-		if (AActor* const Child = ChildWeak.Get())
+		if (!Level)
 		{
-			const FAvaOutlinerItemPtr ChildActorItem = Outliner.FindOrAdd<FAvaOutlinerActor>(Child);
-			const FAvaOutlinerItemFlagGuard Guard(ChildActorItem, EAvaOutlinerItemFlags::IgnorePendingKill);
-			OutChildren.Add(ChildActorItem);
-			if (bRecursive)
-			{
-				ChildActorItem->FindChildren(OutChildren, bRecursive);
-			}
+			continue;
+		}
+
+		const FAvaOutlinerItemPtr LevelItem = Outliner.FindOrAdd<FAvaOutlinerLevel>(Level);
+
+		const FAvaOutlinerItemFlagGuard Guard(LevelItem, EAvaOutlinerItemFlags::IgnorePendingKill);
+		OutChildren.Add(LevelItem);
+		if (bRecursive)
+		{
+			LevelItem->FindChildren(OutChildren, bRecursive);
 		}
 	}
 }
diff --git a/Engine/Plugins/Experimental/Avalanche/Source/AvalancheOutliner/Public/Item/AvaOutlinerLevel.cpp b/Engine/Plugins/Experimental/Avalanche/Source/AvalancheOutliner/Public/Item/AvaOutlinerLevel.cpp
new file mode 100644
--- /dev/null
+++ b/Engine/Plugins/Experimental/Avalanche/Source/AvalancheOutliner/Public/Item/AvaOutlinerLevel.cpp
@@ -0,0 +1,110 @@
+// Copyright Epic Games, Inc. All Rights Reserved.
+
+#include "AvaOutlinerLevel.h"
+#include "AvaOutliner.h"
+#include "AvaOutlinerActor.h"
+#include "Engine/Level.h"
+#include "Engine/World.h"
+#include "GameFramework/Actor.h"
+#include "Input/Reply.h"
+#include "Styling/SlateIconFinder.h"
+
+FAvaOutlinerLevel::FAvaOutlinerLevel(IAvaOutliner& InOutliner, ULevel* InLevel)
+	: Super(InOutliner, InLevel)
+	, LevelWeak(InLevel)
+{
+}
+
+void FAvaOutlinerLevel::FindChildren(TArray<FAvaOutlinerItemPtr>& OutChildren, bool bRecursive)
+{
+	Super::FindChildren(OutChildren, bRecursive);
+
+	ULevel* const Level = GetLevel();
+	if (!Level)
+	{
+		return;
+	}
+
+	OutChildren.Reserve(OutChildren.Num() + Level->Actors.Num());
+
+	for (AActor* Actor : Level->Actors)
+	{
+		// Only consider actors that are not attached to any other actor
+		if (!Actor || Actor->GetSceneOutlinerParent())
+		{
+			continue;
+		}
+
+		const FAvaOutlinerItemPtr ActorItem = Outliner.FindOrAdd<FAvaOutlinerActor>(Actor);
+
+		const FAvaOutlinerItemFlagGuard Guard(ActorItem, EAvaOutlinerItemFlags::IgnorePendingKill);
+		OutChildren.Add(ActorItem);
+		if (bRecursive)
+		{
+			ActorItem->FindChildren(OutChildren, bRecursive);
+		}
+	}
+}
+
+bool FAvaOutlinerLevel::AddChild(const FAvaOutlinerAddItemParams& InAddItemParams)
+{
+	// todo: handle actor being moved to this level item
+	return Super::AddChild(InAddItemParams);
+}
+
+bool FAvaOutlinerLevel::RemoveChild(const FAvaOutlinerRemoveItemParams& InRemoveItemParams)
+{
+	// todo: handle actor being moved from this level item
+	return Super::RemoveChild(InRemoveItemParams);
+}
+
+EAvaOutlinerItemViewMode FAvaOutlinerLevel::GetSupportedViewModes(const FAvaOutlinerView& InOutlinerView) const
+{
+	// Actors should only be visualized in Outliner View and not appear in the Item Column List
+	// Support any other type of View Mode
+	return EAvaOutlinerItemViewMode::ItemTree | ~EAvaOutlinerItemViewMode::HorizontalItemList;
+}
+
+FText FAvaOutlinerLevel::GetDisplayName() const
+{
+	ULevel* Level = GetLevel();
+	if (!Level)
+	{
+		return FText::GetEmpty();
+	}
+
+	UWorld* OwningWorld = Level->GetWorld();
+	UWorld* OuterWorld = Level->GetTypedOuter<UWorld>();
+	if (!OwningWorld || !OuterWorld || OwningWorld == OuterWorld)
+	{
+		return FText::FromString(Level->GetName());
+	}
+
+	return FText::FromString(OuterWorld->GetName());	
+}
+
+FSlateIcon FAvaOutlinerLevel::GetIcon() const
+{
+	return FSlateIconFinder::FindIconForClass(UWorld::StaticClass());
+}
+
+TOptional<EItemDropZone> FAvaOutlinerLevel::CanAcceptDrop(const FDragDropEvent& InDragDropEvent, EItemDropZone InDropZone)
+{
+	return TOptional<EItemDropZone>();
+}
+
+FReply FAvaOutlinerLevel::AcceptDrop(const FDragDropEvent& InDragDropEvent, EItemDropZone InDropZone)
+{
+	return FReply::Unhandled();
+}
+
+ULevel* FAvaOutlinerLevel::GetLevel() const
+{
+	return LevelWeak.Get(IsIgnoringPendingKill());
+}
+
+void FAvaOutlinerLevel::SetObject_Impl(UObject* InObject)
+{
+	Super::SetObject_Impl(InObject);
+	LevelWeak = Cast<ULevel>(InObject);
+}
diff --git a/Engine/Plugins/Experimental/Avalanche/Source/AvalancheOutliner/Public/Item/AvaOutlinerLevel.h b/Engine/Plugins/Experimental/Avalanche/Source/AvalancheOutliner/Public/Item/AvaOutlinerLevel.h
new file mode 100644
--- /dev/null
+++ b/Engine/Plugins/Experimental/Avalanche/Source/AvalancheOutliner/Public/Item/AvaOutlinerLevel.h
@@ -0,0 +1,41 @@
+// Copyright Epic Games, Inc. All Rights Reserved.
+
+#pragma once
+
+#include "AvaOutlinerObject.h"
+
+class ULevel;
+
+/**
+ * Item in Outliner representing a Level.
+ */
+class FAvaOutlinerLevel : public FAvaOutlinerObject
+{
+public:
+	UE_AVA_INHERITS_WITH_SUPER(FAvaOutlinerLevel, FAvaOutlinerObject);
+
+	AVALANCHEOUTLINER_API FAvaOutlinerLevel(IAvaOutliner& InOutliner, ULevel* InLevel);
+
+	//~ Begin IAvaOutlinerItem
+	AVALANCHEOUTLINER_API virtual void FindChildren(TArray<FAvaOutlinerItemPtr>& OutChildren, bool bRecursive) override;
+	AVALANCHEOUTLINER_API virtual bool AddChild(const FAvaOutlinerAddItemParams& InAddItemParams) override;
+	AVALANCHEOUTLINER_API virtual bool RemoveChild(const FAvaOutlinerRemoveItemParams& InRemoveItemParams) override;
+	AVALANCHEOUTLINER_API virtual EAvaOutlinerItemViewMode GetSupportedViewModes(const FAvaOutlinerView& InOutlinerView) const override;
+	virtual bool CanBeTopLevel() const override { return true; }
+	virtual bool CanRename() const override { return false; }
+	virtual bool CanLock() const override { return false; }
+	AVALANCHEOUTLINER_API virtual FText GetDisplayName() const override;
+	AVALANCHEOUTLINER_API virtual FSlateIcon GetIcon() const override;
+	AVALANCHEOUTLINER_API virtual TOptional<EItemDropZone> CanAcceptDrop(const FDragDropEvent& InDragDropEvent, EItemDropZone InDropZone) override;
+	AVALANCHEOUTLINER_API virtual FReply AcceptDrop(const FDragDropEvent& InDragDropEvent, EItemDropZone InDropZone) override;
+	//~ End IAvaOutlinerItem
+
+	AVALANCHEOUTLINER_API ULevel* GetLevel() const;
+
+protected:
+	//~ Begin FAvaOutlinerObjectItem
+	virtual void SetObject_Impl(UObject* InObject) override;
+	//~ End FAvaOutlinerObjectItem
+
+	TWeakObjectPtr<ULevel> LevelWeak;
+};
