[Motion Design] Rundown Server: Adding Transition & Sequence Event Notifications

The new notifications allow better client side logic and UI updates.

- New Rundown API messages: FAvaRundownPageSequenceEvent and FAvaRundownPageTransitionEvent.
- Adding "Paused" sequence event to allow "continue" logic to be implement client side.
- Potential crash fix: Rundown::PlayPage will not play if transitions are still going. It could before it could lead to a crash. The command will reply with a "fail" response and the client will be able to try again when the current transition is finished (which it can now track).

diff --git a/Engine/Plugins/Experimental/Avalanche/Source/AvalancheMedia/Internal/Rundown/AvaRundownMessages.h b/Engine/Plugins/Experimental/Avalanche/Source/AvalancheMedia/Internal/Rundown/AvaRundownMessages.h
--- a/Engine/Plugins/Experimental/Avalanche/Source/AvalancheMedia/Internal/Rundown/AvaRundownMessages.h
+++ b/Engine/Plugins/Experimental/Avalanche/Source/AvalancheMedia/Internal/Rundown/AvaRundownMessages.h
@@ -746,25 +746,68 @@ public:
 	FString PreviewChannelName;
 };
 
+/**
+ * Page Sequence Event Message
+ */
+USTRUCT()
+struct FAvaRundownPageSequenceEvent : public FAvaRundownMsgBase
+{
+	GENERATED_BODY()
+
+	UPROPERTY()
+	FString Channel;
+	
+	UPROPERTY()
+	int32 PageId = FAvaRundownPage::InvalidPageId;
+
+	/** Playable Instance Id. */
+	UPROPERTY()
+	FGuid InstanceId;
+
+	/** Full asset path: /PackagePath/PackageName.AssetName */
+	UPROPERTY()
+	FString AssetPath;
+	
+	UPROPERTY()
+	FString SequenceLabel;
+
+	UPROPERTY()
+	EAvaPlayableSequenceEventType Event = EAvaPlayableSequenceEventType::None;
+};
+
 UENUM()
-enum class EAvaRundownPageEvents
+enum class EAvaRundownPageTransitionEvents
 {
 	None,
-	AnimStarted,
-	AnimPaused,
-	AnimFinished
+	Started,
+	Finished
 };
 
+/**
+ * Page Transition Event Message
+ */
 USTRUCT()
-struct FAvaRundownPageEvent : public FAvaRundownMsgBase
+struct FAvaRundownPageTransitionEvent : public FAvaRundownMsgBase
 {
 	GENERATED_BODY()
-public:
+
 	UPROPERTY()
-	int32 PageId = FAvaRundownPage::InvalidPageId;
+	FString Channel;
+	
+	UPROPERTY()
+	FGuid TransitionId;
+
+	UPROPERTY()
+	TArray<int32> EnteringPageIds;
+
+	UPROPERTY()
+	TArray<int32> PlayingPageIds;
+
+	UPROPERTY()
+	TArray<int32> ExitingPageIds;
 
 	UPROPERTY()
-	EAvaRundownPageEvents Event = EAvaRundownPageEvents::None;
+	EAvaRundownPageTransitionEvents Event = EAvaRundownPageTransitionEvents::None;
 };
 
 USTRUCT()
diff --git a/Engine/Plugins/Experimental/Avalanche/Source/AvalancheMedia/Private/Playable/AvaPlayable.cpp b/Engine/Plugins/Experimental/Avalanche/Source/AvalancheMedia/Private/Playable/AvaPlayable.cpp
--- a/Engine/Plugins/Experimental/Avalanche/Source/AvalancheMedia/Private/Playable/AvaPlayable.cpp
+++ b/Engine/Plugins/Experimental/Avalanche/Source/AvalancheMedia/Private/Playable/AvaPlayable.cpp
@@ -287,6 +287,7 @@ void UAvaPlayable::BeginPlay(const FAvaInstancePlaySettings& InWorldPlaySettings
 		
 		// Playable events need to transit through playback events to reach the rundown for proper impl layer separation.
 		UAvaSequencePlayer::OnSequenceStarted().AddUObject(this, &UAvaPlayable::HandleOnSequenceStarted);
+		UAvaSequencePlayer::OnSequencePaused().AddUObject(this, &UAvaPlayable::HandleOnSequencePaused);
 		UAvaSequencePlayer::OnSequenceFinished().AddUObject(this, &UAvaPlayable::HandleOnSequenceFinished);
 
 		OnPlay();
@@ -302,6 +303,7 @@ void UAvaPlayable::EndPlay(EAvaPlayableEndPlayOptions InOptions)
 
 	bIsPlaying = false;
 	UAvaSequencePlayer::OnSequenceStarted().RemoveAll(this);
+	UAvaSequencePlayer::OnSequencePaused().RemoveAll(this);
 	UAvaSequencePlayer::OnSequenceFinished().RemoveAll(this);
 	OnEndPlay();
 
@@ -363,7 +365,18 @@ void UAvaPlayable::HandleOnSequenceStarted(UAvaSequencePlayer* InSequencePlayer,
 		using namespace UE::AvaPlayable::Private;
 		UE_LOG(LogAvaPlayable, Verbose, TEXT("%s Playable {%s}: Sequence {%s} started."),
 			*GetBriefFrameInfo(), *GetPrettyPlayableInfo(this),  *GetPrettySequenceInfo(InSequence));
-		OnSequenceEventDelegate.Broadcast(this, InSequence->GetFName(), EAvaPlayableSequenceEventType::Started);
+		OnSequenceEventDelegate.Broadcast(this, InSequence->GetLabel(), EAvaPlayableSequenceEventType::Started);
+	}
+}
+
+void UAvaPlayable::HandleOnSequencePaused(UAvaSequencePlayer* InSequencePlayer, UAvaSequence* InSequence)
+{
+	if (HasSequence(InSequence))
+	{
+		using namespace UE::AvaPlayable::Private;
+		UE_LOG(LogAvaPlayable, Verbose, TEXT("%s Playable {%s}: Sequence {%s} paused."),
+			*GetBriefFrameInfo(), *GetPrettyPlayableInfo(this),  *GetPrettySequenceInfo(InSequence));
+		OnSequenceEventDelegate.Broadcast(this, InSequence->GetLabel(), EAvaPlayableSequenceEventType::Paused);
 	}
 }
 
@@ -374,7 +387,7 @@ void UAvaPlayable::HandleOnSequenceFinished(UAvaSequencePlayer* InSequencePlayer
 		using namespace UE::AvaPlayable::Private;
 		UE_LOG(LogAvaPlayable, Verbose, TEXT("%s Playable {%s}: Sequence {%s} finished."),
 			*GetBriefFrameInfo(), *GetPrettyPlayableInfo(this),  *GetPrettySequenceInfo(InSequence));
-		OnSequenceEventDelegate.Broadcast(this, InSequence->GetFName(), EAvaPlayableSequenceEventType::Finished);
+		OnSequenceEventDelegate.Broadcast(this, InSequence->GetLabel(), EAvaPlayableSequenceEventType::Finished);
 	}
 }
 
diff --git a/Engine/Plugins/Experimental/Avalanche/Source/AvalancheMedia/Private/Playable/Transition/AvaPlayableLocalTransition.cpp b/Engine/Plugins/Experimental/Avalanche/Source/AvalancheMedia/Private/Playable/Transition/AvaPlayableLocalTransition.cpp
--- a/Engine/Plugins/Experimental/Avalanche/Source/AvalancheMedia/Private/Playable/Transition/AvaPlayableLocalTransition.cpp
+++ b/Engine/Plugins/Experimental/Avalanche/Source/AvalancheMedia/Private/Playable/Transition/AvaPlayableLocalTransition.cpp
@@ -477,7 +477,7 @@ void UAvaPlayableLocalTransition::FinishWaitOnPostExecutorSequences()
 	UAvaPlayable::OnSequenceEvent().RemoveAll(this);
 }
 
-void UAvaPlayableLocalTransition::OnPlayableSequenceEvent(UAvaPlayable* InPlayable, const FName& SequenceName, EAvaPlayableSequenceEventType InEventType)
+void UAvaPlayableLocalTransition::OnPlayableSequenceEvent(UAvaPlayable* InPlayable, FName InSequenceLabel, EAvaPlayableSequenceEventType InEventType)
 {
 	// Remark: the sequence events are not entirely reliable.
 	// The transitions are also "ticked" to poll this condition.
diff --git a/Engine/Plugins/Experimental/Avalanche/Source/AvalancheMedia/Private/Playable/Transition/AvaPlayableLocalTransition.h b/Engine/Plugins/Experimental/Avalanche/Source/AvalancheMedia/Private/Playable/Transition/AvaPlayableLocalTransition.h
--- a/Engine/Plugins/Experimental/Avalanche/Source/AvalancheMedia/Private/Playable/Transition/AvaPlayableLocalTransition.h
+++ b/Engine/Plugins/Experimental/Avalanche/Source/AvalancheMedia/Private/Playable/Transition/AvaPlayableLocalTransition.h
@@ -35,7 +35,7 @@ protected:
 	
 	void PostTransitionExecutorPhase();
 
-	void OnPlayableSequenceEvent(UAvaPlayable* InPlayable, const FName& SequenceName, EAvaPlayableSequenceEventType InEventType);
+	void OnPlayableSequenceEvent(UAvaPlayable* InPlayable, FName InSequenceLabel, EAvaPlayableSequenceEventType InEventType);
 	
 	void NotifyTransitionFinished();
 
diff --git a/Engine/Plugins/Experimental/Avalanche/Source/AvalancheMedia/Private/Playback/AvaPlaybackClient.cpp b/Engine/Plugins/Experimental/Avalanche/Source/AvalancheMedia/Private/Playback/AvaPlaybackClient.cpp
--- a/Engine/Plugins/Experimental/Avalanche/Source/AvalancheMedia/Private/Playback/AvaPlaybackClient.cpp
+++ b/Engine/Plugins/Experimental/Avalanche/Source/AvalancheMedia/Private/Playback/AvaPlaybackClient.cpp
@@ -1232,7 +1232,7 @@ void FAvaPlaybackClient::HandlePlaybackSequenceEventMessage(const FAvaPlaybackSe
 		InMessage.AssetPath,
 		InMessage.ChannelName,
 		InMessage.ServerName,
-		InMessage.SequenceName,
+		InMessage.SequenceLabel,
 		InMessage.EventType,
 		InMessage.FrameNumber
 	};
diff --git a/Engine/Plugins/Experimental/Avalanche/Source/AvalancheMedia/Private/Playback/AvaPlaybackMessages.h b/Engine/Plugins/Experimental/Avalanche/Source/AvalancheMedia/Private/Playback/AvaPlaybackMessages.h
--- a/Engine/Plugins/Experimental/Avalanche/Source/AvalancheMedia/Private/Playback/AvaPlaybackMessages.h
+++ b/Engine/Plugins/Experimental/Avalanche/Source/AvalancheMedia/Private/Playback/AvaPlaybackMessages.h
@@ -418,7 +418,7 @@ struct FAvaPlaybackSequenceEvent : public FAvaPlaybackServerMessageBase
 	FString ChannelName;
 	
 	UPROPERTY()
-	FString SequenceName;
+	FString SequenceLabel;
 
 	UPROPERTY()
 	EAvaPlayableSequenceEventType EventType = EAvaPlayableSequenceEventType::None;
diff --git a/Engine/Plugins/Experimental/Avalanche/Source/AvalancheMedia/Private/Playback/AvaPlaybackServer.cpp b/Engine/Plugins/Experimental/Avalanche/Source/AvalancheMedia/Private/Playback/AvaPlaybackServer.cpp
--- a/Engine/Plugins/Experimental/Avalanche/Source/AvalancheMedia/Private/Playback/AvaPlaybackServer.cpp
+++ b/Engine/Plugins/Experimental/Avalanche/Source/AvalancheMedia/Private/Playback/AvaPlaybackServer.cpp
@@ -1105,7 +1105,7 @@ void FAvaPlaybackServer::OnPlaybackAssetRemoved(const FSoftObjectPath& InAssetPa
 	}
 }
 
-void FAvaPlaybackServer::OnPlayableSequenceEvent(UAvaPlayable* InPlayable, const FName& SequenceName, EAvaPlayableSequenceEventType InEventType)
+void FAvaPlaybackServer::OnPlayableSequenceEvent(UAvaPlayable* InPlayable, FName InSequenceLabel, EAvaPlayableSequenceEventType InEventType)
 {
 	if (!InPlayable)
 	{
@@ -1132,7 +1132,7 @@ void FAvaPlaybackServer::OnPlayableSequenceEvent(UAvaPlayable* InPlayable, const
 	Message->InstanceId = InPlayable->GetInstanceId();
 	Message->AssetPath = PlaybackInstance->GetSourcePath();
 	Message->ChannelName = PlaybackInstance->GetChannelName();
-	Message->SequenceName = SequenceName.ToString();
+	Message->SequenceLabel = InSequenceLabel.ToString();
 	Message->EventType = InEventType;
 	Message->FrameNumber = GFrameNumber;
 	SendResponse(Message, ClientAddresses);
diff --git a/Engine/Plugins/Experimental/Avalanche/Source/AvalancheMedia/Private/Playback/AvaPlaybackServer.h b/Engine/Plugins/Experimental/Avalanche/Source/AvalancheMedia/Private/Playback/AvaPlaybackServer.h
--- a/Engine/Plugins/Experimental/Avalanche/Source/AvalancheMedia/Private/Playback/AvaPlaybackServer.h
+++ b/Engine/Plugins/Experimental/Avalanche/Source/AvalancheMedia/Private/Playback/AvaPlaybackServer.h
@@ -146,7 +146,7 @@ protected:
 	void OnPlaybackInstanceInvalidated(const FAvaPlaybackInstance& InPlaybackInstance);
 	void OnPlaybackInstanceStatusChanged(const FAvaPlaybackInstance& InPlaybackInstance);
 	void OnPlaybackAssetRemoved(const FSoftObjectPath& InAssetPath);
-	void OnPlayableSequenceEvent(UAvaPlayable* InPlayable, const FName& SequenceName, EAvaPlayableSequenceEventType InEventType);
+	void OnPlayableSequenceEvent(UAvaPlayable* InPlayable, FName InSequenceLabel, EAvaPlayableSequenceEventType InEventType);
 	
 	void ApplyAvaMediaSettings();
 
diff --git a/Engine/Plugins/Experimental/Avalanche/Source/AvalancheMedia/Private/Rundown/AvaRundown.cpp b/Engine/Plugins/Experimental/Avalanche/Source/AvalancheMedia/Private/Rundown/AvaRundown.cpp
--- a/Engine/Plugins/Experimental/Avalanche/Source/AvalancheMedia/Private/Rundown/AvaRundown.cpp
+++ b/Engine/Plugins/Experimental/Avalanche/Source/AvalancheMedia/Private/Rundown/AvaRundown.cpp
@@ -1044,7 +1044,7 @@ TArray<int32> UAvaRundown::PlayPages(const TArray<int32>& InPageIds, EAvaRundown
 		{
 			const bool bIsPreview = UE::AvaRundown::IsPreviewPlayType(InPlayType);
 
-			if (!IsChannelTypeCompatibleForRequest(SelectedPage, bIsPreview, InPreviewChannelName, true))
+			if (!IsChannelTypeCompatibleForRequest(SelectedPage, bIsPreview, InPreviewChannelName, true) || !CanPlayPage(PageId, bIsPreview))
 			{
 				continue;
 			}
@@ -2209,6 +2209,18 @@ FAvaRundownPage& UAvaRundown::GetNextFromSubList(const TArray<int32>& InSubListI
 	return FAvaRundownPage::NullPage;
 }
 
+UAvaRundownPageTransition* UAvaRundown::GetPageTransition(const FGuid& InTransitionId) const
+{
+	for (UAvaRundownPageTransition* PageTransition : PageTransitions)
+	{
+		if (PageTransition && PageTransition->GetTransitionId() == InTransitionId)
+		{
+			return PageTransition;
+		}
+	}
+	return nullptr;
+}
+
 bool UAvaRundown::CanStartTransitionForPage(const FAvaRundownPage& InPage, bool bInIsPreview, const FName& InPreviewChannelName) const
 {
 	// Current constraint: There can only be one transition (running properly) at a time in a world.
@@ -2216,7 +2228,7 @@ bool UAvaRundown::CanStartTransitionForPage(const FAvaRundownPage& InPage, bool
 	// we can equate a "channel" to a "world", this is hardcoded for the level streaming playables.
 	// So, we can just check the channels for now.
 	const FName ChannelName = bInIsPreview ? InPreviewChannelName : InPage.GetChannelName();
-	for (const TObjectPtr<UAvaRundownPageTransition>& PageTransition : PageTransitions)
+	for (const UAvaRundownPageTransition* PageTransition : PageTransitions)
 	{
 		if (PageTransition && PageTransition->GetChannelName() == ChannelName)
 		{
@@ -2234,7 +2246,7 @@ void UAvaRundown::StopPageTransitionsForPage(const FAvaRundownPage& InPage, bool
 	// Note: we build a separate list because stopping the transitions should
 	// lead to the transitions being removed from PageTransitions (through the events).
 	const FName ChannelName = bInIsPreview ? InPreviewChannelName : InPage.GetChannelName();
-	for (TObjectPtr<UAvaRundownPageTransition>& PageTransition : PageTransitions)
+	for (UAvaRundownPageTransition* PageTransition : PageTransitions)
 	{
 		if (PageTransition && PageTransition->GetChannelName() == ChannelName)
 		{
@@ -2274,6 +2286,15 @@ UAvaRundownPagePlayer* UAvaRundown::FindPlayerForPreviewPage(int32 InPageId, con
 	return FoundPlayer ? *FoundPlayer : nullptr;
 }
 
+UAvaRundownPagePlayer* UAvaRundown::FindPagePlayer(int32 InPageId, FName InChannelName) const
+{
+	const TObjectPtr<UAvaRundownPagePlayer>* FoundPlayer = PagePlayers.FindByPredicate([InPageId, InChannelName](const UAvaRundownPagePlayer* InPagePlayer)
+	{
+		return InPagePlayer->PageId == InPageId && InPagePlayer->ChannelFName == InChannelName;
+	});
+	return FoundPlayer ? *FoundPlayer : nullptr;
+}
+
 void UAvaRundown::RemoveStoppedPagePlayers()
 {
 	for (UAvaRundownPagePlayer* PagePlayer : PagePlayers)
diff --git a/Engine/Plugins/Experimental/Avalanche/Source/AvalancheMedia/Private/Rundown/AvaRundownPagePlayer.cpp b/Engine/Plugins/Experimental/Avalanche/Source/AvalancheMedia/Private/Rundown/AvaRundownPagePlayer.cpp
--- a/Engine/Plugins/Experimental/Avalanche/Source/AvalancheMedia/Private/Rundown/AvaRundownPagePlayer.cpp
+++ b/Engine/Plugins/Experimental/Avalanche/Source/AvalancheMedia/Private/Rundown/AvaRundownPagePlayer.cpp
@@ -401,7 +401,7 @@ void UAvaRundownPagePlayer::RemoveInstancePlayer(UAvaRundownPlaybackInstancePlay
 	InstancePlayers.Remove(InInstancePlayer);
 }
 
-void UAvaRundownPagePlayer::HandleOnPlayableSequenceEvent(UAvaPlayable* InPlayable, const FName& SequenceName, EAvaPlayableSequenceEventType InEventType)
+void UAvaRundownPagePlayer::HandleOnPlayableSequenceEvent(UAvaPlayable* InPlayable, FName InSequenceLabel, EAvaPlayableSequenceEventType InEventType)
 {
 	// Check that this is the playable for this page player.
 	if (!HasPlayable(InPlayable))
@@ -415,12 +415,12 @@ void UAvaRundownPagePlayer::HandleOnPlayableSequenceEvent(UAvaPlayable* InPlayab
 		using namespace UE::AvaPlayback::Utils;
 		if (InEventType == EAvaPlayableSequenceEventType::Started)
 		{
-			UE_LOG(LogAvaRundown, Verbose, TEXT("%s Rundown Page %d: Sequence Started \"%s\"."), *GetBriefFrameInfo(), PageId, *SequenceName.ToString());
+			UE_LOG(LogAvaRundown, Verbose, TEXT("%s Rundown Page %d: Sequence Started \"%s\"."), *GetBriefFrameInfo(), PageId, *InSequenceLabel.ToString());
 		}
 
 		if (InEventType == EAvaPlayableSequenceEventType::Finished)
 		{
-			UE_LOG(LogAvaRundown, Verbose, TEXT("%s Rundown Page %d: Sequence Finished \"%s\"."), *GetBriefFrameInfo(), PageId, *SequenceName.ToString());
+			UE_LOG(LogAvaRundown, Verbose, TEXT("%s Rundown Page %d: Sequence Finished \"%s\"."), *GetBriefFrameInfo(), PageId, *InSequenceLabel.ToString());
 			RundownWeak->NotifyPageSequenceFinished(PageId);
 		}
 	}
diff --git a/Engine/Plugins/Experimental/Avalanche/Source/AvalancheMedia/Private/Rundown/AvaRundownServer.cpp b/Engine/Plugins/Experimental/Avalanche/Source/AvalancheMedia/Private/Rundown/AvaRundownServer.cpp
--- a/Engine/Plugins/Experimental/Avalanche/Source/AvalancheMedia/Private/Rundown/AvaRundownServer.cpp
+++ b/Engine/Plugins/Experimental/Avalanche/Source/AvalancheMedia/Private/Rundown/AvaRundownServer.cpp
@@ -25,6 +25,7 @@
 #include "MessageEndpointBuilder.h"
 #include "Misc/FileHelper.h"
 #include "Misc/Paths.h"
+#include "Playable/Transition/AvaPlayableTransition.h"
 #include "Playback/AvaPlaybackManager.h"
 #include "Playback/AvaPlaybackUtils.h"
 #include "RemoteControlSettings.h"
@@ -36,6 +37,7 @@
 #include "Rundown/AvaRundownSerializationUtils.h"
 #include "Rundown/AvaRundownServerMediaOutputUtils.h"
 #include "TextureResource.h"
+#include "Transition/AvaRundownPageTransition.h"
 
 #if WITH_EDITOR
 #include "Editor.h"
@@ -277,6 +279,29 @@ namespace UE::AvaRundownServer::Private
 		}
 		return bSuccess;
 	}
+
+	TArray<int32> GetPageIds(TConstArrayView<TWeakObjectPtr<UAvaRundownPagePlayer>> InPagePlayers)
+	{
+		TArray<int32> PageIds;
+		PageIds.Reserve(InPagePlayers.Num());
+		for (const TWeakObjectPtr<UAvaRundownPagePlayer>& PagePlayerWeak : InPagePlayers)
+		{
+			if (const UAvaRundownPagePlayer* PagePlayer = PagePlayerWeak.Get())
+			{
+				PageIds.Add(PagePlayer->PageId);
+			}
+		}
+		return PageIds;
+	}
+
+	void FillPageTransitionInfo(const UAvaRundownPageTransition& InPageTransition, FAvaRundownPageTransitionEvent& OutMessage)
+	{
+		OutMessage.Channel = InPageTransition.GetChannelName().ToString();
+		OutMessage.TransitionId = InPageTransition.GetTransitionId();
+		OutMessage.EnteringPageIds = GetPageIds(InPageTransition.GetEnterPlayers());
+		OutMessage.PlayingPageIds = GetPageIds(InPageTransition.GetPlayingPlayers());
+		OutMessage.ExitingPageIds = GetPageIds(InPageTransition.GetExitPlayers());
+	}
 }
 
 /**
@@ -412,8 +437,9 @@ void FAvaRundownServer::Init(const FString& InAssignedHostName)
 	.Handling<FAvaRundownGetChannelQualitySettings>(this, &FAvaRundownServer::HandleGetChannelQualitySettings)
 	.Handling<FAvaRundownSetChannelQualitySettings>(this, &FAvaRundownServer::HandleSetChannelQualitySettings)
 	.Handling<FAvaRundownSaveBroadcast>(this, &FAvaRundownServer::HandleSaveBroadcast)
-	.NotificationHandling(FOnBusNotification::CreateRaw(this, &FAvaRundownServer::OnMessageBusNotification));
-	
+	.NotificationHandling(FOnBusNotification::CreateSP(this, &FAvaRundownServer::OnMessageBusNotification))
+	.ReceivingOnThread(ENamedThreads::GameThread);
+
 	if (MessageEndpoint.IsValid())
 	{
 		// Subscribe to the server listing requests
@@ -431,6 +457,9 @@ void FAvaRundownServer::SetupPlaybackDelegates()
 {
 	FAvaPlaybackManager& Manager = IAvaMediaModule::Get().GetLocalPlaybackManager();
 	Manager.OnPlaybackInstanceStatusChanged.AddSP(this, &FAvaRundownServer::OnPlaybackInstanceStatusChanged);
+
+	UAvaPlayable::OnSequenceEvent().AddSP(this, &FAvaRundownServer::OnPlayableSequenceEvent);
+	UAvaPlayable::OnTransitionEvent().AddSP(this, &FAvaRundownServer::OnPlayableTransitionEvent);
 }
 
 void FAvaRundownServer::SetupBroadcastDelegates(UAvaBroadcast* InBroadcast)
@@ -459,6 +488,9 @@ void FAvaRundownServer::RemovePlaybackDelegates() const
 		FAvaPlaybackManager& Manager = AvaMediaModule.GetLocalPlaybackManager();
 		Manager.OnPlaybackInstanceStatusChanged.RemoveAll(this);
 	}
+
+	UAvaPlayable::OnSequenceEvent().RemoveAll(this);
+	UAvaPlayable::OnTransitionEvent().RemoveAll(this);
 }
 
 void FAvaRundownServer::RemoveBroadcastDelegates(UAvaBroadcast* InBroadcast) const
@@ -2408,6 +2440,7 @@ FAvaRundownServer::FRundownEntry::FRundownEntry(const TSharedPtr<FAvaRundownServ
 		Rundown->GetOnPagesChanged().AddSP(RundownServerRef, &FAvaRundownServer::OnPagesChanged);
 		Rundown->GetOnPageListChanged().AddSP(RundownServerRef, &FAvaRundownServer::OnPageListChanged);
 		Rundown->GetOnCanClosePlaybackContext().AddSP(RundownServerRef, &FAvaRundownServer::OnCanClosePlaybackContext);
+		Rundown->GetOnPageTransitionRemoving().AddSP(RundownServerRef, &FAvaRundownServer::OnPageTransitionRemoved);
 	}
 }
 
@@ -2418,6 +2451,7 @@ FAvaRundownServer::FRundownEntry::~FRundownEntry()
 		Rundown->GetOnPagesChanged().RemoveAll(RundownServerRaw);
 		Rundown->GetOnPageListChanged().RemoveAll(RundownServerRaw);
 		Rundown->GetOnCanClosePlaybackContext().RemoveAll(RundownServerRaw);
+		Rundown->GetOnPageTransitionRemoving().RemoveAll(RundownServerRaw);
 	}
 }
 
@@ -2538,6 +2572,91 @@ void FAvaRundownServer::OnPlaybackInstanceStatusChanged(const FAvaPlaybackInstan
 	}
 }
 
+void FAvaRundownServer::OnPlayableSequenceEvent(UAvaPlayable* InPlayable, FName InSequenceLabel, EAvaPlayableSequenceEventType InSequenceEvent)
+{
+	if (!InPlayable || !InPlayable->GetPlayableGroup())
+	{
+		return;
+	}
+	
+	const UAvaRundown* CurrentPlaybackRundown = PlaybackCommandContext.GetCurrentRundown();
+	if (!CurrentPlaybackRundown)
+	{
+		return; // Not an event from current playback rundown.
+	}
+
+	const FName ChannelName = InPlayable->GetPlayableGroup()->GetChannelName();
+	const int32 PageId = UAvaRundownPagePlayer::GetPageIdFromInstanceUserData(InPlayable->GetUserData()); 
+
+	if (!CurrentPlaybackRundown->FindPagePlayer(PageId, ChannelName))
+	{
+		return; // Not a playable from current playback rundown.
+	}
+	
+	FAvaRundownPageSequenceEvent* Message = FMessageEndpoint::MakeMessage<FAvaRundownPageSequenceEvent>();
+	Message->Channel = ChannelName.ToString();
+	Message->PageId = PageId;
+	Message->InstanceId = InPlayable->GetInstanceId();
+	Message->AssetPath = InPlayable->GetSourceAssetPath().ToString();
+	Message->SequenceLabel = InSequenceLabel.ToString();
+	Message->Event = InSequenceEvent;
+
+	SendResponse(Message, ClientAddresses);
+}
+
+void FAvaRundownServer::OnPlayableTransitionEvent(UAvaPlayable* InPlayable, UAvaPlayableTransition* InPlayableTransition, EAvaPlayableTransitionEventFlags InTransitionFlags)
+{
+	const UAvaRundown* CurrentPlaybackRundown = PlaybackCommandContext.GetCurrentRundown();
+	if (!CurrentPlaybackRundown)
+	{
+		return; // Not an event from current playback rundown.
+	}
+	
+	if (!EnumHasAnyFlags(InTransitionFlags, EAvaPlayableTransitionEventFlags::Finished | EAvaPlayableTransitionEventFlags::Starting))
+	{
+		return; // Not interested;
+	}
+
+	// Note: Page Transition can already be removed from Rundown, in that case event is propagated from the OnPageTransitionRemoved callback.
+	const UAvaRundownPageTransition* PageTransition = CurrentPlaybackRundown->GetPageTransition(InPlayableTransition->GetTransitionId()); 
+	if (!PageTransition)
+	{
+		return; // Not transition from current rundown.
+	}
+	
+	FAvaRundownPageTransitionEvent* Message = FMessageEndpoint::MakeMessage<FAvaRundownPageTransitionEvent>();
+
+	UE::AvaRundownServer::Private::FillPageTransitionInfo(*PageTransition, *Message);
+
+	if (EnumHasAnyFlags(InTransitionFlags, EAvaPlayableTransitionEventFlags::Finished))
+	{
+		Message->Event = EAvaRundownPageTransitionEvents::Finished;
+	}
+	else if (EnumHasAnyFlags(InTransitionFlags, EAvaPlayableTransitionEventFlags::Starting))
+	{
+		Message->Event = EAvaRundownPageTransitionEvents::Started;
+	}
+
+	SendResponse(Message, ClientAddresses);
+}
+
+void FAvaRundownServer::OnPageTransitionRemoved(UAvaRundown* InRundown, UAvaRundownPageTransition* InPageTransition)
+{
+	const UAvaRundown* CurrentPlaybackRundown = PlaybackCommandContext.GetCurrentRundown();
+	
+	if (!CurrentPlaybackRundown || InRundown != CurrentPlaybackRundown || !InPageTransition)
+	{
+		return; // Not an event from current playback rundown.
+	}
+
+	// Note: if the page transition is removed from the rundown, it indicates the transition is finished.
+	// It can be received before the "playable" transition event because the order of event handlers is not guaranteed.
+	FAvaRundownPageTransitionEvent* Message = FMessageEndpoint::MakeMessage<FAvaRundownPageTransitionEvent>();
+	UE::AvaRundownServer::Private::FillPageTransitionInfo(*InPageTransition, *Message);
+	Message->Event = EAvaRundownPageTransitionEvents::Finished;
+	SendResponse(Message, ClientAddresses);
+}
+
 void FAvaRundownServer::OnCanClosePlaybackContext(const UAvaRundown* InRundown, bool& bOutResult) const
 {
 	const UAvaRundown* CurrentPlaybackRundown = PlaybackCommandContext.GetCurrentRundown();
diff --git a/Engine/Plugins/Experimental/Avalanche/Source/AvalancheMedia/Private/Rundown/AvaRundownServer.h b/Engine/Plugins/Experimental/Avalanche/Source/AvalancheMedia/Private/Rundown/AvaRundownServer.h
--- a/Engine/Plugins/Experimental/Avalanche/Source/AvalancheMedia/Private/Rundown/AvaRundownServer.h
+++ b/Engine/Plugins/Experimental/Avalanche/Source/AvalancheMedia/Private/Rundown/AvaRundownServer.h
@@ -58,6 +58,9 @@ public:
 	void OnAssetRemoved(const FAssetData& InAssetData) const;
 	void OnAssetsPreDelete(const TArray<UObject*>& InObjects);
 	void OnPlaybackInstanceStatusChanged(const FAvaPlaybackInstance& InPlaybackInstance);
+	void OnPlayableSequenceEvent(UAvaPlayable* InPlayable, FName InSequenceLabel, EAvaPlayableSequenceEventType InSequenceEvent);
+	void OnPlayableTransitionEvent(UAvaPlayable* InPlayable, UAvaPlayableTransition* InPlayableTransition, EAvaPlayableTransitionEventFlags InTransitionFlags);
+	void OnPageTransitionRemoved(UAvaRundown* InRundown, UAvaRundownPageTransition* InPageTransition);
 	void OnCanClosePlaybackContext(const UAvaRundown* InRundown, bool& bOutResult) const;
 	
 	// Rundown message handlers
diff --git a/Engine/Plugins/Experimental/Avalanche/Source/AvalancheMedia/Private/Rundown/Transition/AvaRundownPageTransition.h b/Engine/Plugins/Experimental/Avalanche/Source/AvalancheMedia/Private/Rundown/Transition/AvaRundownPageTransition.h
--- a/Engine/Plugins/Experimental/Avalanche/Source/AvalancheMedia/Private/Rundown/Transition/AvaRundownPageTransition.h
+++ b/Engine/Plugins/Experimental/Avalanche/Source/AvalancheMedia/Private/Rundown/Transition/AvaRundownPageTransition.h
@@ -32,6 +32,10 @@ public:
 	bool AddEnterPage(UAvaRundownPagePlayer* InPagePlayer);
 	bool AddPlayingPage(UAvaRundownPagePlayer* InPagePlayer);
 	bool AddExitPage(UAvaRundownPagePlayer* InPagePlayer);
+
+	TConstArrayView<TWeakObjectPtr<UAvaRundownPagePlayer>> GetEnterPlayers() const { return EnterPlayersWeak; }
+	TConstArrayView<TWeakObjectPtr<UAvaRundownPagePlayer>> GetPlayingPlayers() const { return PlayingPlayersWeak; }
+	TConstArrayView<TWeakObjectPtr<UAvaRundownPagePlayer>> GetExitPlayers() const { return ExitPlayersWeak; }
 	
 	//~ Begin IAvaPlayableVisibilityConstraint
 	virtual bool IsVisibilityConstrained(const UAvaPlayable* InPlayable) const override;
diff --git a/Engine/Plugins/Experimental/Avalanche/Source/AvalancheMedia/Public/AvaMediaDefines.h b/Engine/Plugins/Experimental/Avalanche/Source/AvalancheMedia/Public/AvaMediaDefines.h
--- a/Engine/Plugins/Experimental/Avalanche/Source/AvalancheMedia/Public/AvaMediaDefines.h
+++ b/Engine/Plugins/Experimental/Avalanche/Source/AvalancheMedia/Public/AvaMediaDefines.h
@@ -115,6 +115,7 @@ enum class EAvaPlayableSequenceEventType : uint8
 {
 	None,
 	Started,
+	Paused,
 	Finished
 };
 
diff --git a/Engine/Plugins/Experimental/Avalanche/Source/AvalancheMedia/Public/Playable/AvaPlayable.h b/Engine/Plugins/Experimental/Avalanche/Source/AvalancheMedia/Public/Playable/AvaPlayable.h
--- a/Engine/Plugins/Experimental/Avalanche/Source/AvalancheMedia/Public/Playable/AvaPlayable.h
+++ b/Engine/Plugins/Experimental/Avalanche/Source/AvalancheMedia/Public/Playable/AvaPlayable.h
@@ -100,7 +100,7 @@ public:
 	 */
 	static UAvaPlayable* Create(UObject* InOuter, const FPlayableCreationInfo& InPlayableInfo);
 
-	DECLARE_MULTICAST_DELEGATE_ThreeParams(FOnSequenceEvent, UAvaPlayable*, const FName& /*InSequenceName*/, EAvaPlayableSequenceEventType);
+	DECLARE_MULTICAST_DELEGATE_ThreeParams(FOnSequenceEvent, UAvaPlayable*, FName /*InSequenceLabel*/, EAvaPlayableSequenceEventType);
 	static FOnSequenceEvent& OnSequenceEvent() { return OnSequenceEventDelegate; }
 
 	DECLARE_MULTICAST_DELEGATE_ThreeParams(FOnTransitionEvent, UAvaPlayable*, UAvaPlayableTransition*, EAvaPlayableTransitionEventFlags);
@@ -172,6 +172,7 @@ protected:
 	virtual void OnRemoteControlValuesApplied() {}
 	
 	void HandleOnSequenceStarted(UAvaSequencePlayer* InSequencePlayer, UAvaSequence* InSequence);
+	void HandleOnSequencePaused(UAvaSequencePlayer* InSequencePlayer, UAvaSequence* InSequence);
 	void HandleOnSequenceFinished(UAvaSequencePlayer* InSequencePlayer, UAvaSequence* InSequence);
 
 	static UAvaPlayable* CreateLocalPlayable(UObject* InOuter, const FPlayableCreationInfo& InPlayableInfo);
diff --git a/Engine/Plugins/Experimental/Avalanche/Source/AvalancheMedia/Public/Playable/Transition/AvaPlayableTransition.h b/Engine/Plugins/Experimental/Avalanche/Source/AvalancheMedia/Public/Playable/Transition/AvaPlayableTransition.h
--- a/Engine/Plugins/Experimental/Avalanche/Source/AvalancheMedia/Public/Playable/Transition/AvaPlayableTransition.h
+++ b/Engine/Plugins/Experimental/Avalanche/Source/AvalancheMedia/Public/Playable/Transition/AvaPlayableTransition.h
@@ -55,6 +55,8 @@ public:
 	virtual FString GetPrettyInfo() const;
 
 	EAvaPlayableTransitionFlags GetTransitionFlags() const { return TransitionFlags; }
+
+	const FGuid& GetTransitionId() const { return TransitionId; }
 	
 protected:
 	UAvaPlayable* FindPlayable(const FGuid& InstanceId) const;
diff --git a/Engine/Plugins/Experimental/Avalanche/Source/AvalancheMedia/Public/Rundown/AvaRundown.h b/Engine/Plugins/Experimental/Avalanche/Source/AvalancheMedia/Public/Rundown/AvaRundown.h
--- a/Engine/Plugins/Experimental/Avalanche/Source/AvalancheMedia/Public/Rundown/AvaRundown.h
+++ b/Engine/Plugins/Experimental/Avalanche/Source/AvalancheMedia/Public/Rundown/AvaRundown.h
@@ -717,17 +717,36 @@ public:
 
 	FOnPagePlayerEvent& GetOnPagePlayerAdded() { return OnPagePlayerAdded; }
 	FOnPagePlayerEvent& GetOnPagePlayerRemoving() { return OnPagePlayerRemoving; }
+
+	DECLARE_MULTICAST_DELEGATE_TwoParams(FOnPageTransitionEvent, UAvaRundown*, UAvaRundownPageTransition*);
+
+	/**
+	 * Rundown Transition Event - Called when a new page transition is added to the playback context.
+	 */
+	FOnPageTransitionEvent OnPageTransitionAdded;
+
+	/**
+	 * Rundown Transition Event - Called when a completed page transition is about to be removed from the playback context.
+	 */
+	FOnPageTransitionEvent OnPageTransitionRemoving;
+
+	FOnPageTransitionEvent& GetOnPageTransitionAdded() { return OnPageTransitionAdded; }
+	FOnPageTransitionEvent& GetOnPageTransitionRemoving() { return OnPageTransitionRemoving; }
 	
 	void AddPageTransition(UAvaRundownPageTransition* InPageTransition)
 	{
 		PageTransitions.Add(InPageTransition);
+		OnPageTransitionAdded.Broadcast(this, InPageTransition);
 	}
 
 	void RemovePageTransition(UAvaRundownPageTransition* InPageTransition)
 	{
+		OnPageTransitionRemoving.Broadcast(this, InPageTransition);
 		PageTransitions.Remove(InPageTransition);
 	}
 
+	UAvaRundownPageTransition* GetPageTransition(const FGuid& InTransitionId) const;
+
 	bool CanStartTransitionForPage(const FAvaRundownPage& InPage, bool bInIsPreview, const FName& InPreviewChannelName) const;
 
 	void StopPageTransitionsForPage(const FAvaRundownPage& InPage, bool bInIsPreview, const FName& InPreviewChannelName);
@@ -743,6 +762,9 @@ public:
 		return bInIsPreview ?  FindPlayerForPreviewPage(InPageId, InPreviewChannelName) : FindPlayerForProgramPage(InPageId);
 	}
 
+	/** Finds the page player for the given pageId and channel name. */
+	UAvaRundownPagePlayer* FindPagePlayer(int32 InPageId, FName InChannelName) const;
+
 	void RemoveStoppedPagePlayers();
 
 	FAvaRundownPageListPlaybackContextCollection* GetPageListPlaybackContextCollection() const
diff --git a/Engine/Plugins/Experimental/Avalanche/Source/AvalancheMedia/Public/Rundown/AvaRundownPagePlayer.h b/Engine/Plugins/Experimental/Avalanche/Source/AvalancheMedia/Public/Rundown/AvaRundownPagePlayer.h
--- a/Engine/Plugins/Experimental/Avalanche/Source/AvalancheMedia/Public/Rundown/AvaRundownPagePlayer.h
+++ b/Engine/Plugins/Experimental/Avalanche/Source/AvalancheMedia/Public/Rundown/AvaRundownPagePlayer.h
@@ -161,7 +161,7 @@ protected:
 
 	void RemoveInstancePlayer(UAvaRundownPlaybackInstancePlayer* InInstancePlayer);
 	
-	void HandleOnPlayableSequenceEvent(UAvaPlayable* InPlayable, const FName& SequenceName, EAvaPlayableSequenceEventType InEventType);
+	void HandleOnPlayableSequenceEvent(UAvaPlayable* InPlayable, FName InSequenceLabel, EAvaPlayableSequenceEventType InEventType);
 
 public:
 	UPROPERTY()
diff --git a/Engine/Plugins/Experimental/Avalanche/Source/AvalancheSequence/Private/AvaSequencePlayer.cpp b/Engine/Plugins/Experimental/Avalanche/Source/AvalancheSequence/Private/AvaSequencePlayer.cpp
--- a/Engine/Plugins/Experimental/Avalanche/Source/AvalancheSequence/Private/AvaSequencePlayer.cpp
+++ b/Engine/Plugins/Experimental/Avalanche/Source/AvalancheSequence/Private/AvaSequencePlayer.cpp
@@ -11,12 +11,18 @@
 #include "IAvaSequenceProvider.h"
 
 UAvaSequencePlayer::FOnSequenceEvent UAvaSequencePlayer::OnSequenceStartedDelegate;
+UAvaSequencePlayer::FOnSequenceEvent UAvaSequencePlayer::OnSequencePausedDelegate;
 UAvaSequencePlayer::FOnSequenceEvent UAvaSequencePlayer::OnSequenceFinishedDelegate;
 
 UAvaSequencePlayer::UAvaSequencePlayer(const FObjectInitializer& InObjectInitializer)
 	: ULevelSequencePlayer(InObjectInitializer)
 {
-	OnNativeFinished.BindUObject(this, &UAvaSequencePlayer::NotifySequenceFinished);
+	if (!IsTemplate())
+	{
+		OnNativeFinished.BindUObject(this, &UAvaSequencePlayer::NotifySequenceFinished);
+		// Remark: UMovieSceneSequencePlayer has a virtual OnPaused() function, but it is not called. Using event instead.
+		OnPause.AddDynamic(this, &UAvaSequencePlayer::NotifySequencePaused);
+	}
 }
 
 void UAvaSequencePlayer::InitSequence(UAvaSequence* InSequence, IAvaSequencePlaybackObject* InPlaybackObject, ULevel* InLevel)
@@ -212,6 +218,11 @@ void UAvaSequencePlayer::NotifySequenceStarted()
 	OnSequenceStartedDelegate.Broadcast(this, GetAvaSequence());
 }
 
+void UAvaSequencePlayer::NotifySequencePaused()
+{
+	OnSequencePausedDelegate.Broadcast(this, GetAvaSequence());
+}
+
 void UAvaSequencePlayer::NotifySequenceFinished()
 {
 	OnSequenceFinishedDelegate.Broadcast(this, GetAvaSequence());
diff --git a/Engine/Plugins/Experimental/Avalanche/Source/AvalancheSequence/Public/AvaSequencePlayer.h b/Engine/Plugins/Experimental/Avalanche/Source/AvalancheSequence/Public/AvaSequencePlayer.h
--- a/Engine/Plugins/Experimental/Avalanche/Source/AvalancheSequence/Public/AvaSequencePlayer.h
+++ b/Engine/Plugins/Experimental/Avalanche/Source/AvalancheSequence/Public/AvaSequencePlayer.h
@@ -24,6 +24,7 @@ public:
 
 	DECLARE_MULTICAST_DELEGATE_TwoParams(FOnSequenceEvent, UAvaSequencePlayer*, UAvaSequence*);
 	static FOnSequenceEvent& OnSequenceStarted() { return OnSequenceStartedDelegate; }
+	static FOnSequenceEvent& OnSequencePaused() { return OnSequencePausedDelegate; }
 	static FOnSequenceEvent& OnSequenceFinished() { return OnSequenceFinishedDelegate; }
 
 	void InitSequence(UAvaSequence* InSequence, IAvaSequencePlaybackObject* InPlaybackObject, ULevel* InLevel);
@@ -75,9 +76,14 @@ private:
 	FFrameTime CalculateDeltaFrameTime(float InDeltaSeconds) const;
 
 	void NotifySequenceStarted();
+
+	UFUNCTION()
+	void NotifySequencePaused();
+
 	void NotifySequenceFinished();
 
 	static FOnSequenceEvent OnSequenceStartedDelegate;
+	static FOnSequenceEvent OnSequencePausedDelegate;
 	static FOnSequenceEvent OnSequenceFinishedDelegate;
 
 	/** Encapsulate OnNativeFinished as it's a simple delegate that only this Player will bind to and through it call the multicast version (OnSequenceFinishedDelegate) */
