[Motion Design] Rundown Managed Instance is using a LoadPackageAsync (in blocking mode still) instead of StaticDuplicateObject.

This avoids a number of issues with StaticDuplicateObject (especially in game builds) for instancing levels.
Also, in an upcoming change, the Managed Instances will be loaded async and non blocking to keep the editor responsive, so this is a step in that direction.

diff --git a/Engine/Plugins/Experimental/Avalanche/Source/AvalancheMedia/Private/Rundown/AvaRundownManagedInstanceLevel.cpp b/Engine/Plugins/Experimental/Avalanche/Source/AvalancheMedia/Private/Rundown/AvaRundownManagedInstanceLevel.cpp
--- a/Engine/Plugins/Experimental/Avalanche/Source/AvalancheMedia/Private/Rundown/AvaRundownManagedInstanceLevel.cpp
+++ b/Engine/Plugins/Experimental/Avalanche/Source/AvalancheMedia/Private/Rundown/AvaRundownManagedInstanceLevel.cpp
@@ -7,8 +7,10 @@
 #include "AvaRemoteControlUtils.h"
 #include "AvaScene.h"
 #include "Engine/Engine.h"
+#include "Engine/Level.h"
 #include "Engine/World.h"
 #include "IAvaMediaModule.h"
+#include "Playback/AvaPlaybackUtils.h"
 #include "RemoteControlPreset.h"
 #include "Rundown/AvaRundownManagedInstanceUtils.h"
 #include "UObject/Package.h"
@@ -20,6 +22,53 @@ namespace UE::AvaRundownManagedInstanceLevel::Private
 		return Cast<UWorld>(InAssetPath.TryLoad());
 	}
 
+	/** Loads the given source level in the given package. */
+	UWorld* LoadLevelInstanceInPackage(const FSoftObjectPath& InSourceAssetPath, UPackage* InDestinationPackage)
+	{
+		if (!InDestinationPackage)
+		{
+			return nullptr;
+		}
+		
+		const FPackagePath SourcePackagePath = FPackagePath::FromPackageNameUnchecked(InSourceAssetPath.GetLongPackageFName());
+		constexpr EPackageFlags PackageFlags = PKG_ContainsMap;
+		const FLinkerInstancingContext* InstancingContextPtr = nullptr;
+		const FName ManagedPackageName = InDestinationPackage->GetFName();
+#if WITH_EDITOR
+		FLinkerInstancingContext InstancingContext;
+
+		// When loading an instanced package we need to invoke an instancing context function in case non external actors
+		// part of the level are pulling on external actors.
+		const FString ExternalActorsPathStr = ULevel::GetExternalActorsPath(SourcePackagePath.GetPackageName());
+		const FString DesiredPackageNameStr = ManagedPackageName.ToString();
+
+		InstancingContext.AddPackageMappingFunc([ExternalActorsPathStr, DesiredPackageNameStr](FName Original)
+		{
+			const FString OriginalStr = Original.ToString();
+			if (OriginalStr.StartsWith(ExternalActorsPathStr))
+			{
+				return FName(*ULevel::GetExternalActorPackageInstanceName(DesiredPackageNameStr, OriginalStr));
+			}
+			return Original;
+		});
+
+		InstancingContextPtr = &InstancingContext;
+#endif
+
+		// Since we are going to block on it, make sure it is high priority.
+		constexpr int32 LoadPriority = MAX_int32;
+
+		const int32 LocalRequestId = LoadPackageAsync(SourcePackagePath, ManagedPackageName,
+			FLoadPackageAsyncDelegate(), PackageFlags, INDEX_NONE, LoadPriority, InstancingContextPtr);
+
+		FlushAsyncLoading(LocalRequestId);
+
+		// Workaround to destroy the Linker Load so that it does not keep the underlying File Opened
+		FAvaPlaybackUtils::FlushPackageLoading(InDestinationPackage);
+
+		return UWorld::FindWorldInPackage(InDestinationPackage);
+	}
+
 	URemoteControlPreset* FindRemoteControlPreset(ULevel* InLevel)
 	{
 		const AAvaScene* AvaScene = AAvaScene::GetScene(InLevel, false);
@@ -31,19 +80,15 @@ FAvaRundownManagedInstanceLevel::FAvaRundownManagedInstanceLevel(FAvaRundownMana
 	: FAvaRundownManagedInstance(InParentCache, InAssetPath)
 {
 	using namespace UE::AvaRundownManagedInstanceLevel::Private;
-	// We need to load the source Motion Design level.
-	UWorld* SourceLevel = LoadLevel(InAssetPath);
-	if (!SourceLevel)
+	
+	// Register the delegates on the source RCP (if loaded) in case the level is currently being edited.
+	if (UWorld* SourceLevel = Cast<UWorld>(InAssetPath.ResolveObject()))
 	{
-		UE_LOG(LogAvaMedia, Error, TEXT("Unable to load Source Motion Design Level: %s"), *InAssetPath.ToString());
-		return;
-	}
+		// Keep a weak pointer
+		SourceLevelWeak = SourceLevel;
 
-	// Keep a weak pointer
-	SourceLevelWeak = SourceLevel;
-
-	// Register the delegates on the source RCP just in case the level is currently being edited.
-	RegisterSourceRemoteControlPresetDelegates(FindRemoteControlPreset(SourceLevel->PersistentLevel));
+		RegisterSourceRemoteControlPresetDelegates(FindRemoteControlPreset(SourceLevel->PersistentLevel));
+	}
 
 	ManagedLevelPackage = FAvaRundownManagedInstanceUtils::MakeManagedInstancePackage(InAssetPath);
 	if (!ManagedLevelPackage)
@@ -52,10 +97,8 @@ FAvaRundownManagedInstanceLevel::FAvaRundownManagedInstanceLevel(FAvaRundownMana
 		return;
 	}
 
-	{
-		FRCPresetGuidRenewGuard PresetGuidRenewGuard;
-		ManagedLevel = Cast<UWorld>(StaticDuplicateObject(SourceLevel, ManagedLevelPackage.Get()));
-	}
+	// Load a copy of the source package.
+	ManagedLevel = LoadLevelInstanceInPackage(InAssetPath, ManagedLevelPackage);
 	
 	if (!ManagedLevel)
 	{
