MotionDesign : ActorModifierCore
- Fixed RadialArrange modifier defaults to match previous version

diff --git a/Engine/Plugins/Experimental/Avalanche/Source/AvalancheModifiers/Public/Modifiers/AvaRadialArrangeModifier.h b/Engine/Plugins/Experimental/Avalanche/Source/AvalancheModifiers/Public/Modifiers/AvaRadialArrangeModifier.h
--- a/Engine/Plugins/Experimental/Avalanche/Source/AvalancheModifiers/Public/Modifiers/AvaRadialArrangeModifier.h
+++ b/Engine/Plugins/Experimental/Avalanche/Source/AvalancheModifiers/Public/Modifiers/AvaRadialArrangeModifier.h
@@ -182,7 +182,7 @@ protected:
 
 	/** Base plane for the radial arrangement */
 	UPROPERTY(EditInstanceOnly, Setter, Getter, Category="RadialArrange", meta=(AllowPrivateAccess="true"))
-	EAvaRadialArrangePlane Plane = EAvaRadialArrangePlane::XY;
+	EAvaRadialArrangePlane Plane = EAvaRadialArrangePlane::YZ;
 
 	/** The number of child elements to limit in the arrangement, or -1 if unlimited. Children whose index is greater than or equal to this value will be hidden. */
 	UPROPERTY(EditInstanceOnly, Setter, Getter, Category="RadialArrange", meta=(ClampMin="-1", UIMin="-1", AllowPrivateAccess="true"))
