[Motion Design] PlayableGroup: Auto hide Pawn actors

- Adding automatic hiding of Pawn actors in the "container" world because we don't want it visible in rundown. A new settings is added to enable or disable it.
- Adding FAvaPlayableSettings. New playable related settings will be added to this. Replicated on server from client.
- Moved Playable related settings from FAvaInstanceSettings to FAvaPlayableSettings that where not yet released in 5.4.

diff --git a/Engine/Plugins/Experimental/Avalanche/Source/Avalanche/Public/Framework/AvaInstanceSettings.h b/Engine/Plugins/Experimental/Avalanche/Source/Avalanche/Public/Framework/AvaInstanceSettings.h
--- a/Engine/Plugins/Experimental/Avalanche/Source/Avalanche/Public/Framework/AvaInstanceSettings.h
+++ b/Engine/Plugins/Experimental/Avalanche/Source/Avalanche/Public/Framework/AvaInstanceSettings.h
@@ -4,19 +4,6 @@
 
 #include "AvaInstanceSettings.generated.h"
 
-USTRUCT()
-struct FAvaSynchronizedEventsFeatureSelection
-{
-	GENERATED_BODY()
-	
-	/**
-	 * Select the implementation for synchronizing events.
-	 * "Default" will select the most appropriate implementation available.
-	 */
-	UPROPERTY(Config, EditAnywhere, Category = Settings)
-	FString Implementation;
-};
-
 /**
  * Settings applied when instancing a Motion Design Asset for playback.
  */
@@ -35,10 +22,4 @@ struct FAvaInstanceSettings
 	 */
 	UPROPERTY(config, EditAnywhere, Category = Settings)
 	bool bDefaultPlayableTransitionWaitForSequences = false;
-
-	/**
-	 * Select the implementation for synchronizing events.
-	 */
-	UPROPERTY(Config, EditAnywhere, Category = Settings)
-	FAvaSynchronizedEventsFeatureSelection SynchronizedEventsFeature;
 };
diff --git a/Engine/Plugins/Experimental/Avalanche/Source/AvalancheMedia/Private/AvaMediaModule.cpp b/Engine/Plugins/Experimental/Avalanche/Source/AvalancheMedia/Private/AvaMediaModule.cpp
--- a/Engine/Plugins/Experimental/Avalanche/Source/AvalancheMedia/Private/AvaMediaModule.cpp
+++ b/Engine/Plugins/Experimental/Avalanche/Source/AvalancheMedia/Private/AvaMediaModule.cpp
@@ -400,6 +400,20 @@ const FAvaInstanceSettings& FAvaMediaModule::GetAvaInstanceSettings() const
 	return UAvaMediaSettings::Get().AvaInstanceSettings;
 }
 
+const FAvaPlayableSettings& FAvaMediaModule::GetPlayableSettings() const
+{
+	// if the server is enabled, fetch the settings from the currently connected client.
+	if (AvaPlaybackServer.IsValid())
+	{
+		if (const FAvaPlayableSettings* SettingsFromClient = AvaPlaybackServer->GetPlayableSettings())
+		{
+			return *SettingsFromClient;
+		}
+	}
+	// Return the local settings.
+	return UAvaMediaSettings::Get().PlayableSettings;
+}
+
 bool FAvaMediaModule::IsLocalPlaybackManagerAvailable() const
 {
 	return LocalPlaybackManager.IsValid(); 
diff --git a/Engine/Plugins/Experimental/Avalanche/Source/AvalancheMedia/Private/AvaMediaModule.h b/Engine/Plugins/Experimental/Avalanche/Source/AvalancheMedia/Private/AvaMediaModule.h
--- a/Engine/Plugins/Experimental/Avalanche/Source/AvalancheMedia/Private/AvaMediaModule.h
+++ b/Engine/Plugins/Experimental/Avalanche/Source/AvalancheMedia/Private/AvaMediaModule.h
@@ -44,6 +44,7 @@ public:
 	virtual bool IsGameModeLocalPlaybackServerLaunched() const override;
 	virtual const IAvaBroadcastSettings& GetBroadcastSettings() const override;
 	virtual const FAvaInstanceSettings& GetAvaInstanceSettings() const override;
+	virtual const FAvaPlayableSettings& GetPlayableSettings() const override;
 	virtual bool IsLocalPlaybackManagerAvailable() const override;
 	virtual FAvaPlaybackManager& GetLocalPlaybackManager() const override;
 	virtual bool IsManagedInstanceCacheAvailable() const override;
diff --git a/Engine/Plugins/Experimental/Avalanche/Source/AvalancheMedia/Private/AvaMediaSettings.cpp b/Engine/Plugins/Experimental/Avalanche/Source/AvalancheMedia/Private/AvaMediaSettings.cpp
--- a/Engine/Plugins/Experimental/Avalanche/Source/AvalancheMedia/Private/AvaMediaSettings.cpp
+++ b/Engine/Plugins/Experimental/Avalanche/Source/AvalancheMedia/Private/AvaMediaSettings.cpp
@@ -12,7 +12,7 @@ UAvaMediaSettings::UAvaMediaSettings()
 {
 	CategoryName = TEXT("Motion Design");
 	SectionName  = TEXT("Playback & Broadcast");
-	AvaInstanceSettings.SynchronizedEventsFeature.Implementation = SynchronizedEventsFeatureSelection_Default.ToString();
+	PlayableSettings.SynchronizedEventsFeature.Implementation = SynchronizedEventsFeatureSelection_Default.ToString();
 }
 
 UAvaMediaSettings* UAvaMediaSettings::GetSingletonInstance()
diff --git a/Engine/Plugins/Experimental/Avalanche/Source/AvalancheMedia/Private/ModularFeature/AvaMediaSynchronizedEventsFeature.cpp b/Engine/Plugins/Experimental/Avalanche/Source/AvalancheMedia/Private/ModularFeature/AvaMediaSynchronizedEventsFeature.cpp
--- a/Engine/Plugins/Experimental/Avalanche/Source/AvalancheMedia/Private/ModularFeature/AvaMediaSynchronizedEventsFeature.cpp
+++ b/Engine/Plugins/Experimental/Avalanche/Source/AvalancheMedia/Private/ModularFeature/AvaMediaSynchronizedEventsFeature.cpp
@@ -4,7 +4,6 @@
 
 #include "AvaMediaSettings.h"
 #include "Features/IModularFeatures.h"
-#include "Framework/AvaInstanceSettings.h"
 #include "HAL/IConsoleManager.h"
 #include "IAvaMediaModule.h"
 #include "ModularFeature/AvaMediaSynchronizedEvent.h"
@@ -116,7 +115,7 @@ void FAvaMediaSynchronizedEventsFeature::Shutdown()
 
 IAvaMediaSynchronizedEventsFeature* FAvaMediaSynchronizedEventsFeature::Get()
 {
-	const FName SelectedImplementation(IAvaMediaModule::Get().GetAvaInstanceSettings().SynchronizedEventsFeature.Implementation);
+	const FName SelectedImplementation(IAvaMediaModule::Get().GetPlayableSettings().SynchronizedEventsFeature.Implementation);
 	
 	IModularFeatures::FScopedLockModularFeatureList ScopedLockModularFeatureList;
 	return FindImplementation(SelectedImplementation);
diff --git a/Engine/Plugins/Experimental/Avalanche/Source/AvalancheMedia/Private/Playable/AvaPlayableGroup.cpp b/Engine/Plugins/Experimental/Avalanche/Source/AvalancheMedia/Private/Playable/AvaPlayableGroup.cpp
--- a/Engine/Plugins/Experimental/Avalanche/Source/AvalancheMedia/Private/Playable/AvaPlayableGroup.cpp
+++ b/Engine/Plugins/Experimental/Avalanche/Source/AvalancheMedia/Private/Playable/AvaPlayableGroup.cpp
@@ -3,14 +3,18 @@
 #include "Playable/AvaPlayableGroup.h"
 
 #include "AvaPlayableGroupSubsystem.h"
+#include "AvaPlayableUtils.h"
 #include "Engine/GameInstance.h"
 #include "Engine/Level.h"
 #include "Engine/ViewportStatsSubsystem.h"
 #include "Engine/World.h"
+#include "EngineUtils.h"
+#include "GameFramework/Pawn.h"
 #include "IAvaMediaModule.h"
 #include "Playable/AvaPlayable.h"
 #include "Playable/AvaPlayableGroupAssetUserData.h"
 #include "Playable/AvaPlayableGroupManager.h"
+#include "Playable/AvaPlayableSettings.h"
 #include "Playable/PlayableGroups/AvaGameInstancePlayableGroup.h"
 #include "Playable/PlayableGroups/AvaGameViewportPlayableGroup.h"
 #include "Playable/PlayableGroups/AvaRemoteProxyPlayableGroup.h"
@@ -18,6 +22,7 @@
 #include "Playback/AvaPlaybackManager.h"
 #include "Playback/AvaPlaybackUtils.h"
 #include "Playback/IAvaPlaybackServer.h"
+#include "SceneView.h"
 #include "UObject/Package.h"
 
 #define LOCTEXT_NAMESPACE "AvaPlayableGroup"
@@ -397,6 +402,12 @@ void UAvaPlayableGroup::RequestSetVisibility(UAvaPlayable* InPlayable, bool bInS
 
 void UAvaPlayableGroup::SetupView(FSceneViewFamily& InViewFamily, FSceneView& InView)
 {
+	const FAvaPlayableSettings& PlayableSettings = IAvaMediaModule::Get().GetPlayableSettings();
+	if (PlayableSettings.bHidePawnActors)
+	{
+		HidePawnsForView(GetPlayWorld(), InView);
+	}
+
 	for (const TObjectKey<UAvaPlayable>& PlayableKey : Playables)
 	{
 		UAvaPlayable* Playable = PlayableKey.ResolveObjectPtr();
@@ -650,4 +661,18 @@ bool UAvaPlayableGroup::DisplayTransitions(FText& OutText, FLinearColor& OutColo
 	return false;
 }
 
+void UAvaPlayableGroup::HidePawnsForView(const UWorld* InPlayWorld, FSceneView& InView) const
+{
+	if (!InPlayWorld)
+	{
+		return;
+	}
+
+	for (const APawn* Pawn : TActorRange<APawn>(InPlayWorld))
+	{
+		using namespace UE::AvaMedia::PlayableUtils;
+		AddPrimitiveComponentIds(Pawn, InView.HiddenPrimitives);
+	}
+}
+
 #undef LOCTEXT_NAMESPACE
\ No newline at end of file
diff --git a/Engine/Plugins/Experimental/Avalanche/Source/AvalancheMedia/Private/Playable/AvaPlayableUtils.cpp b/Engine/Plugins/Experimental/Avalanche/Source/AvalancheMedia/Private/Playable/AvaPlayableUtils.cpp
new file mode 100644
--- /dev/null
+++ b/Engine/Plugins/Experimental/Avalanche/Source/AvalancheMedia/Private/Playable/AvaPlayableUtils.cpp
@@ -0,0 +1,31 @@
+// Copyright Epic Games, Inc. All Rights Reserved.
+
+#include "AvaPlayableUtils.h"
+
+#include "Components/PrimitiveComponent.h"
+
+namespace UE::AvaMedia::PlayableUtils
+{
+	void AddPrimitiveComponentIds(const AActor* InActor, TSet<FPrimitiveComponentId>& InComponentIds)
+	{
+		TInlineComponentArray<UPrimitiveComponent*> Components;
+		InActor->GetComponents(Components);
+		for (int32 ComponentIndex = 0; ComponentIndex < Components.Num(); ComponentIndex++)
+		{
+			const UPrimitiveComponent* PrimitiveComponent = Components[ComponentIndex];
+			if (PrimitiveComponent->IsRegistered())
+			{
+				InComponentIds.Add(PrimitiveComponent->GetPrimitiveSceneId());
+
+				for (USceneComponent* AttachedChild : PrimitiveComponent->GetAttachChildren())
+				{						
+					const UPrimitiveComponent* AttachChildPC = Cast<UPrimitiveComponent>(AttachedChild);
+					if (AttachChildPC && AttachChildPC->IsRegistered())
+					{
+						InComponentIds.Add(AttachChildPC->GetPrimitiveSceneId());
+					}
+				}
+			}
+		}
+	}
+}
\ No newline at end of file
diff --git a/Engine/Plugins/Experimental/Avalanche/Source/AvalancheMedia/Private/Playable/AvaPlayableUtils.h b/Engine/Plugins/Experimental/Avalanche/Source/AvalancheMedia/Private/Playable/AvaPlayableUtils.h
new file mode 100644
--- /dev/null
+++ b/Engine/Plugins/Experimental/Avalanche/Source/AvalancheMedia/Private/Playable/AvaPlayableUtils.h
@@ -0,0 +1,13 @@
+// Copyright Epic Games, Inc. All Rights Reserved.
+
+#pragma once
+
+#include "Containers/Set.h"
+
+class AActor;
+class FPrimitiveComponentId;
+
+namespace UE::AvaMedia::PlayableUtils
+{
+	void AddPrimitiveComponentIds(const AActor* InActor, TSet<FPrimitiveComponentId>& InComponentIds);
+}
\ No newline at end of file
diff --git a/Engine/Plugins/Experimental/Avalanche/Source/AvalancheMedia/Private/Playable/Playables/AvaPlayableLevelStreaming.cpp b/Engine/Plugins/Experimental/Avalanche/Source/AvalancheMedia/Private/Playable/Playables/AvaPlayableLevelStreaming.cpp
--- a/Engine/Plugins/Experimental/Avalanche/Source/AvalancheMedia/Private/Playable/Playables/AvaPlayableLevelStreaming.cpp
+++ b/Engine/Plugins/Experimental/Avalanche/Source/AvalancheMedia/Private/Playable/Playables/AvaPlayableLevelStreaming.cpp
@@ -18,6 +18,7 @@
 #include "Playable/AvaPlayableAssetUserData.h"
 #include "Playable/AvaPlayableGroup.h"
 #include "Playable/AvaPlayableGroupManager.h"
+#include "Playable/AvaPlayableUtils.h"
 #include "Playback/AvaPlaybackUtils.h"
 #include "SceneView.h"
 #include "Streaming/LevelStreamingDelegates.h"
@@ -223,29 +224,6 @@ namespace UE::AvaMedia::LevelStreamingPlayable::Private
 		
 		return GetPlayableStatusRelevanceShouldBeUnloaded(InStatus);
 	}
-
-	void BuildPrimitiveComponentIdList(const AActor* InActor, TSet<FPrimitiveComponentId>& OutComponentIds)
-	{
-		TInlineComponentArray<UPrimitiveComponent*> Components;
-		InActor->GetComponents(Components);
-		for (int32 ComponentIndex = 0; ComponentIndex < Components.Num(); ComponentIndex++)
-		{
-			const UPrimitiveComponent* PrimitiveComponent = Components[ComponentIndex];
-			if (PrimitiveComponent->IsRegistered())
-			{
-				OutComponentIds.Add(PrimitiveComponent->GetPrimitiveSceneId());
-
-				for (USceneComponent* AttachedChild : PrimitiveComponent->GetAttachChildren())
-				{						
-					const UPrimitiveComponent* AttachChildPC = Cast<UPrimitiveComponent>(AttachedChild);
-					if (AttachChildPC && AttachChildPC->IsRegistered())
-					{
-						OutComponentIds.Add(AttachChildPC->GetPrimitiveSceneId());
-					}
-				}
-			}
-		}
-	}
 }
 
 bool UAvaPlayableLevelStreaming::LoadAsset(const FAvaSoftAssetPtr& InSourceAsset, bool bInInitiallyVisible)
@@ -444,13 +422,13 @@ void UAvaPlayableLevelStreaming::SetupView(FSceneViewFamily& InViewFamily, FScen
 	{
 		TSet<FPrimitiveComponentId> HiddenPrimitives;	// Todo(opt): cache this?
 
-		using namespace UE::AvaMedia::LevelStreamingPlayable::Private;
+		using namespace UE::AvaMedia::PlayableUtils;
 
 		for (TObjectPtr<AActor> Actor : Level->Actors)
 		{
 			if (IsValid(Actor))
 			{
-				BuildPrimitiveComponentIdList(Actor, HiddenPrimitives);
+				AddPrimitiveComponentIds(Actor, HiddenPrimitives);
 			}
 		}
 
diff --git a/Engine/Plugins/Experimental/Avalanche/Source/AvalancheMedia/Private/Playback/AvaPlaybackClient.cpp b/Engine/Plugins/Experimental/Avalanche/Source/AvalancheMedia/Private/Playback/AvaPlaybackClient.cpp
--- a/Engine/Plugins/Experimental/Avalanche/Source/AvalancheMedia/Private/Playback/AvaPlaybackClient.cpp
+++ b/Engine/Plugins/Experimental/Avalanche/Source/AvalancheMedia/Private/Playback/AvaPlaybackClient.cpp
@@ -1586,6 +1586,7 @@ void FAvaPlaybackClient::OnAvaMediaSettingsChanged(UObject*, struct FPropertyCha
 	ApplyAvaMediaSettings();
 	SendBroadcastSettingsUpdate(AllServerAddresses);
 	SendAvaInstanceSettingsUpdate(AllServerAddresses);
+	SendPlayableSettingsUpdate(AllServerAddresses);
 }
 
 void FAvaPlaybackClient::OnPreSavePackage(UPackage* InPackage, FObjectPreSaveContext InObjectSaveContext)
@@ -1723,6 +1724,13 @@ void FAvaPlaybackClient::SendAvaInstanceSettingsUpdate(const TArray<FMessageAddr
 	SendRequest(InstanceSettingsUpdate, InRecipients, EMessageFlags::Reliable);
 }
 
+void FAvaPlaybackClient::SendPlayableSettingsUpdate(const TArray<FMessageAddress>& InRecipients)
+{
+	FAvaPlaybackPlayableSettingsUpdate* PlayableSettingsUpdate = FMessageEndpoint::MakeMessage<FAvaPlaybackPlayableSettingsUpdate>();
+	PlayableSettingsUpdate->PlayableSettings = UAvaMediaSettings::Get().PlayableSettings;
+	SendRequest(PlayableSettingsUpdate, InRecipients, EMessageFlags::Reliable);
+}
+
 void FAvaPlaybackClient::SendBroadcastChannelSettingsUpdate(const TArray<FMessageAddress>& InRecipients, const FAvaBroadcastOutputChannel& InChannel)
 {
 	FAvaBroadcastChannelSettingsUpdate* ChannelSettingsUpdate = FMessageEndpoint::MakeMessage<FAvaBroadcastChannelSettingsUpdate>();
@@ -1756,11 +1764,13 @@ void FAvaPlaybackClient::SendClientInfo(const FMessageAddress& InRecipient)
 	ClientInfo->ProcessId = ProcessId;
 	ClientInfo->ProjectContentPath = ProjectContentPath;
 	SendRequest(ClientInfo, InRecipient);
-	
-	SendUserDataUpdate({InRecipient});
-	SendBroadcastSettingsUpdate({InRecipient});
-	SendAvaInstanceSettingsUpdate({InRecipient});
-	SendStatCommand(FString(), true, {InRecipient});	// Send empty stat command, will just send current states.
+
+	const TArray<FMessageAddress> Recipients = {InRecipient}; 
+	SendUserDataUpdate(Recipients);
+	SendBroadcastSettingsUpdate(Recipients);
+	SendAvaInstanceSettingsUpdate(Recipients);
+	SendPlayableSettingsUpdate(Recipients);
+	SendStatCommand(FString(), true, Recipients);	// Send empty stat command, will just send current states.
 }
 
 void FAvaPlaybackClient::RemoveDeadServers(const FDateTime& InCurrentTime)
diff --git a/Engine/Plugins/Experimental/Avalanche/Source/AvalancheMedia/Private/Playback/AvaPlaybackClient.h b/Engine/Plugins/Experimental/Avalanche/Source/AvalancheMedia/Private/Playback/AvaPlaybackClient.h
--- a/Engine/Plugins/Experimental/Avalanche/Source/AvalancheMedia/Private/Playback/AvaPlaybackClient.h
+++ b/Engine/Plugins/Experimental/Avalanche/Source/AvalancheMedia/Private/Playback/AvaPlaybackClient.h
@@ -141,6 +141,7 @@ protected:
 	void SendUserDataUpdate(const TArray<FMessageAddress>& InRecipients);
 	void SendBroadcastSettingsUpdate(const TArray<FMessageAddress>& InRecipients);
 	void SendAvaInstanceSettingsUpdate(const TArray<FMessageAddress>& InRecipients);
+	void SendPlayableSettingsUpdate(const TArray<FMessageAddress>& InRecipients);
 	void SendBroadcastChannelSettingsUpdate(const TArray<FMessageAddress>& InRecipients, const FAvaBroadcastOutputChannel& InChannel);
 	void SendPackageEvent(const TArray<FMessageAddress>& InRecipients, const FName& InPackageName, EAvaPlaybackPackageEvent InEvent);
 	void SendStatCommand(const FString& InCommand, bool bInBroadcastLocalState, const TArray<FMessageAddress>& InRecipients);
diff --git a/Engine/Plugins/Experimental/Avalanche/Source/AvalancheMedia/Private/Playback/AvaPlaybackMessages.h b/Engine/Plugins/Experimental/Avalanche/Source/AvalancheMedia/Private/Playback/AvaPlaybackMessages.h
--- a/Engine/Plugins/Experimental/Avalanche/Source/AvalancheMedia/Private/Playback/AvaPlaybackMessages.h
+++ b/Engine/Plugins/Experimental/Avalanche/Source/AvalancheMedia/Private/Playback/AvaPlaybackMessages.h
@@ -8,6 +8,7 @@
 #include "Framework/AvaInstanceSettings.h"
 #include "PixelFormat.h"
 #include "Playable/AvaPlayableRemoteControlValues.h"
+#include "Playable/AvaPlayableSettings.h"
 #include "Playback/Nodes/Events/Actions/AvaPlaybackAnimations.h"
 #include "Viewport/AvaViewportQualitySettings.h"
 #include "AvaPlaybackMessages.generated.h"
@@ -215,6 +216,18 @@ struct FAvaPlaybackInstanceSettingsUpdate : public FAvaPlaybackClientMessageBase
 	FAvaInstanceSettings InstanceSettings;
 };
 
+/**
+ *	Request for the client to replicate it's Motion Design instance settings to the server.
+ **/
+USTRUCT()
+struct FAvaPlaybackPlayableSettingsUpdate : public FAvaPlaybackClientMessageBase
+{
+	GENERATED_BODY()
+
+	UPROPERTY()
+	FAvaPlayableSettings PlayableSettings;
+};
+
 UENUM()
 enum class EAvaPlaybackPackageEvent
 {
diff --git a/Engine/Plugins/Experimental/Avalanche/Source/AvalancheMedia/Private/Playback/AvaPlaybackServer.cpp b/Engine/Plugins/Experimental/Avalanche/Source/AvalancheMedia/Private/Playback/AvaPlaybackServer.cpp
--- a/Engine/Plugins/Experimental/Avalanche/Source/AvalancheMedia/Private/Playback/AvaPlaybackServer.cpp
+++ b/Engine/Plugins/Experimental/Avalanche/Source/AvalancheMedia/Private/Playback/AvaPlaybackServer.cpp
@@ -162,6 +162,7 @@ void FAvaPlaybackServer::Init(const FString& InAssignedServerName)
 	.Handling<FAvaPlaybackDeviceProviderDataRequest>(this, &FAvaPlaybackServer::HandleDeviceProviderDataRequest)
 	.Handling<FAvaPlaybackUpdateClientInfo>(this, &FAvaPlaybackServer::HandleUpdateClientInfo)
 	.Handling<FAvaPlaybackInstanceSettingsUpdate>(this, &FAvaPlaybackServer::HandleAvaInstanceSettingsUpdate)
+	.Handling<FAvaPlaybackPlayableSettingsUpdate>(this, &FAvaPlaybackServer::HandlePlayableSettingsUpdate)
 	.Handling<FAvaPlaybackPackageEvent>(this, &FAvaPlaybackServer::HandlePackageEvent)
 	.Handling<FAvaPlaybackAssetStatusRequest>(this, &FAvaPlaybackServer::HandlePlaybackAssetStatusRequest)
 	.Handling<FAvaPlaybackRequest>(this, &FAvaPlaybackServer::HandlePlaybackRequest)
@@ -416,6 +417,17 @@ const FAvaInstanceSettings* FAvaPlaybackServer::GetAvaInstanceSettings() const
 	return nullptr;
 }
 
+const FAvaPlayableSettings* FAvaPlaybackServer::GetPlayableSettings() const
+{
+	// Returns the first client we have.
+	// Todo: In case we have multiple clients, we will need a smarter way to handle this.
+	for (const TPair<FString, TSharedPtr<FClientInfo>>& Client : Clients)
+	{
+		return &Client.Value->PlayableSettings;
+	}
+	return nullptr;
+}
+
 bool FAvaPlaybackServer::RemovePlaybackInstanceTransition(const FGuid& InTransitionId)
 {
 	if (PlaybackInstanceTransitions)
@@ -531,6 +543,14 @@ void FAvaPlaybackServer::HandleAvaInstanceSettingsUpdate(const FAvaPlaybackInsta
 	UE_LOG(LogAvaPlaybackServer, Verbose, TEXT("Received new instance settings from client \"%s\"."), *InMessage.ClientName);
 }
 
+void FAvaPlaybackServer::HandlePlayableSettingsUpdate(const FAvaPlaybackPlayableSettingsUpdate& InMessage, const TSharedRef<IMessageContext, ESPMode::ThreadSafe>& InContext)
+{
+	FClientInfo& ClientInfo = GetOrCreateClientInfo(InMessage.ClientName, InContext->GetSender());
+	ClientInfo.PlayableSettings = InMessage.PlayableSettings;
+	
+	UE_LOG(LogAvaPlaybackServer, Verbose, TEXT("Received new playable settings from client \"%s\"."), *InMessage.ClientName);
+}
+
 void FAvaPlaybackServer::HandlePackageEvent(const FAvaPlaybackPackageEvent& InMessage, const TSharedRef<IMessageContext, ESPMode::ThreadSafe>& InContext)
 {
 	const FClientInfo* ClientInfo = GetClientInfo(InContext->GetSender());
diff --git a/Engine/Plugins/Experimental/Avalanche/Source/AvalancheMedia/Private/Playback/AvaPlaybackServer.h b/Engine/Plugins/Experimental/Avalanche/Source/AvalancheMedia/Private/Playback/AvaPlaybackServer.h
--- a/Engine/Plugins/Experimental/Avalanche/Source/AvalancheMedia/Private/Playback/AvaPlaybackServer.h
+++ b/Engine/Plugins/Experimental/Avalanche/Source/AvalancheMedia/Private/Playback/AvaPlaybackServer.h
@@ -7,6 +7,7 @@
 #include "Playback/IAvaPlaybackServer.h"
 #include "IMessageContext.h"
 #include "MessageEndpoint.h"
+#include "Playable/AvaPlayableSettings.h"
 #include "Playback/AvaPlaybackManager.h"
 #include "Playback/AvaPlaybackMessages.h"
 #include "Templates/SharedPointer.h"
@@ -79,6 +80,8 @@ public:
 	virtual const IAvaBroadcastSettings* GetBroadcastSettings() const override;
 
 	virtual const FAvaInstanceSettings* GetAvaInstanceSettings() const override;
+
+	virtual const FAvaPlayableSettings* GetPlayableSettings() const override;
 	
 	virtual const FAvaPlaybackManager& GetPlaybackManager() const override { check(Manager); return *Manager; }
 	virtual FAvaPlaybackManager& GetPlaybackManager() override { check(Manager); return *Manager; }
@@ -110,6 +113,7 @@ public:
 	void HandleDeviceProviderDataRequest(const FAvaPlaybackDeviceProviderDataRequest& InMessage, const TSharedRef<IMessageContext, ESPMode::ThreadSafe>& InContext);
 	void HandleUpdateClientInfo(const FAvaPlaybackUpdateClientInfo& InMessage, const TSharedRef<IMessageContext, ESPMode::ThreadSafe>& InContext);
 	void HandleAvaInstanceSettingsUpdate(const FAvaPlaybackInstanceSettingsUpdate& InMessage, const TSharedRef<IMessageContext, ESPMode::ThreadSafe>& InContext);
+	void HandlePlayableSettingsUpdate(const FAvaPlaybackPlayableSettingsUpdate& InMessage, const TSharedRef<IMessageContext, ESPMode::ThreadSafe>& InContext);
 	void HandlePackageEvent(const FAvaPlaybackPackageEvent& InMessage, const TSharedRef<IMessageContext, ESPMode::ThreadSafe>& InContext);
 	void HandlePlaybackAssetStatusRequest(const FAvaPlaybackAssetStatusRequest& InMessage, const TSharedRef<IMessageContext, ESPMode::ThreadSafe>& InContext);
 	void HandlePlaybackRequest(const FAvaPlaybackRequest& InMessage, const TSharedRef<IMessageContext, ESPMode::ThreadSafe>& InContext);
@@ -298,6 +302,7 @@ private:
 		TMap<FString, FString> UserDataEntries;
 		FClientBroadcastSettings BroadcastSettings;
 		FAvaInstanceSettings AvaInstanceSettings;
+		FAvaPlayableSettings PlayableSettings;
 
 		bool bClientInfoReceived = false;
 	
diff --git a/Engine/Plugins/Experimental/Avalanche/Source/AvalancheMedia/Public/AvaMediaSettings.h b/Engine/Plugins/Experimental/Avalanche/Source/AvalancheMedia/Public/AvaMediaSettings.h
--- a/Engine/Plugins/Experimental/Avalanche/Source/AvalancheMedia/Public/AvaMediaSettings.h
+++ b/Engine/Plugins/Experimental/Avalanche/Source/AvalancheMedia/Public/AvaMediaSettings.h
@@ -10,6 +10,7 @@
 #include "Logging/LogVerbosity.h"
 #include "Math/MathFwd.h"
 #include "PixelFormat.h"
+#include "Playable/AvaPlayableSettings.h"
 #include "UObject/SoftObjectPtr.h"
 
 #include "AvaMediaSettings.generated.h"
@@ -193,6 +194,9 @@ public:
 	UPROPERTY(Config, EditAnywhere, Category = "Playback Manager")
 	FAvaInstanceSettings AvaInstanceSettings;
 
+	UPROPERTY(Config, EditAnywhere, Category = "Playback Manager")
+	FAvaPlayableSettings PlayableSettings;
+
 	/**
 	 * Maximum cached Managed Motion Design assets used for rundown editor's page details.
 	 * A value of 0 indicate the cache will grow without limit.
diff --git a/Engine/Plugins/Experimental/Avalanche/Source/AvalancheMedia/Public/IAvaMediaModule.h b/Engine/Plugins/Experimental/Avalanche/Source/AvalancheMedia/Public/IAvaMediaModule.h
--- a/Engine/Plugins/Experimental/Avalanche/Source/AvalancheMedia/Public/IAvaMediaModule.h
+++ b/Engine/Plugins/Experimental/Avalanche/Source/AvalancheMedia/Public/IAvaMediaModule.h
@@ -21,6 +21,7 @@ class IAvaRundownServer;
 class IMediaIOCoreDeviceProvider;
 class UWorld;
 struct FAvaInstanceSettings;
+struct FAvaPlayableSettings;
 struct FMediaIOOutputConfiguration;
 
 /** Maps one to one with the editor's map changed type (for now). */
@@ -117,6 +118,12 @@ public:
 	 */
 	virtual const FAvaInstanceSettings& GetAvaInstanceSettings() const = 0;
 
+	/**
+	 * Access global Playable Settings.
+	 * @remark These settings a replicated from connected the playback client (if connected).
+	 */
+	virtual const FAvaPlayableSettings& GetPlayableSettings() const = 0;
+	
 	/**
 	 *	Returns true if the local playback manager is (still) available.
 	 */
diff --git a/Engine/Plugins/Experimental/Avalanche/Source/AvalancheMedia/Public/Playable/AvaPlayableGroup.h b/Engine/Plugins/Experimental/Avalanche/Source/AvalancheMedia/Public/Playable/AvaPlayableGroup.h
--- a/Engine/Plugins/Experimental/Avalanche/Source/AvalancheMedia/Public/Playable/AvaPlayableGroup.h
+++ b/Engine/Plugins/Experimental/Avalanche/Source/AvalancheMedia/Public/Playable/AvaPlayableGroup.h
@@ -201,6 +201,8 @@ protected:
 	bool DisplayPlayingAssets(FText& OutText, FLinearColor& OutColor);
 	bool DisplayTransitions(FText& OutText, FLinearColor& OutColor);
 
+	void HidePawnsForView(const UWorld* InPlayWorld, FSceneView& InView) const;
+	
 protected:
 	/**
 	 * Managed Render Target for this playable group.
diff --git a/Engine/Plugins/Experimental/Avalanche/Source/AvalancheMedia/Public/Playable/AvaPlayableSettings.h b/Engine/Plugins/Experimental/Avalanche/Source/AvalancheMedia/Public/Playable/AvaPlayableSettings.h
new file mode 100644
--- /dev/null
+++ b/Engine/Plugins/Experimental/Avalanche/Source/AvalancheMedia/Public/Playable/AvaPlayableSettings.h
@@ -0,0 +1,39 @@
+// Copyright Epic Games, Inc. All Rights Reserved.
+
+#pragma once
+
+#include "AvaPlayableSettings.generated.h"
+
+USTRUCT()
+struct FAvaSynchronizedEventsFeatureSelection
+{
+	GENERATED_BODY()
+	
+	/**
+	 * Select the implementation for synchronizing events.
+	 * "Default" will select the most appropriate implementation available.
+	 */
+	UPROPERTY(Config, EditAnywhere, Category = Settings)
+	FString Implementation;
+};
+
+/**
+ * Settings applied when instancing a Motion Design Playable.
+ */
+USTRUCT()
+struct FAvaPlayableSettings
+{
+	GENERATED_BODY()
+
+	/**
+	 * Select the implementation for synchronizing events.
+	 */
+	UPROPERTY(Config, EditAnywhere, Category = Settings)
+	FAvaSynchronizedEventsFeatureSelection SynchronizedEventsFeature;
+
+	/**
+	 * Automatically hide pawn actor in playables.
+	 */
+	UPROPERTY(Config, EditAnywhere, Category = Settings)
+	bool bHidePawnActors = true;
+};
diff --git a/Engine/Plugins/Experimental/Avalanche/Source/AvalancheMedia/Public/Playback/IAvaPlaybackServer.h b/Engine/Plugins/Experimental/Avalanche/Source/AvalancheMedia/Public/Playback/IAvaPlaybackServer.h
--- a/Engine/Plugins/Experimental/Avalanche/Source/AvalancheMedia/Public/Playback/IAvaPlaybackServer.h
+++ b/Engine/Plugins/Experimental/Avalanche/Source/AvalancheMedia/Public/Playback/IAvaPlaybackServer.h
@@ -8,6 +8,7 @@
 class FAvaPlaybackManager;
 class IAvaBroadcastSettings;
 struct FAvaInstanceSettings;
+struct FAvaPlayableSettings;
 
 class IAvaPlaybackServer
 {
@@ -64,6 +65,12 @@ public:
 	 * Will return nullptr if no clients are connected.
 	 */
 	virtual const FAvaInstanceSettings* GetAvaInstanceSettings() const = 0;
+
+	/**
+	 * Access Motion Design Playable settings replicated from connected client(s).
+	 * Will return nullptr if no clients are connected.
+	 */
+	virtual const FAvaPlayableSettings* GetPlayableSettings() const = 0;
 	
 	/** Access the server's playback manager. */
 	virtual const FAvaPlaybackManager& GetPlaybackManager() const = 0;
diff --git a/Engine/Plugins/Experimental/Avalanche/Source/AvalancheMediaEditor/Private/AvaMediaEditorModule.cpp b/Engine/Plugins/Experimental/Avalanche/Source/AvalancheMediaEditor/Private/AvaMediaEditorModule.cpp
--- a/Engine/Plugins/Experimental/Avalanche/Source/AvalancheMediaEditor/Private/AvaMediaEditorModule.cpp
+++ b/Engine/Plugins/Experimental/Avalanche/Source/AvalancheMediaEditor/Private/AvaMediaEditorModule.cpp
@@ -11,7 +11,6 @@
 #include "Editor.h"
 #include "Engine/GameViewportClient.h"
 #include "Framework/Application/SlateApplication.h"
-#include "Framework/AvaInstanceSettings.h"
 #include "IAvaMediaModule.h"
 #include "LevelEditor.h"
 #include "LevelEditorViewport.h"
diff --git a/Engine/Plugins/Experimental/Avalanche/Source/AvalancheMediaEditor/Private/Customization/AvaSynchronizedEventsFeatureSelectionCustomization.cpp b/Engine/Plugins/Experimental/Avalanche/Source/AvalancheMediaEditor/Private/Customization/AvaSynchronizedEventsFeatureSelectionCustomization.cpp
--- a/Engine/Plugins/Experimental/Avalanche/Source/AvalancheMediaEditor/Private/Customization/AvaSynchronizedEventsFeatureSelectionCustomization.cpp
+++ b/Engine/Plugins/Experimental/Avalanche/Source/AvalancheMediaEditor/Private/Customization/AvaSynchronizedEventsFeatureSelectionCustomization.cpp
@@ -3,8 +3,8 @@
 #include "AvaSynchronizedEventsFeatureSelectionCustomization.h"
 
 #include "DetailWidgetRow.h"
-#include "Framework/AvaInstanceSettings.h"
 #include "IDetailChildrenBuilder.h"
+#include "Playable/AvaPlayableSettings.h"
 #include "SAvaSynchronizedEventsImplementationSelector.h"
 
 TSharedRef<IPropertyTypeCustomization> FAvaSynchronizedEventsFeatureSelectionCustomization::MakeInstance()
