Motion Camera System refactor: Changed how cameras work in Motion Design rundown to integrate with the Game Framework system (Player Camera Manager / Player Controller). This has the benefit that for systems that already use Game Framework the cameras set within the Motion Design level will not require further set up as well. As an example, prior to this change a camera cut track is required for MRQ to properly view the motion design camera. With this change, this is no longer required.

diff --git a/Engine/Plugins/Experimental/ActorModifierCore/Source/ActorModifierCore/Private/Subsystems/ActorModifierCoreSubsystem.cpp b/Engine/Plugins/Experimental/ActorModifierCore/Source/ActorModifierCore/Private/Subsystems/ActorModifierCoreSubsystem.cpp
--- a/Engine/Plugins/Experimental/ActorModifierCore/Source/ActorModifierCore/Private/Subsystems/ActorModifierCoreSubsystem.cpp
+++ b/Engine/Plugins/Experimental/ActorModifierCore/Source/ActorModifierCore/Private/Subsystems/ActorModifierCoreSubsystem.cpp
@@ -899,7 +899,7 @@ UActorModifierCoreStack* UActorModifierCoreSubsystem::GetActorModifierStack(cons
 {
 	if (const TWeakObjectPtr<UActorModifierCoreStack>* ActorStack = ModifierStacks.Find(InActor))
 	{
-		return (*ActorStack).Get();
+		return ActorStack->Get();
 	}
 	return nullptr;
 }
diff --git a/Engine/Plugins/Experimental/ActorModifierCore/Source/ActorModifierCore/Public/Modifiers/ActorModifierCoreStack.h b/Engine/Plugins/Experimental/ActorModifierCore/Source/ActorModifierCore/Public/Modifiers/ActorModifierCoreStack.h
--- a/Engine/Plugins/Experimental/ActorModifierCore/Source/ActorModifierCore/Public/Modifiers/ActorModifierCoreStack.h
+++ b/Engine/Plugins/Experimental/ActorModifierCore/Source/ActorModifierCore/Public/Modifiers/ActorModifierCoreStack.h
@@ -57,18 +57,33 @@ public:
 
 	/** Get modifiers of a specific class only in this stack, does not recurse */
 	template <class InModifierType
-		UE_REQUIRES(std::is_base_of<UActorModifierCoreBase, InModifierType>::value)>
+		UE_REQUIRES(std::is_base_of_v<UActorModifierCoreBase, InModifierType>)>
 	void GetClassModifiers(TArray<InModifierType*>& OutModifiers) const
 	{
 		for (const TObjectPtr<UActorModifierCoreBase>& Modifier : Modifiers)
 		{
-			if (Modifier->IsA(InModifierType::StaticClass()))
+			if (InModifierType* CastedModifier = Cast<InModifierType>(Modifier))
 			{
-				OutModifiers.Add(static_cast<InModifierType*>(Modifier));
+				OutModifiers.Add(CastedModifier);
 			}
 		}
 	}
 
+	/** Gets the first modifier of a specific class only in this stack, does not recurse */
+	template <class InModifierType
+		UE_REQUIRES(std::is_base_of_v<UActorModifierCoreBase, InModifierType>)>
+	InModifierType* GetClassModifier() const
+	{
+		for (const TObjectPtr<UActorModifierCoreBase>& Modifier : Modifiers)
+		{
+			if (InModifierType* CastedModifier = Cast<InModifierType>(Modifier))
+			{
+				return CastedModifier;
+			}
+		}
+		return nullptr;
+	}
+
 	/** Gets the first modifier in this stack, does not recurse */
 	ACTORMODIFIERCORE_API UActorModifierCoreBase* GetFirstModifier() const;
 
diff --git a/Engine/Plugins/Experimental/Avalanche/Avalanche.uplugin b/Engine/Plugins/Experimental/Avalanche/Avalanche.uplugin
--- a/Engine/Plugins/Experimental/Avalanche/Avalanche.uplugin
+++ b/Engine/Plugins/Experimental/Avalanche/Avalanche.uplugin
@@ -33,6 +33,12 @@
 			"LoadingPhase": "Default",
 			"TargetDenyList": [ "Server" ]
 		},
+		{
+			"Name": "AvalancheCamera",
+			"Type": "Runtime",
+			"LoadingPhase": "Default",
+			"TargetDenyList": [ "Server" ]
+		},
 		{
 			"Name": "AvalancheEditor",
 			"Type": "Editor",
diff --git a/Engine/Plugins/Experimental/Avalanche/Source/Avalanche/Avalanche.Build.cs b/Engine/Plugins/Experimental/Avalanche/Source/Avalanche/Avalanche.Build.cs
--- a/Engine/Plugins/Experimental/Avalanche/Source/Avalanche/Avalanche.Build.cs
+++ b/Engine/Plugins/Experimental/Avalanche/Source/Avalanche/Avalanche.Build.cs
@@ -33,6 +33,8 @@ public class Avalanche : ModuleRules
 			new string[]
 			{
 				"ActorModifierCore",
+				"ApplicationCore",
+				"AvalancheCamera",
 				"AvalancheRemoteControl",
 				"CinematicCamera",
 				"CoreUObject",
diff --git a/Engine/Plugins/Experimental/Avalanche/Source/Avalanche/Private/AvaScene.cpp b/Engine/Plugins/Experimental/Avalanche/Source/Avalanche/Private/AvaScene.cpp
--- a/Engine/Plugins/Experimental/Avalanche/Source/Avalanche/Private/AvaScene.cpp
+++ b/Engine/Plugins/Experimental/Avalanche/Source/Avalanche/Private/AvaScene.cpp
@@ -2,6 +2,7 @@
 
 #include "AvaScene.h"
 #include "AvaAssetTags.h"
+#include "AvaCameraSubsystem.h"
 #include "AvaRemoteControlUtils.h"
 #include "AvaSceneSettings.h"
 #include "AvaSceneState.h"
@@ -267,6 +268,26 @@ void AAvaScene::PostActorCreated()
 	RegisterObjects();
 }
 
+void AAvaScene::BeginPlay()
+{
+	Super::BeginPlay();
+
+	if (UAvaCameraSubsystem* CameraSubsystem = UAvaCameraSubsystem::Get(this))
+	{
+		CameraSubsystem->RegisterScene(GetLevel());
+	}
+}
+
+void AAvaScene::EndPlay(const EEndPlayReason::Type InEndPlayReason)
+{
+	Super::EndPlay(InEndPlayReason);
+
+	if (UAvaCameraSubsystem* CameraSubsystem = UAvaCameraSubsystem::Get(this))
+	{
+		CameraSubsystem->UnregisterScene(GetLevel());
+	}
+}
+
 void AAvaScene::PostLoad()
 {
 	Super::PostLoad();
diff --git a/Engine/Plugins/Experimental/Avalanche/Source/Avalanche/Private/Framework/AvaGameInstance.cpp b/Engine/Plugins/Experimental/Avalanche/Source/Avalanche/Private/Framework/AvaGameInstance.cpp
--- a/Engine/Plugins/Experimental/Avalanche/Source/Avalanche/Private/Framework/AvaGameInstance.cpp
+++ b/Engine/Plugins/Experimental/Avalanche/Source/Avalanche/Private/Framework/AvaGameInstance.cpp
@@ -2,6 +2,7 @@
 
 #include "Framework/AvaGameInstance.h"
 #include "AudioDevice.h"
+#include "AvaLog.h"
 #include "AvaRemoteControlRebind.h"
 #include "Components/ReflectionCaptureComponent.h"
 #include "Components/SkyLightComponent.h"
@@ -12,7 +13,6 @@
 #include "Framework/Application/SlateApplication.h"
 #include "Misc/TimeGuard.h"
 #include "Slate/SceneViewport.h"
-#include "Viewport/AvaCameraManager.h"
 #include "Viewport/AvaViewportQualitySettings.h"
 
 namespace UE::AvaGameInstance::Private
@@ -81,10 +81,11 @@ bool UAvaGameInstance::CreateWorld()
 		TRACE_CPUPROFILER_EVENT_SCOPE(UAvaGameInstance::LoadInstance::InitWorld);
 		
 		const FName MotionDesignWorldName = MakeUniqueObjectName(GetOuter(), UWorld::StaticClass(), TEXT("AvaGameInstanceWorld"));
-		
+
 		PlayWorld = NewObject<UWorld>(this, MotionDesignWorldName, RF_Transient);
 		check(PlayWorld);
 		PlayWorld->WorldType = WorldType;
+
 		PlayWorld->InitializeNewWorld(UWorld::InitializationValues()
 			.AllowAudioPlayback(true)
 			.CreatePhysicsScene(true)
@@ -124,10 +125,12 @@ bool UAvaGameInstance::BeginPlayWorld(const FAvaInstancePlaySettings& InWorldPla
 
 	TRACE_CPUPROFILER_EVENT_SCOPE(UAvaGameInstance::BeginPlay);
 	TRACE_BOOKMARK(TEXT("UAvaGameInstance::BeginPlay"));
-	
+
 	ViewportClient = NewObject<UAvaGameViewportClient>(GEngine, NAME_None, RF_Transient);
 	check(ViewportClient);
 
+	WorldContext->GameViewport = ViewportClient;
+
 	// Note: UGameViewportClient::Init ignores "bCreateNewAudioDevice" parameter and always create a device for the world. 
 	ViewportClient->Init(*WorldContext, this);
 	ViewportClient->SetRenderTarget(InWorldPlaySettings.RenderTarget);
@@ -145,10 +148,25 @@ bool UAvaGameInstance::BeginPlayWorld(const FAvaInstancePlaySettings& InWorldPla
 	TSharedPtr<SViewport> ViewportWidget;
 	Viewport = MakeShareable(ViewportClient->CreateGameViewport(ViewportWidget));
 	Viewport->SetInitialSize(InWorldPlaySettings.ViewportSize);
+	ViewportClient->Viewport = Viewport.Get();
+
+	// Attempt to initialize a Local Player.
+	FString Error;
+	LocalPlayer = ViewportClient->SetupInitialLocalPlayer(Error);
+	if (!LocalPlayer)
+	{
+		UE_LOG(LogAva, Error, TEXT("Couldn't create initial local player: %s"), *Error);
+	}
 
 	PlayWorld->InitializeActorsForPlay(PlayWorld->URL);
+
+	if (!LocalPlayer->SpawnPlayActor(TEXT(""), Error, PlayWorld))
+	{
+		UE_LOG(LogAva, Error, TEXT("Couldn't spawn play actor: %s"), *Error);
+	}
+
 	PlayWorld->BeginPlay();
-	
+
 	FCoreDelegates::OnEndFrame.AddUObject(this, &UAvaGameInstance::OnEndFrameTick);
 
 	bWorldPlaying = true;
@@ -387,6 +405,34 @@ void UAvaGameInstance::OnEnginePreExit()
 	}
 }
 
+ULocalPlayer* UAvaGameInstance::CreateInitialPlayer(FString& OutError)
+{
+	if (!GetGameViewportClient())
+	{
+		if (ensure(IsDedicatedServerInstance()))
+		{
+			OutError = FString::Printf(TEXT("Dedicated servers cannot have local players"));
+			return nullptr;
+		}
+	}
+
+	const FPlatformUserId UserId = IPlatformInputDeviceMapper::Get().GetPrimaryPlatformUser();
+
+	if (ULocalPlayer* ExistingLocalPlayer = FindLocalPlayerFromPlatformUserId(UserId))
+	{
+		return ExistingLocalPlayer;
+	}
+
+	if (LocalPlayers.IsEmpty())
+	{
+		ULocalPlayer* NewLocalPlayer = NewObject<ULocalPlayer>(GEngine, GEngine->LocalPlayerClass);
+		AddLocalPlayer(NewLocalPlayer, UserId);
+		return NewLocalPlayer;
+	}
+
+	return LocalPlayers[0];
+}
+
 void UAvaGameInstance::BeginDestroy()
 {
 	FCoreDelegates::OnEnginePreExit.Remove(EnginePreExitHandle);
diff --git a/Engine/Plugins/Experimental/Avalanche/Source/Avalanche/Private/Framework/AvaGameViewportClient.cpp b/Engine/Plugins/Experimental/Avalanche/Source/Avalanche/Private/Framework/AvaGameViewportClient.cpp
--- a/Engine/Plugins/Experimental/Avalanche/Source/Avalanche/Private/Framework/AvaGameViewportClient.cpp
+++ b/Engine/Plugins/Experimental/Avalanche/Source/Avalanche/Private/Framework/AvaGameViewportClient.cpp
@@ -2,6 +2,8 @@
 
 #include "Framework/AvaGameViewportClient.h"
 
+#include "AudioDevice.h"
+#include "AudioDeviceManager.h"
 #include "Camera/CameraComponent.h"
 #include "Camera/CameraPhotography.h"
 #include "Components/LineBatchComponent.h"
@@ -25,6 +27,7 @@
 #include "TextureResource.h"
 #include "UObject/Package.h"
 #include "UObject/UObjectGlobals.h"
+#include "UnrealEngine.h"
 
 /** Util to find named canvas in transient package, and create if not found */
 static UCanvas* GetCanvasByName(FName CanvasName)
@@ -48,26 +51,26 @@ static UCanvas* GetCanvasByName(FName CanvasName)
 	return *FoundCanvas;
 }
 
-UAvaGameViewportClient::UAvaGameViewportClient()
-	: CameraManager(MakeShared<FAvaCameraManager>())
+void UAvaGameViewportClient::Draw(FViewport* InViewport, FCanvas* InCanvas)
 {
-}
+	if (!GEngine)
+	{
+		return;
+	}
 
-void UAvaGameViewportClient::Draw(FViewport* InViewport, FCanvas* Canvas)
-{
 	// Override the Canvas Render Target with ours
 	if (RenderTarget.IsValid())
 	{
-		Canvas->SetRenderTarget_GameThread(RenderTarget->GameThread_GetRenderTargetResource());
+		InCanvas->SetRenderTarget_GameThread(RenderTarget->GameThread_GetRenderTargetResource());
 	}
-	
+
 	// Allow HMD to modify the view later, just before rendering
 	const bool bStereoRendering = GEngine->IsStereoscopic3D(InViewport);
-	
+
 	// Create a temporary canvas if there isn't already one.
 	static FName CanvasObjectName(TEXT("CanvasObject"));
 	UCanvas* CanvasObject = GetCanvasByName(CanvasObjectName);
-	CanvasObject->Canvas = Canvas;
+	CanvasObject->Canvas = InCanvas;
 
 	// Create temp debug canvas object
 	FCanvas* DebugCanvas = InViewport->GetDebugCanvas();
@@ -85,10 +88,10 @@ void UAvaGameViewportClient::Draw(FViewport* InViewport, FCanvas* Canvas)
 		DebugCanvas->SetScaledToRenderTarget(bStereoRendering);
 		DebugCanvas->SetStereoRendering(bStereoRendering);
 	}
-	Canvas->SetScaledToRenderTarget(bStereoRendering);
-	Canvas->SetStereoRendering(bStereoRendering);
-	
-	if (World == nullptr)
+	InCanvas->SetScaledToRenderTarget(bStereoRendering);
+	InCanvas->SetStereoRendering(bStereoRendering);
+
+	if (!World)
 	{
 		return;
 	}
@@ -103,7 +106,7 @@ void UAvaGameViewportClient::Draw(FViewport* InViewport, FCanvas* Canvas)
 	FGameTime Time = FGameTime::GetTimeSinceAppStart();
 
 	// Setup a FSceneViewFamily/FSceneView for the viewport.
-	FSceneViewFamilyContext ViewFamilyContext(FSceneViewFamily::ConstructionValues(Canvas->GetRenderTarget()
+	FSceneViewFamilyContext ViewFamilyContext(FSceneViewFamily::ConstructionValues(InCanvas->GetRenderTarget()
 		, World->Scene
 		, EngineShowFlags)
 		.SetRealtimeUpdate(true)
@@ -169,11 +172,27 @@ void UAvaGameViewportClient::Draw(FViewport* InViewport, FCanvas* Canvas)
 	}
 
 	TArray<FSceneView*> Views;
+
+	for (FLocalPlayerIterator PlayerIterator(GEngine, World); PlayerIterator; ++PlayerIterator)
 	{
+		ULocalPlayer* LocalPlayer = *PlayerIterator;
+		if (!LocalPlayer)
+		{
+			continue;
+		}
+
+		const bool bEnableStereo = GEngine->IsStereoscopic3D(InViewport);
 		const int32 NumViews = bStereoRendering ? GEngine->StereoRenderingDevice->GetDesiredNumberOfViews(bStereoRendering) : 1;
+
 		for (int32 ViewIndex = 0; ViewIndex < NumViews; ++ViewIndex)
 		{
-			FSceneView* const View = CalcSceneView(&ViewFamilyContext
+			// Calculate the player's view information.
+			FVector ViewLocation;
+			FRotator ViewRotation;
+
+			FSceneView* View = LocalPlayer->CalcSceneView(&ViewFamilyContext
+				, ViewLocation
+				, ViewRotation
 				, InViewport
 				, nullptr
 				, bStereoRendering ? ViewIndex : INDEX_NONE);
@@ -182,6 +201,13 @@ void UAvaGameViewportClient::Draw(FViewport* InViewport, FCanvas* Canvas)
 			{
 				Views.Add(View);
 
+				// If this is the primary drawing pass, update things that depend on the view location
+				if (ViewIndex == 0)
+				{
+					// Save the location of the view.
+					LocalPlayer->LastViewLocation = ViewLocation;
+				}
+
 				View->CameraConstrainedViewRect = View->UnscaledViewRect;
 
 				// Add view information for resource streaming. Allow up to 5X boost for small FOV.
@@ -193,16 +219,22 @@ void UAvaGameViewportClient::Draw(FViewport* InViewport, FCanvas* Canvas)
 					, StreamingScale);
 
 				World->ViewLocationsRenderedLastFrame.Add(View->ViewMatrices.GetViewOrigin());
+
+				FWorldCachedViewInfo& WorldViewInfo = World->CachedViewInfoRenderedLastFrame.AddDefaulted_GetRef();
+				WorldViewInfo.ViewMatrix           = View->ViewMatrices.GetViewMatrix();
+				WorldViewInfo.ProjectionMatrix     = View->ViewMatrices.GetProjectionMatrix();
+				WorldViewInfo.ViewProjectionMatrix = View->ViewMatrices.GetViewProjectionMatrix();
+				WorldViewInfo.ViewToWorld          = View->ViewMatrices.GetInvViewMatrix();
+
+				World->LastRenderTime = World->GetTimeSeconds();
 			}
 		}
 	}
 
-	//FinalizeViews(&ViewFamily); TODO Possible Interface
-
 	// Update level streaming.
 	World->UpdateLevelStreaming();
 
-	Canvas->Clear(FLinearColor::Transparent);
+	InCanvas->Clear(FLinearColor::Transparent);
 
 	// If a screen percentage interface was not set by one of the view extension, then set the legacy one.
 	if (ViewFamilyContext.GetScreenPercentageInterface() == nullptr)
@@ -252,7 +284,7 @@ void UAvaGameViewportClient::Draw(FViewport* InViewport, FCanvas* Canvas)
 				ViewExt->SetupView(ViewFamilyContext, *View);
 			}
 		}
-		GetRendererModule().BeginRenderingViewFamily(Canvas, &ViewFamilyContext);
+		GetRendererModule().BeginRenderingViewFamily(InCanvas, &ViewFamilyContext);
 	}
 	else
 	{
@@ -276,13 +308,11 @@ void UAvaGameViewportClient::Draw(FViewport* InViewport, FCanvas* Canvas)
 	// and is displayed in the broadcast channel's outputs.
 	if (!Views.IsEmpty() && IAvaModule::Get().ShouldShowRuntimeStats())
 	{
-		DrawStatsHUD( World, InViewport, Canvas, nullptr, DebugProperties, Views[0]->ViewLocation, Views[0]->ViewRotation);
+		DrawStatsHUD( World, InViewport, InCanvas, nullptr, DebugProperties, Views[0]->ViewLocation, Views[0]->ViewRotation);
 	}
-	
-	// Ensure canvas has been flushed before rendering UI
-	Canvas->Flush_GameThread();
 
-	CameraManager->ResetCameraCut();
+	// Ensure canvas has been flushed before rendering UI
+	InCanvas->Flush_GameThread();
 }
 
 bool UAvaGameViewportClient::IsStatEnabled(const FString& InName) const
@@ -304,247 +334,7 @@ void UAvaGameViewportClient::AddReferencedObjects(UObject* InThis, FReferenceCol
 	Super::AddReferencedObjects(InThis, InCollector);
 }
 
-void UAvaGameViewportClient::SetCameraCutThisFrame()
-{
-	bCameraCutThisFrame = true;
-}
-
 void UAvaGameViewportClient::SetRenderTarget(UTextureRenderTarget2D* InRenderTarget)
 {
 	RenderTarget = InRenderTarget;
 }
-
-bool UAvaGameViewportClient::CalcSceneViewInitOptions(FSceneViewInitOptions& OutInitOptions
-	, FViewport* InViewport
-	, FViewElementDrawer* ViewDrawer
-	, int32 StereoViewIndex)
-{
-	QUICK_SCOPE_CYCLE_COUNTER(STAT_AvaMediaCalcSceneViewInitOptions);
-	if (InViewport == nullptr || !World)
-	{
-		return false;
-	}
-
-	// get the projection data
-	if (GetProjectionData(InViewport, /*inout*/ OutInitOptions, StereoViewIndex) == false)
-	{
-		// Return false if this we didn't get back the info we needed
-		return false;
-	}
-
-	// return if we have an invalid view rect
-	if (!OutInitOptions.IsValidViewRectangle())
-	{
-		return false;
-	}
-
-	// Was there a camera cut this frame?
-	OutInitOptions.bInCameraCut = CameraManager->ShouldEnableCameraCut() || bCameraCutThisFrame;
-	bCameraCutThisFrame = false;
-
-	if (GEngine->StereoRenderingDevice.IsValid())
-	{
-		OutInitOptions.StereoPass = GEngine->StereoRenderingDevice->GetViewPassForIndex(StereoViewIndex != INDEX_NONE, StereoViewIndex);
-	}
-
-	const uint32 ViewIndex = StereoViewIndex != INDEX_NONE ? StereoViewIndex : 0;
-
-	// Make sure the ViewStates array has enough elements for the given ViewIndex.
-	{
-		const int32 RequiredViewStates = (ViewIndex + 1) - ViewStates.Num();
-
-		if (RequiredViewStates > 0)
-		{
-			ViewStates.AddDefaulted(RequiredViewStates);
-		}
-	}
-
-	// Allocate the current ViewState if necessary
-	if (ViewStates[ViewIndex].GetReference() == nullptr)
-	{
-#if UE_VERSION_OLDER_THAN(5, 3, 0)
-		const ERHIFeatureLevel::Type FeatureLevel = World ? World->FeatureLevel.GetValue() : GMaxRHIFeatureLevel;
-#else
-		const ERHIFeatureLevel::Type FeatureLevel = World ? World->GetFeatureLevel() : GMaxRHIFeatureLevel;
-#endif
-		ViewStates[ViewIndex].Allocate(FeatureLevel);
-	}
-
-	OutInitOptions.SceneViewStateInterface = ViewStates[ViewIndex].GetReference();
-	OutInitOptions.ViewActor = CameraManager->GetViewTarget();
-	OutInitOptions.PlayerIndex = 0;
-	OutInitOptions.ViewElementDrawer = ViewDrawer;
-	OutInitOptions.BackgroundColor = FLinearColor::Transparent;
-	OutInitOptions.LODDistanceFactor = 1.f;
-	OutInitOptions.StereoViewIndex = StereoViewIndex;
-	OutInitOptions.WorldToMetersScale = World->GetWorldSettings()->WorldToMeters;
-	OutInitOptions.CursorPos = InViewport->HasMouseCapture() ? FIntPoint(-1, -1) : FIntPoint(InViewport->GetMouseX(), InViewport->GetMouseY());
-	OutInitOptions.OriginOffsetThisFrame = World->OriginOffsetThisFrame;
-	return true;
-}
-
-bool UAvaGameViewportClient::GetProjectionData(FViewport* InViewport
-	, FSceneViewProjectionData& ProjectionData
-	, int32 StereoViewIndex) const
-{
-		// If the actor
-	if ((InViewport == nullptr) || (InViewport->GetSizeXY().X == 0) || (InViewport->GetSizeXY().Y == 0))
-	{
-		return false;
-	}
-
-	const FVector2D Origin(0.f, 0.f);
-	const FVector2D Size(1.f, 1.f);
-
-	int32 X = FMath::TruncToInt(Origin.X * InViewport->GetSizeXY().X);
-	int32 Y = FMath::TruncToInt(Origin.Y * InViewport->GetSizeXY().Y);
-
-	X += InViewport->GetInitialPositionXY().X;
-	Y += InViewport->GetInitialPositionXY().Y;
-
-	uint32 SizeX = FMath::TruncToInt(Size.X * InViewport->GetSizeXY().X);
-	uint32 SizeY = FMath::TruncToInt(Size.Y * InViewport->GetSizeXY().Y);
-
-	FIntRect UnconstrainedRectangle = FIntRect(X, Y, X + SizeX, Y + SizeY);
-
-	ProjectionData.SetViewRectangle(UnconstrainedRectangle);
-
-	// Get the viewpoint.
-	FMinimalViewInfo ViewInfo;
-	CameraManager->SetCameraViewInfo(ViewInfo);
-
-	// If stereo rendering is enabled, update the size and offset appropriately for this pass
-	const bool bNeedStereo = StereoViewIndex != INDEX_NONE && GEngine->IsStereoscopic3D();
-	const bool bIsHeadTrackingAllowed = GEngine->XRSystem.IsValid()
-		&& (World != nullptr
-			? GEngine->XRSystem->IsHeadTrackingAllowedForWorld(*World)
-			: GEngine->XRSystem->IsHeadTrackingAllowed());
-
-	if (bNeedStereo)
-	{
-		GEngine->StereoRenderingDevice->AdjustViewRect(StereoViewIndex, X, Y, SizeX, SizeY);
-	}
-
-    FVector StereoViewLocation = ViewInfo.Location;
-    if (bNeedStereo || bIsHeadTrackingAllowed)
-    {
-	    TSharedPtr<IXRCamera> XRCamera = GEngine->XRSystem.IsValid()
-    		? GEngine->XRSystem->GetXRCamera()
-    		: nullptr;
-
-		if (XRCamera.IsValid())
-		{
-			const bool bHasActiveCamera = !!CameraManager->GetCachedCameraComponent();
-			XRCamera->UseImplicitHMDPosition(bHasActiveCamera);
-		}
-
-		if (GEngine->StereoRenderingDevice.IsValid())
-		{
-			GEngine->StereoRenderingDevice->CalculateStereoViewOffset(StereoViewIndex, ViewInfo.Rotation, World->GetWorldSettings()->WorldToMeters, StereoViewLocation);
-		}
-    }
-
-	// Create the view matrix
-	ProjectionData.ViewOrigin = StereoViewLocation;
-	ProjectionData.ViewRotationMatrix = FInverseRotationMatrix(ViewInfo.Rotation) * FMatrix(
-		FPlane(0,	0,	1,	0),
-		FPlane(1,	0,	0,	0),
-		FPlane(0,	1,	0,	0),
-		FPlane(0,	0,	0,	1));
-
-	// @todo viewext this use case needs to be revisited
-	if (!bNeedStereo)
-	{
-		// Create the projection matrix (and possibly constrain the view rectangle)
-		FMinimalViewInfo::CalculateProjectionMatrixGivenView(ViewInfo
-			, GetDefault<ULocalPlayer>()->AspectRatioAxisConstraint
-			, InViewport
-			, ProjectionData);
-
-		for (const TSharedRef<ISceneViewExtension>& ViewExt : GEngine->ViewExtensions->GatherActiveExtensions(FSceneViewExtensionContext(InViewport)))
-        {
-			ViewExt->SetupViewProjectionMatrix(ProjectionData);
-		};
-	}
-	else
-	{
-		// Let the stereoscopic rendering device handle creating its own projection matrix, as needed
-		ProjectionData.ProjectionMatrix = GEngine->StereoRenderingDevice->GetStereoProjectionMatrix(StereoViewIndex);
-
-		// calculate the out rect
-		ProjectionData.SetViewRectangle(FIntRect(X, Y, X + SizeX, Y + SizeY));
-	}
-	return true;
-}
-
-FSceneView* UAvaGameViewportClient::CalcSceneView(FSceneViewFamily* ViewFamily
-	, FViewport* InViewport
-	, FViewElementDrawer* ViewDrawer
-	, int32 StereoViewIndex)
-{
-	//QUICK_SCOPE_CYCLE_COUNTER(STAT_AvaMediaCalcSceneView);
-	FSceneViewInitOptions ViewInitOptions;
-	if (!CalcSceneViewInitOptions(ViewInitOptions, InViewport, ViewDrawer, StereoViewIndex))
-	{
-		return nullptr;
-	}
-
-	// Get the viewpoint...technically doing this twice
-	// but it makes GetProjectionData better
-	FMinimalViewInfo ViewInfo;
-	CameraManager->SetCameraViewInfo(ViewInfo);
-
-	ViewInitOptions.ViewLocation = ViewInfo.Location;
-	ViewInitOptions.ViewRotation = ViewInfo.Rotation;
-	ViewInitOptions.bUseFieldOfViewForLOD = ViewInfo.bUseFieldOfViewForLOD;
-	ViewInitOptions.FOV = ViewInfo.FOV;
-	ViewInitOptions.DesiredFOV = ViewInfo.DesiredFOV;
-
-	// Fill out the rest of the view init options
-	ViewInitOptions.ViewFamily = ViewFamily;
-
-	FSceneView* const View = new FSceneView(ViewInitOptions);
-
-	//Pass on the previous view transform from the view info (probably provided by the camera if set)
-	View->PreviousViewTransform = ViewInfo.PreviousViewTransform;
-
-	ViewFamily->Views.Add(View);
-
-	{
-		View->StartFinalPostprocessSettings(ViewInfo.Location);
-
-		//Post Process Setting
-		//TODO?
-
-		//	CAMERA OVERRIDE
-		View->OverridePostProcessSettings(ViewInfo.PostProcessSettings, ViewInfo.PostProcessBlendWeight);
-		CameraManager->ApplyExtraPostProcessBlends(View);
-		FCameraPhotographyManager::Get().UpdatePostProcessing(View->FinalPostProcessSettings);
-
-		if (GEngine->StereoRenderingDevice.IsValid())
-		{
-			FPostProcessSettings StereoDeviceOverridePostProcessingSettings;
-			float BlendWeight = 1.0f;
-
-			const bool StereoSettingsAvailable = GEngine->StereoRenderingDevice->OverrideFinalPostprocessSettings(&StereoDeviceOverridePostProcessingSettings
-				, View->StereoPass
-				, View->StereoViewIndex
-				, BlendWeight);
-
-			if (StereoSettingsAvailable)
-			{
-				View->OverridePostProcessSettings(StereoDeviceOverridePostProcessingSettings, BlendWeight);
-			}
-		}
-
-
-		View->EndFinalPostprocessSettings(ViewInitOptions);
-	}
-
-	for (int32 ViewExt = 0; ViewExt < ViewFamily->ViewExtensions.Num(); ++ViewExt)
-	{
-		ViewFamily->ViewExtensions[ViewExt]->SetupView(*ViewFamily, *View);
-	}
-	
-	return View;
-}
diff --git a/Engine/Plugins/Experimental/Avalanche/Source/Avalanche/Private/Viewport/AvaCameraManager.cpp b/Engine/Plugins/Experimental/Avalanche/Source/Avalanche/Private/Viewport/AvaCameraManager.cpp
deleted file mode 100644
diff --git a/Engine/Plugins/Experimental/Avalanche/Source/Avalanche/Public/AvaScene.h b/Engine/Plugins/Experimental/Avalanche/Source/Avalanche/Public/AvaScene.h
--- a/Engine/Plugins/Experimental/Avalanche/Source/Avalanche/Public/AvaScene.h
+++ b/Engine/Plugins/Experimental/Avalanche/Source/Avalanche/Public/AvaScene.h
@@ -80,6 +80,8 @@ public:
 
 	//~ Begin AActor
 	virtual void PostActorCreated() override;
+	virtual void BeginPlay() override;
+	virtual void EndPlay(const EEndPlayReason::Type InEndPlayReason) override;
 #if WITH_EDITOR
 	virtual bool IsSelectable() const override { return false; }
 	virtual bool SupportsExternalPackaging() const override { return false; }
diff --git a/Engine/Plugins/Experimental/Avalanche/Source/Avalanche/Public/Framework/AvaGameInstance.h b/Engine/Plugins/Experimental/Avalanche/Source/Avalanche/Public/Framework/AvaGameInstance.h
--- a/Engine/Plugins/Experimental/Avalanche/Source/Avalanche/Public/Framework/AvaGameInstance.h
+++ b/Engine/Plugins/Experimental/Avalanche/Source/Avalanche/Public/Framework/AvaGameInstance.h
@@ -108,6 +108,10 @@ protected:
 	void OnEndFrameTick();
 	void OnEnginePreExit();
 
+	//~ Begin UGameInstance
+	virtual ULocalPlayer* CreateInitialPlayer(FString& OutError) override;
+	//~ End UGameInstance
+	
 	//~ Begin UObject
 	virtual void BeginDestroy() override;
 	//~ End UObject
@@ -115,6 +119,9 @@ protected:
 	UPROPERTY(Transient)
 	TObjectPtr<UWorld> PlayWorld;
 
+	UPROPERTY(Transient)
+	TObjectPtr<ULocalPlayer> LocalPlayer;
+
 	UPROPERTY(Transient)
 	TObjectPtr<UAvaGameViewportClient> ViewportClient;
 
diff --git a/Engine/Plugins/Experimental/Avalanche/Source/Avalanche/Public/Framework/AvaGameViewportClient.h b/Engine/Plugins/Experimental/Avalanche/Source/Avalanche/Public/Framework/AvaGameViewportClient.h
--- a/Engine/Plugins/Experimental/Avalanche/Source/Avalanche/Public/Framework/AvaGameViewportClient.h
+++ b/Engine/Plugins/Experimental/Avalanche/Source/Avalanche/Public/Framework/AvaGameViewportClient.h
@@ -3,61 +3,33 @@
 #pragma once
 
 #include "Engine/GameViewportClient.h"
-#include "Viewport/AvaCameraManager.h"
 #include "SceneTypes.h"
 #include "AvaGameViewportClient.generated.h"
 
-class AAvaViewportCameraActor;
-class FPrimitiveComponentId;
 class FViewElementDrawer;
 class UTextureRenderTarget2D;
-struct FMinimalViewInfo;
 struct FSceneViewInitOptions;
 struct FSceneViewProjectionData;
 
-UCLASS(DisplayName = "Motion Design Game Viewport Client")
-class AVALANCHE_API UAvaGameViewportClient : public UGameViewportClient
+UCLASS(MinimalAPI, DisplayName = "Motion Design Game Viewport Client")
+class UAvaGameViewportClient : public UGameViewportClient
 {
 	GENERATED_BODY()
-	
-public:
-	UAvaGameViewportClient();
 
+public:
 	//~ Begin FViewportClient
-	virtual void Draw(FViewport* Viewport,FCanvas* Canvas) override;
+	virtual void Draw(FViewport* InViewport,FCanvas* InCanvas) override;
 	virtual bool IsStatEnabled(const FString& InName) const override;
 	//~ End FViewportClient
 
-	static void AddReferencedObjects(UObject* InThis, FReferenceCollector& InCollector);
-	
-	void SetCameraCutThisFrame();
+	AVALANCHE_API static void AddReferencedObjects(UObject* InThis, FReferenceCollector& InCollector);
 
-	void SetRenderTarget(UTextureRenderTarget2D* InRenderTarget);
+	AVALANCHE_API void SetRenderTarget(UTextureRenderTarget2D* InRenderTarget);
 
 	UTextureRenderTarget2D* GetRenderTarget() const { return RenderTarget.Get(); }
 
-	TSharedRef<FAvaCameraManager> GetCameraManager() const { return CameraManager; }
-
 private:
-	bool CalcSceneViewInitOptions(FSceneViewInitOptions& OutInitOptions
-		, FViewport* InViewport
-		, FViewElementDrawer* ViewDrawer = nullptr
-		, int32 StereoViewIndex = INDEX_NONE);
-
-	bool GetProjectionData(FViewport* InViewport
-		, FSceneViewProjectionData& ProjectionData
-		, int32 StereoViewIndex = INDEX_NONE) const;
-
-	FSceneView* CalcSceneView(FSceneViewFamily* ViewFamily
-		, FViewport* InViewport
-		, FViewElementDrawer* ViewDrawer = nullptr
-		, int32 StereoViewIndex = INDEX_NONE);
-
 	TArray<FSceneViewStateReference> ViewStates;
 
 	TWeakObjectPtr<UTextureRenderTarget2D> RenderTarget;
-
-	TSharedRef<FAvaCameraManager> CameraManager;
-
-	bool bCameraCutThisFrame = false;
 };
diff --git a/Engine/Plugins/Experimental/Avalanche/Source/Avalanche/Public/Viewport/AvaCameraManager.h b/Engine/Plugins/Experimental/Avalanche/Source/Avalanche/Public/Viewport/AvaCameraManager.h
deleted file mode 100644
diff --git a/Engine/Plugins/Experimental/Avalanche/Source/AvalancheCamera/AvalancheCamera.Build.cs b/Engine/Plugins/Experimental/Avalanche/Source/AvalancheCamera/AvalancheCamera.Build.cs
new file mode 100644
--- /dev/null
+++ b/Engine/Plugins/Experimental/Avalanche/Source/AvalancheCamera/AvalancheCamera.Build.cs
@@ -0,0 +1,37 @@
+﻿// Copyright Epic Games, Inc. All Rights Reserved.
+
+using UnrealBuildTool;
+
+public class AvalancheCamera : ModuleRules
+{
+    public AvalancheCamera(ReadOnlyTargetRules Target) : base(Target)
+    {
+        PCHUsage = ModuleRules.PCHUsageMode.UseExplicitOrSharedPCHs;
+
+        PublicDependencyModuleNames.AddRange(
+            new string[]
+            {
+                "ActorModifierCore",
+                "AvalancheTransition",
+                "Core",
+                "StateTreeModule",
+            }
+        );
+
+        PrivateDependencyModuleNames.AddRange(
+            new string[]
+            {
+                "CoreUObject",
+                "Engine",
+            }
+        );
+
+        if (Target.Type == TargetRules.TargetType.Editor)
+        {
+            PrivateDependencyModuleNames.AddRange(new string[]
+            {
+                "SlateCore",
+            });
+        }
+    }
+}
\ No newline at end of file
diff --git a/Engine/Plugins/Experimental/Avalanche/Source/AvalancheCamera/Private/AvaCameraModule.cpp b/Engine/Plugins/Experimental/Avalanche/Source/AvalancheCamera/Private/AvaCameraModule.cpp
new file mode 100644
--- /dev/null
+++ b/Engine/Plugins/Experimental/Avalanche/Source/AvalancheCamera/Private/AvaCameraModule.cpp
@@ -0,0 +1,5 @@
+﻿// Copyright Epic Games, Inc. All Rights Reserved.
+
+#include "Modules/ModuleManager.h"
+
+IMPLEMENT_MODULE(FDefaultModuleImpl, AvalancheCamera);
diff --git a/Engine/Plugins/Experimental/Avalanche/Source/AvalancheCamera/Private/AvaCameraPriorityModifier.cpp b/Engine/Plugins/Experimental/Avalanche/Source/AvalancheCamera/Private/AvaCameraPriorityModifier.cpp
new file mode 100644
--- /dev/null
+++ b/Engine/Plugins/Experimental/Avalanche/Source/AvalancheCamera/Private/AvaCameraPriorityModifier.cpp
@@ -0,0 +1,41 @@
+// Copyright Epic Games, Inc. All Rights Reserved.
+
+#include "AvaCameraPriorityModifier.h"
+#include "Camera/CameraComponent.h"
+
+#if WITH_EDITOR
+#include "Styling/SlateIconFinder.h"
+#endif
+
+#define LOCTEXT_NAMESPACE "AvaCameraModifier"
+
+UAvaCameraPriorityModifier::UAvaCameraPriorityModifier()
+{
+	TransitionParams.bLockOutgoing = true;
+}
+
+void UAvaCameraPriorityModifier::OnModifierCDOSetup(FActorModifierCoreMetadata& InMetadata)
+{
+	Super::OnModifierCDOSetup(InMetadata);
+
+	InMetadata.SetCompatibilityRule([](const AActor* InActor)->bool
+	{
+		return InActor && InActor->FindComponentByClass<UCameraComponent>();
+	});
+
+	InMetadata.SetName(TEXT("CameraPriority"));
+	InMetadata.SetCategory(TEXT("Camera"));
+#if WITH_EDITOR
+	InMetadata.SetDisplayName(LOCTEXT("ModifierDisplayName", "Camera Priority"));
+	InMetadata.SetIcon(FSlateIconFinder::FindIconForClass(UCameraComponent::StaticClass()));
+	InMetadata.SetDescription(LOCTEXT("ModifierDescription", "Camera Priority information for Motion Design Scenes"));
+#endif
+}
+
+void UAvaCameraPriorityModifier::Apply()
+{
+	// Nothing to apply. Continue onto next
+	Next();
+}
+
+#undef LOCTEXT_NAMESPACE
diff --git a/Engine/Plugins/Experimental/Avalanche/Source/AvalancheCamera/Private/AvaCameraSubsystem.cpp b/Engine/Plugins/Experimental/Avalanche/Source/AvalancheCamera/Private/AvaCameraSubsystem.cpp
new file mode 100644
--- /dev/null
+++ b/Engine/Plugins/Experimental/Avalanche/Source/AvalancheCamera/Private/AvaCameraSubsystem.cpp
@@ -0,0 +1,222 @@
+// Copyright Epic Games, Inc. All Rights Reserved.
+
+#include "AvaCameraSubsystem.h"
+#include "AvaCameraPriorityModifier.h"
+#include "AvaTransitionSubsystem.h"
+#include "AvaTransitionTree.h"
+#include "Behavior/AvaTransitionBehaviorInstance.h"
+#include "Behavior/IAvaTransitionBehavior.h"
+#include "Camera/PlayerCameraManager.h"
+#include "Engine/Level.h"
+#include "Execution/IAvaTransitionExecutor.h"
+#include "GameFramework/PlayerController.h"
+#include "Modifiers/ActorModifierCoreStack.h"
+#include "Subsystems/ActorModifierCoreSubsystem.h"
+#include "Transition/Tasks/AvaCameraBlendTask.h"
+
+bool FAvaViewTarget::IsValid() const
+{
+	return ::IsValid(Actor) && ::IsValid(CameraPriorityModifier);
+}
+
+UAvaCameraSubsystem* UAvaCameraSubsystem::Get(const UObject* InObject)
+{
+	if (!InObject)
+	{
+		return nullptr;
+	}
+
+	const UWorld* World = InObject->GetWorld();
+	if (!World)
+	{
+		return nullptr;
+	}
+
+	return World->GetSubsystem<UAvaCameraSubsystem>();
+}
+
+void UAvaCameraSubsystem::RegisterScene(const ULevel* InSceneLevel)
+{
+	if (!InSceneLevel)
+	{
+		return;
+	}
+
+	const UActorModifierCoreSubsystem* ModifierSubsystem = UActorModifierCoreSubsystem::Get();
+	if (!ModifierSubsystem)
+	{
+		return;
+	}
+
+	for (AActor* Actor : InSceneLevel->Actors)
+	{
+		if (!Actor)
+		{
+			continue;
+		}
+
+		const UActorModifierCoreStack* ModifierStack = ModifierSubsystem->GetActorModifierStack(Actor);
+		if (!ModifierStack)
+		{
+			continue;
+		}
+
+		if (UAvaCameraPriorityModifier* CameraPriorityModifier = ModifierStack->GetClassModifier<UAvaCameraPriorityModifier>())
+		{
+			FAvaViewTarget& Entry = ViewTargets.AddDefaulted_GetRef();
+			Entry.Actor = Actor;
+			Entry.CameraPriorityModifier = CameraPriorityModifier;
+		}
+	}
+
+	// Remove invalid entries
+	ViewTargets.RemoveAll([](const FAvaViewTarget& InViewTarget)
+		{
+			return !InViewTarget.IsValid();
+		});
+
+	// Sort so that the higher priorities are at the end of the list (end = current)
+	// Stable sorting so that the more recently registered scenes are preferred when the priorities match
+	ViewTargets.StableSort(
+		[](const FAvaViewTarget& A, const FAvaViewTarget& B)
+		{
+			return A.CameraPriorityModifier->GetPriority() < B.CameraPriorityModifier->GetPriority();
+		});
+
+	if (!HasCustomViewTargetting(InSceneLevel))
+	{
+		UpdatePlayerControllerViewTarget();
+	}
+}
+
+void UAvaCameraSubsystem::UnregisterScene(const ULevel* InSceneLevel)
+{
+	if (!InSceneLevel)
+	{
+		return;
+	}
+
+	// Remove entries with invalid view targets or targets that match the level that is being unregistered
+	ViewTargets.RemoveAll(
+		[InSceneLevel](const FAvaViewTarget& InViewTarget)
+		{
+			return !InViewTarget.IsValid() || InViewTarget.Actor->GetLevel() == InSceneLevel;
+		});
+
+	if (!HasCustomViewTargetting(InSceneLevel))
+    {
+    	UpdatePlayerControllerViewTarget();
+    }
+}
+
+bool UAvaCameraSubsystem::IsBlendingToViewTarget(const ULevel* InSceneLevel) const
+{
+	if (InSceneLevel && PlayerController && PlayerController->PlayerCameraManager)
+	{
+		return PlayerController->PlayerCameraManager->BlendTimeToGo > 0.f
+			&& PlayerController->PlayerCameraManager->ViewTarget.Target
+			&& PlayerController->PlayerCameraManager->ViewTarget.Target->GetLevel() == InSceneLevel;
+	}
+
+	return false;
+}
+
+void UAvaCameraSubsystem::UpdatePlayerControllerViewTarget(const FViewTargetTransitionParams* InOverrideTransitionParams) const
+{
+	if (!IsValid(PlayerController) || ViewTargets.IsEmpty())
+	{
+		return;
+	}
+
+	const FAvaViewTarget& DesiredViewTarget = ViewTargets.Last();
+
+	// These should've been removed prior to calling this private function
+	check(DesiredViewTarget.Actor && DesiredViewTarget.CameraPriorityModifier);
+
+	// Current View Target matches the desired. Nothing to do
+	if (DesiredViewTarget.Actor == PlayerController->GetViewTarget())
+	{
+		return;
+	}
+
+	const FViewTargetTransitionParams* ViewTargetTransitionParams = InOverrideTransitionParams
+		? InOverrideTransitionParams
+		: &DesiredViewTarget.CameraPriorityModifier->GetTransitionParams();
+
+	PlayerController->SetViewTarget(DesiredViewTarget.Actor, *ViewTargetTransitionParams);
+}
+
+bool UAvaCameraSubsystem::DoesSupportWorldType(const EWorldType::Type InWorldType) const
+{
+	return InWorldType == EWorldType::GamePreview
+		|| InWorldType == EWorldType::Game
+		|| InWorldType == EWorldType::PIE;
+}
+
+void UAvaCameraSubsystem::OnWorldBeginPlay(UWorld& InWorld)
+{
+	Super::OnWorldBeginPlay(InWorld);
+	PlayerController = InWorld.GetFirstPlayerController();
+}
+
+void UAvaCameraSubsystem::Initialize(FSubsystemCollectionBase& InCollection)
+{
+	Super::Initialize(InCollection);
+	IAvaTransitionExecutor::GetOnTransitionStart().AddUObject(this, &UAvaCameraSubsystem::OnTransitionStart);
+}
+
+void UAvaCameraSubsystem::Deinitialize()
+{
+	Super::Deinitialize();
+	IAvaTransitionExecutor::GetOnTransitionStart().RemoveAll(this);
+}
+
+void UAvaCameraSubsystem::OnTransitionStart(const IAvaTransitionExecutor& InExecutor)
+{
+	// Unregister scenes that are marked as Needing Discard from the Start
+	InExecutor.ForEachBehaviorInstance(
+		[this](const FAvaTransitionBehaviorInstance& InInstance)
+		{
+			const FAvaTransitionScene* TransitionScene = InInstance.GetTransitionContext().GetTransitionScene();
+			if (!TransitionScene || !TransitionScene->HasAnyFlags(EAvaTransitionSceneFlags::NeedsDiscard))
+			{
+				return;
+			}
+
+			const ULevel* Level = TransitionScene->GetLevel();
+			if (!Level)
+			{
+				return;
+			}
+
+			UnregisterScene(Level);
+		});
+}
+
+bool UAvaCameraSubsystem::HasCustomViewTargetting(const ULevel* InSceneLevel) const
+{
+	if (!InSceneLevel)
+	{
+		return false;
+	}
+
+	UAvaTransitionSubsystem* TransitionSubsystem = InSceneLevel->OwningWorld->GetSubsystem<UAvaTransitionSubsystem>();
+	if (!TransitionSubsystem)
+	{
+		return false;
+	}
+
+	IAvaTransitionBehavior* TransitionBehavior = TransitionSubsystem->GetTransitionBehavior(InSceneLevel);
+	if (!TransitionBehavior)
+	{
+		return false;
+	}
+
+	UAvaTransitionTree* TransitionTree = TransitionBehavior->GetTransitionTree();
+	if (!TransitionTree)
+	{
+		return false;
+	}
+
+	return TransitionTree->ContainsTask<FAvaCameraBlendTask>();
+}
diff --git a/Engine/Plugins/Experimental/Avalanche/Source/AvalancheCamera/Private/Transition/Tasks/AvaCameraBlendTask.cpp b/Engine/Plugins/Experimental/Avalanche/Source/AvalancheCamera/Private/Transition/Tasks/AvaCameraBlendTask.cpp
new file mode 100644
--- /dev/null
+++ b/Engine/Plugins/Experimental/Avalanche/Source/AvalancheCamera/Private/Transition/Tasks/AvaCameraBlendTask.cpp
@@ -0,0 +1,71 @@
+// Copyright Epic Games, Inc. All Rights Reserved.
+
+#include "Transition/Tasks/AvaCameraBlendTask.h"
+#include "AvaCameraSubsystem.h"
+#include "AvaTransitionContext.h"
+#include "StateTreeExecutionContext.h"
+#include "StateTreeLinker.h"
+
+FAvaCameraBlendInstanceData::FAvaCameraBlendInstanceData()
+{
+	TransitionParams.bLockOutgoing = true;
+}
+
+bool FAvaCameraBlendTask::Link(FStateTreeLinker& InLinker)
+{
+	Super::Link(InLinker);
+	InLinker.LinkExternalData(CameraSubsystemHandle);
+	return true;
+}
+
+EStateTreeRunStatus FAvaCameraBlendTask::EnterState(FStateTreeExecutionContext& InContext, const FStateTreeTransitionResult& InTransition) const
+{
+	const ULevel* TransitionLevel = GetTransitionLevel(InContext);
+	if (!TransitionLevel)
+	{
+		return EStateTreeRunStatus::Failed;
+	}
+
+	const FInstanceDataType& InstanceData = InContext.GetInstanceData(*this);
+
+	const UAvaCameraSubsystem& CameraSubsystem = InContext.GetExternalData(CameraSubsystemHandle);
+
+	const FViewTargetTransitionParams* TransitionParams = nullptr;
+	if (InstanceData.bOverrideTransitionParams)
+	{
+		TransitionParams = &InstanceData.TransitionParams;
+	}
+
+	CameraSubsystem.UpdatePlayerControllerViewTarget(TransitionParams);
+
+	return CameraSubsystem.IsBlendingToViewTarget(TransitionLevel)
+		? EStateTreeRunStatus::Running
+		: EStateTreeRunStatus::Succeeded;
+}
+
+EStateTreeRunStatus FAvaCameraBlendTask::Tick(FStateTreeExecutionContext& InContext, const float InDeltaTime) const
+{
+	const ULevel* TransitionLevel = GetTransitionLevel(InContext);
+	if (!TransitionLevel)
+	{
+		return EStateTreeRunStatus::Failed;
+	}
+
+	const UAvaCameraSubsystem& CameraSubsystem = InContext.GetExternalData(CameraSubsystemHandle);
+
+	return CameraSubsystem.IsBlendingToViewTarget(TransitionLevel)
+		? EStateTreeRunStatus::Succeeded
+		: EStateTreeRunStatus::Running;
+}
+
+const ULevel* FAvaCameraBlendTask::GetTransitionLevel(FStateTreeExecutionContext& InContext) const
+{
+	const FAvaTransitionContext& TransitionContext = InContext.GetExternalData(TransitionContextHandle);
+
+	if (const FAvaTransitionScene* TransitionScene = TransitionContext.GetTransitionScene())
+	{
+		return TransitionScene->GetLevel();
+	}
+
+	return nullptr;
+}
diff --git a/Engine/Plugins/Experimental/Avalanche/Source/AvalancheCamera/Public/AvaCameraPriorityModifier.h b/Engine/Plugins/Experimental/Avalanche/Source/AvalancheCamera/Public/AvaCameraPriorityModifier.h
new file mode 100644
--- /dev/null
+++ b/Engine/Plugins/Experimental/Avalanche/Source/AvalancheCamera/Public/AvaCameraPriorityModifier.h
@@ -0,0 +1,39 @@
+// Copyright Epic Games, Inc. All Rights Reserved.
+
+#pragma once
+
+#include "Camera/PlayerCameraManager.h"
+#include "Modifiers/ActorModifierCoreBase.h"
+#include "AvaCameraPriorityModifier.generated.h"
+
+UCLASS(MinimalAPI, BlueprintType)
+class UAvaCameraPriorityModifier : public UActorModifierCoreBase
+{
+	GENERATED_BODY()
+
+public:
+	UAvaCameraPriorityModifier();
+
+	int32 GetPriority() const
+	{
+		return Priority;
+	}
+
+	const FViewTargetTransitionParams& GetTransitionParams() const
+	{
+		return TransitionParams;
+	}
+
+protected:
+	//~ Begin UActorModifierCoreBase
+	virtual void OnModifierCDOSetup(FActorModifierCoreMetadata& InMetadata) override;
+	virtual void Apply() override;
+	//~ End UActorModifierCoreBase
+
+	UPROPERTY(EditInstanceOnly, Category="Camera", meta=(ClampMin="0", UIMin="0", AllowPrivateAccess="true"))
+	int32 Priority = 0;
+
+	/** The parameters to blend to the view target */
+	UPROPERTY(EditInstanceOnly, Category="Camera", meta=(AllowPrivateAccess="true"))
+	FViewTargetTransitionParams TransitionParams;
+};
diff --git a/Engine/Plugins/Experimental/Avalanche/Source/AvalancheCamera/Public/AvaCameraSubsystem.h b/Engine/Plugins/Experimental/Avalanche/Source/AvalancheCamera/Public/AvaCameraSubsystem.h
new file mode 100644
--- /dev/null
+++ b/Engine/Plugins/Experimental/Avalanche/Source/AvalancheCamera/Public/AvaCameraSubsystem.h
@@ -0,0 +1,67 @@
+// Copyright Epic Games, Inc. All Rights Reserved.
+
+#pragma once
+
+#include "Subsystems/WorldSubsystem.h"
+#include "AvaCameraSubsystem.generated.h"
+
+class AActor;
+class APlayerController;
+class IAvaTransitionExecutor;
+class UAvaCameraPriorityModifier;
+struct FViewTargetTransitionParams;
+class ULevel;
+
+USTRUCT()
+struct FAvaViewTarget
+{
+	GENERATED_BODY()
+
+	bool IsValid() const;
+
+	/** The View Target Actor */
+	UPROPERTY()
+	TObjectPtr<AActor> Actor;
+
+	/** The Camera Modifier that the View Target Actor has */
+	UPROPERTY()
+	TObjectPtr<UAvaCameraPriorityModifier> CameraPriorityModifier;
+};
+
+UCLASS(MinimalAPI)
+class UAvaCameraSubsystem : public UWorldSubsystem
+{
+	GENERATED_BODY()
+
+public:
+	AVALANCHECAMERA_API static UAvaCameraSubsystem* Get(const UObject* InObject);
+
+	AVALANCHECAMERA_API void RegisterScene(const ULevel* InSceneLevel);
+	AVALANCHECAMERA_API void UnregisterScene(const ULevel* InSceneLevel);
+
+	bool IsBlendingToViewTarget(const ULevel* InSceneLevel) const;
+
+	void UpdatePlayerControllerViewTarget(const FViewTargetTransitionParams* InOverrideTransitionParams = nullptr) const;
+
+protected:
+	//~ Begin UWorldSubsystem
+	virtual bool DoesSupportWorldType(const EWorldType::Type InWorldType) const override;
+	virtual void OnWorldBeginPlay(UWorld& InWorld) override;
+	//~ End UWorldSubsystem
+
+	//~ Begin USubsystem
+	virtual void Initialize(FSubsystemCollectionBase& InCollection) override;
+	virtual void Deinitialize() override;
+	//~ End USubsystem
+
+private:
+	void OnTransitionStart(const IAvaTransitionExecutor& InExecutor);
+
+	bool HasCustomViewTargetting(const ULevel* InSceneLevel) const;
+
+	UPROPERTY()
+	TObjectPtr<APlayerController> PlayerController;
+
+	UPROPERTY()
+	TArray<FAvaViewTarget> ViewTargets;
+};
diff --git a/Engine/Plugins/Experimental/Avalanche/Source/AvalancheCamera/Public/Transition/Tasks/AvaCameraBlendTask.h b/Engine/Plugins/Experimental/Avalanche/Source/AvalancheCamera/Public/Transition/Tasks/AvaCameraBlendTask.h
new file mode 100644
--- /dev/null
+++ b/Engine/Plugins/Experimental/Avalanche/Source/AvalancheCamera/Public/Transition/Tasks/AvaCameraBlendTask.h
@@ -0,0 +1,46 @@
+// Copyright Epic Games, Inc. All Rights Reserved.
+
+#pragma once
+
+#include "Camera/PlayerCameraManager.h"
+#include "Tasks/AvaTransitionTask.h"
+#include "AvaCameraBlendTask.generated.h"
+
+class UAvaCameraSubsystem;
+
+USTRUCT()
+struct FAvaCameraBlendInstanceData
+{
+	GENERATED_BODY()
+
+	FAvaCameraBlendInstanceData();
+
+	UPROPERTY(EditAnywhere, Category="Camera")
+	bool bOverrideTransitionParams = false;
+
+	UPROPERTY(EditAnywhere, Category="Camera", meta=(EditCondition="bOverrideTransitionParams"))
+	FViewTargetTransitionParams TransitionParams;
+};
+
+USTRUCT(DisplayName="Blend Scene Camera", Category="Scene Camera")
+struct FAvaCameraBlendTask : public FAvaTransitionTask
+{
+	GENERATED_BODY()
+
+	using FInstanceDataType = FAvaCameraBlendInstanceData;
+
+	//~ Begin FStateTreeNodeBase
+	virtual const UStruct* GetInstanceDataType() const override { return FInstanceDataType::StaticStruct(); }
+	virtual bool Link(FStateTreeLinker& InLinker) override;
+	//~ End FStateTreeNodeBase
+
+	//~ Begin FStateTreeTaskBase
+	AVALANCHECAMERA_API virtual EStateTreeRunStatus EnterState(FStateTreeExecutionContext& InContext, const FStateTreeTransitionResult& InTransition) const override;
+	AVALANCHECAMERA_API virtual EStateTreeRunStatus Tick(FStateTreeExecutionContext& InContext, const float InDeltaTime) const override;
+	//~ End FStateTreeTaskBase
+
+protected:
+	const ULevel* GetTransitionLevel(FStateTreeExecutionContext& InContext) const;
+
+	TStateTreeExternalDataHandle<UAvaCameraSubsystem> CameraSubsystemHandle;
+};
diff --git a/Engine/Plugins/Experimental/Avalanche/Source/AvalancheMedia/Private/Playable/AvaPlayable.cpp b/Engine/Plugins/Experimental/Avalanche/Source/AvalancheMedia/Private/Playable/AvaPlayable.cpp
--- a/Engine/Plugins/Experimental/Avalanche/Source/AvalancheMedia/Private/Playable/AvaPlayable.cpp
+++ b/Engine/Plugins/Experimental/Avalanche/Source/AvalancheMedia/Private/Playable/AvaPlayable.cpp
@@ -307,15 +307,10 @@ void UAvaPlayable::EndPlay(EAvaPlayableEndPlayOptions InOptions)
 	UAvaSequencePlayer::OnSequenceFinished().RemoveAll(this);
 	OnEndPlay();
 
-	if (PlayableGroup)
+	if (PlayableGroup && EnumHasAnyFlags(InOptions, EAvaPlayableEndPlayOptions::ConditionalEndPlayWorld) && !PlayableGroup->HasPlayingPlayables())
 	{
-		PlayableGroup->UpdateCameraSetup();
-
-		if (EnumHasAnyFlags(InOptions, EAvaPlayableEndPlayOptions::ConditionalEndPlayWorld) && !PlayableGroup->HasPlayingPlayables())
-		{
-			const bool bForceImmediate = EnumHasAnyFlags(InOptions, EAvaPlayableEndPlayOptions::ForceImmediate);
-			PlayableGroup->RequestEndPlayWorld(bForceImmediate);
-		}
+		const bool bForceImmediate = EnumHasAnyFlags(InOptions, EAvaPlayableEndPlayOptions::ForceImmediate);
+		PlayableGroup->RequestEndPlayWorld(bForceImmediate);
 	}
 }
 
diff --git a/Engine/Plugins/Experimental/Avalanche/Source/AvalancheMedia/Private/Playable/AvaPlayableGroup.cpp b/Engine/Plugins/Experimental/Avalanche/Source/AvalancheMedia/Private/Playable/AvaPlayableGroup.cpp
--- a/Engine/Plugins/Experimental/Avalanche/Source/AvalancheMedia/Private/Playable/AvaPlayableGroup.cpp
+++ b/Engine/Plugins/Experimental/Avalanche/Source/AvalancheMedia/Private/Playable/AvaPlayableGroup.cpp
@@ -285,35 +285,6 @@ void UAvaPlayableGroup::SetLastAppliedCameraPlayable(UAvaPlayable* InPlayable)
 	LastAppliedCameraPlayableWeak = InPlayable;
 }
 
-bool UAvaPlayableGroup::UpdateCameraSetup()
-{
-	// With rigs in sub-playables, it may still be valid and won't need updating.
-	if (const UAvaPlayable* LastAppliedCameraPlayable = LastAppliedCameraPlayableWeak.Get())
-	{
-		if (LastAppliedCameraPlayable->GetPlayableStatus() == EAvaPlayableStatus::Visible
-			&& LastAppliedCameraPlayable->GetShouldBeVisible())
-		{
-			return true;	// Camera setup is still valid.
-		}
-	}
-
-	// This is called when a playable from the group is stopped. We need to
-	// select a playable in the group that is still playing and will use it's camera.
-	for (const TObjectKey<UAvaPlayable>& PlayableKey : Playables )
-	{
-		UAvaPlayable* Playable = PlayableKey.ResolveObjectPtr();
-		
-		if (Playable && Playable->IsPlaying())
-		{
-			if (Playable->ApplyCamera())
-			{
-				return true;
-			}
-		}
-	}
-	return false;
-}
-
 UTextureRenderTarget2D* UAvaPlayableGroup::GetRenderTarget() const
 {
 	return ManagedRenderTarget.Get();
diff --git a/Engine/Plugins/Experimental/Avalanche/Source/AvalancheMedia/Private/Playable/PlayableGroups/AvaGameInstancePlayableGroup.cpp b/Engine/Plugins/Experimental/Avalanche/Source/AvalancheMedia/Private/Playable/PlayableGroups/AvaGameInstancePlayableGroup.cpp
--- a/Engine/Plugins/Experimental/Avalanche/Source/AvalancheMedia/Private/Playable/PlayableGroups/AvaGameInstancePlayableGroup.cpp
+++ b/Engine/Plugins/Experimental/Avalanche/Source/AvalancheMedia/Private/Playable/PlayableGroups/AvaGameInstancePlayableGroup.cpp
@@ -202,14 +202,3 @@ bool UAvaGameInstancePlayableGroup::ConditionalRequestUnloadWorld(bool bForceImm
 	}
 	return false;
 }
-
-void UAvaGameInstancePlayableGroup::QueueCameraCut()
-{
-	if (const UAvaGameInstance* AvaGameInstance = GetAvaGameInstance())
-	{
-		if (UAvaGameViewportClient* GameViewportClient = AvaGameInstance->GetAvaGameViewportClient())
-		{
-			GameViewportClient->SetCameraCutThisFrame();
-		}
-	}
-}
\ No newline at end of file
diff --git a/Engine/Plugins/Experimental/Avalanche/Source/AvalancheMedia/Private/Playable/PlayableGroups/AvaGameInstancePlayableGroup.h b/Engine/Plugins/Experimental/Avalanche/Source/AvalancheMedia/Private/Playable/PlayableGroups/AvaGameInstancePlayableGroup.h
--- a/Engine/Plugins/Experimental/Avalanche/Source/AvalancheMedia/Private/Playable/PlayableGroups/AvaGameInstancePlayableGroup.h
+++ b/Engine/Plugins/Experimental/Avalanche/Source/AvalancheMedia/Private/Playable/PlayableGroups/AvaGameInstancePlayableGroup.h
@@ -32,7 +32,6 @@ public:
 	virtual bool ConditionalBeginPlay(const FAvaInstancePlaySettings& InWorldPlaySettings) override;
 	virtual void RequestEndPlayWorld(bool bInForceImmediate) override;
 	virtual bool ConditionalRequestUnloadWorld(bool bForceImmediate) override;
-	virtual void QueueCameraCut() override;
 	virtual bool IsWorldPlaying() const override;
 	virtual bool IsRenderTargetReady() const override;
 	virtual UTextureRenderTarget2D* GetRenderTarget() const override;
diff --git a/Engine/Plugins/Experimental/Avalanche/Source/AvalancheMedia/Private/Playable/Playables/AvaPlayableLevelStreaming.cpp b/Engine/Plugins/Experimental/Avalanche/Source/AvalancheMedia/Private/Playable/Playables/AvaPlayableLevelStreaming.cpp
--- a/Engine/Plugins/Experimental/Avalanche/Source/AvalancheMedia/Private/Playable/Playables/AvaPlayableLevelStreaming.cpp
+++ b/Engine/Plugins/Experimental/Avalanche/Source/AvalancheMedia/Private/Playable/Playables/AvaPlayableLevelStreaming.cpp
@@ -23,7 +23,6 @@
 #include "SceneView.h"
 #include "Streaming/LevelStreamingDelegates.h"
 #include "UObject/Package.h"
-#include "Viewport/AvaCameraManager.h"
 
 #define LOCTEXT_NAMESPACE "AvaPlayableLevelStreaming"
 
@@ -332,57 +331,6 @@ IAvaSceneInterface* UAvaPlayableLevelStreaming::GetSceneInterface() const
 	return nullptr;
 }
 
-bool UAvaPlayableLevelStreaming::ApplyCamera()
-{
-	using namespace UE::AvaMedia::LevelStreamingPlayable;
-
-	if (!Scene)
-	{
-		return false;
-	}
-
-	bool bSetupDone = false;
-	UAvaPlayable* PlayableOwningCamera = nullptr;
-
-	// Camera setup.
-	if (AActor* CameraActor = FindStartupCameraActor(NAME_None, &PlayableOwningCamera))
-	{
-		// New code using player controller.
-		// Note: doesn't work yet because there is no player controller in Motion Design game instance, but it works in PIE.
-		if (APlayerController* PlayerController = GetPlayableGroup()->GetPlayWorld()->GetFirstPlayerController())
-		{
-			PlayerController->SetViewTargetWithBlend(CameraActor);
-			GetPlayableGroup()->SetLastAppliedCameraPlayable(PlayableOwningCamera);
-			bSetupDone = true;
-		}
-		else
-		{
-			// If there is no player controller, fallback to Motion Design camera manager, if the game instance has it.
-			const UAvaGameInstance* AvaGameInstance = Cast<UAvaGameInstance>(GetPlayableGroup()->GetGameInstance());			
-			if (const UAvaGameViewportClient* ViewportClient = AvaGameInstance ? AvaGameInstance->GetAvaGameViewportClient() : nullptr)
-			{
-				// Old code using ava camera manager. To retire.
-				constexpr bool bIsCanvasController = false;
-				ViewportClient->GetCameraManager()->Init(Scene->GetPlaybackObject(), bIsCanvasController);
-				ViewportClient->GetCameraManager()->SetViewTarget(CameraActor);
-				GetPlayableGroup()->SetLastAppliedCameraPlayable(PlayableOwningCamera);
-				bSetupDone = true;
-				UE_LOG(LogAvaPlayable, Log, TEXT("No player controller found in \"%s\" - Using Ava Camera Manager instead."), *LevelStreaming->PackageNameToLoad.ToString())
-			}
-			else
-			{
-				UE_LOG(LogAvaPlayable, Warning, TEXT("No player controller found in \"%s\" - Ava Camera Manager not available either."), *LevelStreaming->PackageNameToLoad.ToString());				
-			}
-		}
-	}
-	else
-	{
-		UE_LOG(LogAvaPlayable, Error, TEXT("Failed to find camera \"%s\", or any camera actor, in loaded level."),
-			*Private::GetStartupCameraName(Scene).ToString());
-	}
-	return bSetupDone;
-}
-
 bool UAvaPlayableLevelStreaming::GetShouldBeVisible() const
 {
 	return LevelStreaming ? LevelStreaming->GetShouldBeVisibleFlag() : false;
@@ -746,13 +694,6 @@ void UAvaPlayableLevelStreaming::OnPlay()
 
 	// Ensure scene is resolved.
 	ResolveScene(Level);
-
-	if (!Scene)
-	{
-		return;
-	}
-
-	ApplyCamera();
 }
 
 void UAvaPlayableLevelStreaming::OnEndPlay()
diff --git a/Engine/Plugins/Experimental/Avalanche/Source/AvalancheMedia/Private/Playable/Playables/AvaPlayableLevelStreaming.h b/Engine/Plugins/Experimental/Avalanche/Source/AvalancheMedia/Private/Playable/Playables/AvaPlayableLevelStreaming.h
--- a/Engine/Plugins/Experimental/Avalanche/Source/AvalancheMedia/Private/Playable/Playables/AvaPlayableLevelStreaming.h
+++ b/Engine/Plugins/Experimental/Avalanche/Source/AvalancheMedia/Private/Playable/Playables/AvaPlayableLevelStreaming.h
@@ -24,7 +24,6 @@ public:
 	virtual const FSoftObjectPath& GetSourceAssetPath() const override { return SourceLevel.ToSoftObjectPath(); }
 	virtual EAvaPlayableStatus GetPlayableStatus() const override { return PlayableStatus; }
 	virtual IAvaSceneInterface* GetSceneInterface() const override;
-	virtual bool ApplyCamera() override;
 	virtual bool GetShouldBeVisible() const override;
 	virtual void SetShouldBeVisible(bool bInShouldBeVisible) override;
 	virtual void SetupView(FSceneViewFamily& InViewFamily, FSceneView& InView) override;
diff --git a/Engine/Plugins/Experimental/Avalanche/Source/AvalancheMedia/Private/Playable/Playables/AvaPlayableRemoteProxy.cpp b/Engine/Plugins/Experimental/Avalanche/Source/AvalancheMedia/Private/Playable/Playables/AvaPlayableRemoteProxy.cpp
--- a/Engine/Plugins/Experimental/Avalanche/Source/AvalancheMedia/Private/Playable/Playables/AvaPlayableRemoteProxy.cpp
+++ b/Engine/Plugins/Experimental/Avalanche/Source/AvalancheMedia/Private/Playable/Playables/AvaPlayableRemoteProxy.cpp
@@ -231,11 +231,6 @@ EAvaPlayableCommandResult UAvaPlayableRemoteProxy::UpdateRemoteControlCommand(co
 	return EAvaPlayableCommandResult::Executed;
 }
 
-bool UAvaPlayableRemoteProxy::ApplyCamera()
-{
-	return false;
-}
-
 void UAvaPlayableRemoteProxy::SetUserData(const FString& InUserData)
 {
 	if (UserData != InUserData)
diff --git a/Engine/Plugins/Experimental/Avalanche/Source/AvalancheMedia/Private/Playable/Playables/AvaPlayableRemoteProxy.h b/Engine/Plugins/Experimental/Avalanche/Source/AvalancheMedia/Private/Playable/Playables/AvaPlayableRemoteProxy.h
--- a/Engine/Plugins/Experimental/Avalanche/Source/AvalancheMedia/Private/Playable/Playables/AvaPlayableRemoteProxy.h
+++ b/Engine/Plugins/Experimental/Avalanche/Source/AvalancheMedia/Private/Playable/Playables/AvaPlayableRemoteProxy.h
@@ -29,7 +29,6 @@ public:
 	virtual IAvaSceneInterface* GetSceneInterface() const override;
 	virtual EAvaPlayableCommandResult ExecuteAnimationCommand(EAvaPlaybackAnimAction InAnimAction, const FAvaPlaybackAnimPlaySettings& InAnimPlaySettings) override;
 	virtual EAvaPlayableCommandResult UpdateRemoteControlCommand(const TSharedRef<FAvaPlayableRemoteControlValues>& InRemoteControlValues) override;
-	virtual bool ApplyCamera() override;
 	virtual bool IsRemoteProxy() const override { return true; }
 	virtual void SetUserData(const FString& InUserData) override;
 
diff --git a/Engine/Plugins/Experimental/Avalanche/Source/AvalancheMedia/Public/Playable/AvaPlayable.h b/Engine/Plugins/Experimental/Avalanche/Source/AvalancheMedia/Public/Playable/AvaPlayable.h
--- a/Engine/Plugins/Experimental/Avalanche/Source/AvalancheMedia/Public/Playable/AvaPlayable.h
+++ b/Engine/Plugins/Experimental/Avalanche/Source/AvalancheMedia/Public/Playable/AvaPlayable.h
@@ -116,7 +116,6 @@ public:
 	virtual IAvaSceneInterface* GetSceneInterface() const {return nullptr;}
 	virtual EAvaPlayableCommandResult ExecuteAnimationCommand(EAvaPlaybackAnimAction InAnimAction, const FAvaPlaybackAnimPlaySettings& InAnimPlaySettings);
 	virtual EAvaPlayableCommandResult UpdateRemoteControlCommand(const TSharedRef<FAvaPlayableRemoteControlValues>& InRemoteControlValues);
-	virtual bool ApplyCamera() { return false; }
 	virtual bool IsRemoteProxy() const { return false; }
 	virtual bool GetShouldBeVisible() const { return true; }
 	virtual void SetShouldBeVisible(bool bInShouldBeVisible) {}
diff --git a/Engine/Plugins/Experimental/Avalanche/Source/AvalancheMedia/Public/Playable/AvaPlayableGroup.h b/Engine/Plugins/Experimental/Avalanche/Source/AvalancheMedia/Public/Playable/AvaPlayableGroup.h
--- a/Engine/Plugins/Experimental/Avalanche/Source/AvalancheMedia/Public/Playable/AvaPlayableGroup.h
+++ b/Engine/Plugins/Experimental/Avalanche/Source/AvalancheMedia/Public/Playable/AvaPlayableGroup.h
@@ -109,9 +109,6 @@ public:
 
 	/** Keep track of the last playable that applied it's camera in the viewport/controller. */
 	void SetLastAppliedCameraPlayable(UAvaPlayable* InPlayable);
-	
-	/** Determines which camera, from currently playing playables, should be used and use it. */
-	bool UpdateCameraSetup();
 
 	virtual bool IsWorldPlaying() const { return true; }
 
@@ -158,11 +155,6 @@ public:
 	 */ 
 	virtual bool ConditionalRequestUnloadWorld(bool bForceImmediate) { return true; }
 
-	/**
-	 * @brief Queue a camera cut for the next rendered frame.
-	 */
-	virtual void QueueCameraCut() {}
-
 	/**
 	 * @brief Notify the playable group that a playable is loading an asset.
 	 * @param InPlayable Source of the event.
diff --git a/Engine/Plugins/Experimental/Avalanche/Source/AvalancheModifiersEditor/Private/AvaModifiersEditorModule.cpp b/Engine/Plugins/Experimental/Avalanche/Source/AvalancheModifiersEditor/Private/AvaModifiersEditorModule.cpp
--- a/Engine/Plugins/Experimental/Avalanche/Source/AvalancheModifiersEditor/Private/AvaModifiersEditorModule.cpp
+++ b/Engine/Plugins/Experimental/Avalanche/Source/AvalancheModifiersEditor/Private/AvaModifiersEditorModule.cpp
@@ -1,7 +1,6 @@
 // Copyright Epic Games, Inc. All Rights Reserved.
 
 #include "AvaModifiersEditorModule.h"
-#include "AvaDefs.h"
 #include "AvaModifiersEditorStyle.h"
 #include "DragDropOps/AvaOutlinerItemDragDropOp.h"
 #include "IAvaOutliner.h"
diff --git a/Engine/Plugins/Experimental/Avalanche/Source/AvalancheTransition/Private/AvaTransitionSubsystem.cpp b/Engine/Plugins/Experimental/Avalanche/Source/AvalancheTransition/Private/AvaTransitionSubsystem.cpp
--- a/Engine/Plugins/Experimental/Avalanche/Source/AvalancheTransition/Private/AvaTransitionSubsystem.cpp
+++ b/Engine/Plugins/Experimental/Avalanche/Source/AvalancheTransition/Private/AvaTransitionSubsystem.cpp
@@ -8,14 +8,27 @@
 #include "Execution/IAvaTransitionExecutor.h"
 #include "StateTreeExecutionTypes.h"
 
-void UAvaTransitionSubsystem::RegisterTransitionBehavior(ULevel* InLevel, IAvaTransitionBehavior* InBehavior)
+void UAvaTransitionSubsystem::RegisterTransitionBehavior(const ULevel* InLevel, IAvaTransitionBehavior* InBehavior)
 {
 	TransitionBehaviors.Add(InLevel, TWeakInterfacePtr<IAvaTransitionBehavior>(InBehavior));
 }
 
 IAvaTransitionBehavior* UAvaTransitionSubsystem::GetOrCreateTransitionBehavior(ULevel* InLevel)
 {
-	if (!EnsureLevelIsAppropriate(InLevel))
+	UWorld* const World = GetWorld();
+	if (!World)
+	{
+		return nullptr;
+	}
+
+	// if passed in nullptr, set it to the persistent level
+	if (!InLevel)
+	{
+		InLevel = World->PersistentLevel;
+	}
+
+	// Ensure the provided level belongs to the world of this subsystem
+	if (!InLevel || !ensureAlways(InLevel->OwningWorld == World))
 	{
 		return nullptr;
 	}
@@ -36,9 +49,22 @@ IAvaTransitionBehavior* UAvaTransitionSubsystem::GetOrCreateTransitionBehavior(U
 	return InLevel->OwningWorld->SpawnActor<AAvaTransitionBehaviorActor>(SpawnParameters);		
 }
 
-IAvaTransitionBehavior* UAvaTransitionSubsystem::GetTransitionBehavior(ULevel* InLevel) const
+IAvaTransitionBehavior* UAvaTransitionSubsystem::GetTransitionBehavior(const ULevel* InLevel) const
 {
-	if (!EnsureLevelIsAppropriate(InLevel))
+	UWorld* const World = GetWorld();
+	if (!World)
+	{
+		return nullptr;
+	}
+
+	// if passed in nullptr, set it to the persistent level
+	if (!InLevel)
+	{
+		InLevel = World->PersistentLevel;
+	}
+
+	// Ensure the provided level belongs to the world of this subsystem
+	if (!InLevel || !ensureAlways(InLevel->OwningWorld == World))
 	{
 		return nullptr;
 	}
@@ -60,7 +86,7 @@ IAvaTransitionBehavior* UAvaTransitionSubsystem::GetTransitionBehavior(ULevel* I
 	return nullptr;
 }
 
-IAvaTransitionBehavior* UAvaTransitionSubsystem::FindTransitionBehavior(ULevel* InLevel)
+IAvaTransitionBehavior* UAvaTransitionSubsystem::FindTransitionBehavior(const ULevel* InLevel)
 {
 	// find existing transition behavior actor in level
 	AAvaTransitionBehaviorActor* TransitionBehaviorActor;
@@ -127,21 +153,3 @@ void UAvaTransitionSubsystem::PostInitialize()
 		}
 	}
 }
-
-bool UAvaTransitionSubsystem::EnsureLevelIsAppropriate(ULevel*& InLevel) const
-{
-	UWorld* const World = GetWorld();
-	if (!World)
-	{
-		return false;
-	}
-
-	// if passed in nullptr, set it to the persistent level
-	if (!InLevel)
-	{
-		InLevel = World->PersistentLevel;
-	}
-
-	// Ensure the provided level belongs to the world of this subsystem
-	return InLevel && ensureAlways(InLevel->OwningWorld == World);
-}
diff --git a/Engine/Plugins/Experimental/Avalanche/Source/AvalancheTransition/Private/AvaTransitionTree.cpp b/Engine/Plugins/Experimental/Avalanche/Source/AvalancheTransition/Private/AvaTransitionTree.cpp
--- a/Engine/Plugins/Experimental/Avalanche/Source/AvalancheTransition/Private/AvaTransitionTree.cpp
+++ b/Engine/Plugins/Experimental/Avalanche/Source/AvalancheTransition/Private/AvaTransitionTree.cpp
@@ -2,6 +2,8 @@
 
 #include "AvaTransitionTree.h"
 
+#include "StateTreeTaskBase.h"
+
 FAvaTagHandle UAvaTransitionTree::GetTransitionLayer() const
 {
 	return TransitionLayer;
@@ -21,3 +23,31 @@ void UAvaTransitionTree::SetEnabled(bool bInEnabled)
 {
 	bEnabled = bInEnabled;
 }
+
+bool UAvaTransitionTree::ContainsTask(const UScriptStruct* InTaskStruct) const
+{
+	for (const FCompactStateTreeState& State : GetStates())
+	{
+		if (!State.bEnabled)
+		{
+			continue;
+		}
+
+		for (int32 TaskIndex = State.TasksBegin; TaskIndex < State.TasksBegin + State.TasksNum; ++TaskIndex)
+		{
+			const FConstStructView Node = GetNodes()[TaskIndex];
+
+			const UScriptStruct* TaskStruct = Node.GetScriptStruct();
+			if (TaskStruct && TaskStruct->IsChildOf(InTaskStruct))
+			{
+				const FStateTreeTaskBase* TaskNode = Node.GetPtr<const FStateTreeTaskBase>();
+				if (TaskNode && TaskNode->bTaskEnabled)
+				{
+					return true;
+				}
+			}
+		}
+	}
+
+	return false;
+}
diff --git a/Engine/Plugins/Experimental/Avalanche/Source/AvalancheTransition/Private/Execution/AvaTransitionExecutor.cpp b/Engine/Plugins/Experimental/Avalanche/Source/AvalancheTransition/Private/Execution/AvaTransitionExecutor.cpp
--- a/Engine/Plugins/Experimental/Avalanche/Source/AvalancheTransition/Private/Execution/AvaTransitionExecutor.cpp
+++ b/Engine/Plugins/Experimental/Avalanche/Source/AvalancheTransition/Private/Execution/AvaTransitionExecutor.cpp
@@ -12,6 +12,8 @@
 
 DEFINE_LOG_CATEGORY_STATIC(LogAvaTransitionExecutor, Log, All);
 
+TMulticastDelegate<void(const IAvaTransitionExecutor&)> IAvaTransitionExecutor::OnTransitionStart;
+
 FAvaTransitionExecutor::FAvaTransitionExecutor(FAvaTransitionExecutorBuilder& InBuilder)
 	: NullInstance(MoveTemp(InBuilder.NullInstance))
 	, ContextName(MoveTemp(InBuilder.ContextName))
@@ -126,6 +128,8 @@ void FAvaTransitionExecutor::Start()
 
 	Setup();
 
+	GetOnTransitionStart().Broadcast(*this);
+
 	ForEachInstance(
 		[](FAvaTransitionBehaviorInstance& InInstance)
 		{
@@ -162,6 +166,11 @@ TArray<const FAvaTransitionBehaviorInstance*> FAvaTransitionExecutor::GetBehavio
 	return OutInstances;
 }
 
+void FAvaTransitionExecutor::ForEachBehaviorInstance(TFunctionRef<void(const FAvaTransitionBehaviorInstance&)> InCallable) const
+{
+	ForEachInstance(InCallable);
+}
+
 void FAvaTransitionExecutor::Stop()
 {
 	ForEachInstance(
diff --git a/Engine/Plugins/Experimental/Avalanche/Source/AvalancheTransition/Private/Execution/AvaTransitionExecutor.h b/Engine/Plugins/Experimental/Avalanche/Source/AvalancheTransition/Private/Execution/AvaTransitionExecutor.h
--- a/Engine/Plugins/Experimental/Avalanche/Source/AvalancheTransition/Private/Execution/AvaTransitionExecutor.h
+++ b/Engine/Plugins/Experimental/Avalanche/Source/AvalancheTransition/Private/Execution/AvaTransitionExecutor.h
@@ -29,6 +29,7 @@ protected:
 
 	//~ Begin IAvaTransitionExecutor
 	virtual TArray<const FAvaTransitionBehaviorInstance*> GetBehaviorInstances(const FAvaTransitionLayerComparator& InComparator) const;
+	virtual void ForEachBehaviorInstance(TFunctionRef<void(const FAvaTransitionBehaviorInstance&)> InCallable) const override;
 	virtual void Start() override;
 	virtual void Stop() override;
 	//~ End IAvaTransitionExecutor
diff --git a/Engine/Plugins/Experimental/Avalanche/Source/AvalancheTransition/Public/AvaTransitionSubsystem.h b/Engine/Plugins/Experimental/Avalanche/Source/AvalancheTransition/Public/AvaTransitionSubsystem.h
--- a/Engine/Plugins/Experimental/Avalanche/Source/AvalancheTransition/Public/AvaTransitionSubsystem.h
+++ b/Engine/Plugins/Experimental/Avalanche/Source/AvalancheTransition/Public/AvaTransitionSubsystem.h
@@ -8,8 +8,8 @@
 #include "Subsystems/WorldSubsystem.h"
 #include "Templates/Function.h"
 #include "Templates/SharedPointer.h"
+#include "UObject/ObjectKey.h"
 #include "UObject/WeakInterfacePtr.h"
-#include "UObject/WeakObjectPtr.h"
 #include "AvaTransitionSubsystem.generated.h"
 
 class IAvaTransitionBehavior;
@@ -25,7 +25,7 @@ class UAvaTransitionSubsystem : public UWorldSubsystem
 	GENERATED_BODY()
 
 public:
-	void RegisterTransitionBehavior(ULevel* InLevel, IAvaTransitionBehavior* InBehavior);
+	void RegisterTransitionBehavior(const ULevel* InLevel, IAvaTransitionBehavior* InBehavior);
 
 	/**
 	 * Gets or Creates a Transition Behavior for the given Level in this Subsystem's World
@@ -39,9 +39,9 @@ public:
 	 * @param InLevel the level containing the transition behavior. if null, will default to the persistent level of the world
 	 * @return the behavior for the provided level
 	 */
-	AVALANCHETRANSITION_API IAvaTransitionBehavior* GetTransitionBehavior(ULevel* InLevel = nullptr) const;
+	AVALANCHETRANSITION_API IAvaTransitionBehavior* GetTransitionBehavior(const ULevel* InLevel = nullptr) const;
 
-	AVALANCHETRANSITION_API static IAvaTransitionBehavior* FindTransitionBehavior(ULevel* InLevel);
+	AVALANCHETRANSITION_API static IAvaTransitionBehavior* FindTransitionBehavior(const ULevel* InLevel);
 
 	void RegisterTransitionExecutor(const TSharedRef<IAvaTransitionExecutor>& InExecutor);
 
@@ -58,9 +58,7 @@ protected:
 	//~ End UWorldSubsystem
 
 private:
-	bool EnsureLevelIsAppropriate(ULevel*& InLevel) const;
-
-	TMap<TWeakObjectPtr<ULevel>, TWeakInterfacePtr<IAvaTransitionBehavior>> TransitionBehaviors;
+	TMap<TObjectKey<ULevel>, TWeakInterfacePtr<IAvaTransitionBehavior>> TransitionBehaviors;
 
 	TArray<TWeakPtr<IAvaTransitionExecutor>> TransitionExecutors;
 };
diff --git a/Engine/Plugins/Experimental/Avalanche/Source/AvalancheTransition/Public/AvaTransitionTree.h b/Engine/Plugins/Experimental/Avalanche/Source/AvalancheTransition/Public/AvaTransitionTree.h
--- a/Engine/Plugins/Experimental/Avalanche/Source/AvalancheTransition/Public/AvaTransitionTree.h
+++ b/Engine/Plugins/Experimental/Avalanche/Source/AvalancheTransition/Public/AvaTransitionTree.h
@@ -7,23 +7,35 @@
 #include "StateTree.h"
 #include "AvaTransitionTree.generated.h"
 
+struct FAvaTransitionTask;
+
 /**
  * Motion Design Transition Tree is a State Tree with the purpose of executing user-defined logic
  * when there's a Transition between multiple scenes in multiple layers.
  */
-UCLASS(DisplayName = "Motion Design Transition Tree")
-class AVALANCHETRANSITION_API UAvaTransitionTree : public UStateTree
+UCLASS(MinimalAPI, DisplayName = "Motion Design Transition Tree")
+class UAvaTransitionTree : public UStateTree
 {
 	GENERATED_BODY()
 
 public:
-	FAvaTagHandle GetTransitionLayer() const;
+	AVALANCHETRANSITION_API FAvaTagHandle GetTransitionLayer() const;
+
+	AVALANCHETRANSITION_API void SetTransitionLayer(FAvaTagHandle InTransitionLayer);
+
+	AVALANCHETRANSITION_API bool IsEnabled() const;
 
-	void SetTransitionLayer(FAvaTagHandle InTransitionLayer);
+	AVALANCHETRANSITION_API void SetEnabled(bool bInEnabled);
 
-	bool IsEnabled() const;
+	/** Returns whether a Task of a given type exists and is enabled within an enabled state in the Transition Tree */
+	AVALANCHETRANSITION_API bool ContainsTask(const UScriptStruct* InTaskStruct) const;
 
-	void SetEnabled(bool bInEnabled);
+	/** Returns whether a Task of a given type exists and is enabled within an enabled state in the Transition Tree */
+	template <class InTaskType UE_REQUIRES(std::is_base_of_v<FAvaTransitionTask, InTaskType>)>
+	bool ContainsTask() const
+	{
+		return ContainsTask(InTaskType::StaticStruct());
+	}
 
 	void SetInstancingMode(EAvaTransitionInstancingMode InInstancingMode)
 	{
diff --git a/Engine/Plugins/Experimental/Avalanche/Source/AvalancheTransition/Public/Execution/IAvaTransitionExecutor.h b/Engine/Plugins/Experimental/Avalanche/Source/AvalancheTransition/Public/Execution/IAvaTransitionExecutor.h
--- a/Engine/Plugins/Experimental/Avalanche/Source/AvalancheTransition/Public/Execution/IAvaTransitionExecutor.h
+++ b/Engine/Plugins/Experimental/Avalanche/Source/AvalancheTransition/Public/Execution/IAvaTransitionExecutor.h
@@ -19,9 +19,20 @@ public:
 	/** Gets all the Behavior Instances matching the Layer Comparator */
 	virtual TArray<const FAvaTransitionBehaviorInstance*> GetBehaviorInstances(const FAvaTransitionLayerComparator& InComparator) const = 0;
 
+	/** Runs the provided callback for each Behavior Instance */
+	virtual void ForEachBehaviorInstance(TFunctionRef<void(const FAvaTransitionBehaviorInstance&)> InCallable) const = 0;
+
 	/** Start Transition. Should only be called once */
 	virtual void Start() = 0;
 
 	/** Stop the current Execution */
 	virtual void Stop() = 0;
+
+	static TMulticastDelegate<void(const IAvaTransitionExecutor&)>& GetOnTransitionStart()
+	{
+		return OnTransitionStart;
+	}
+
+private:
+	AVALANCHETRANSITION_API static TMulticastDelegate<void(const IAvaTransitionExecutor&)> OnTransitionStart;
 };
