Motion Design: Camera system to consider Actors with Camera Components but without the Camera Priority Modifier to be at priority 0

diff --git a/Engine/Plugins/Experimental/Avalanche/Source/AvalancheCamera/Private/AvaCameraSubsystem.cpp b/Engine/Plugins/Experimental/Avalanche/Source/AvalancheCamera/Private/AvaCameraSubsystem.cpp
--- a/Engine/Plugins/Experimental/Avalanche/Source/AvalancheCamera/Private/AvaCameraSubsystem.cpp
+++ b/Engine/Plugins/Experimental/Avalanche/Source/AvalancheCamera/Private/AvaCameraSubsystem.cpp
@@ -6,6 +6,7 @@
 #include "AvaTransitionTree.h"
 #include "Behavior/AvaTransitionBehaviorInstance.h"
 #include "Behavior/IAvaTransitionBehavior.h"
+#include "Camera/CameraComponent.h"
 #include "Camera/PlayerCameraManager.h"
 #include "Engine/Level.h"
 #include "Execution/IAvaTransitionExecutor.h"
@@ -16,7 +17,27 @@
 
 bool FAvaViewTarget::IsValid() const
 {
-	return ::IsValid(Actor) && ::IsValid(CameraPriorityModifier);
+	return ::IsValid(Actor);
+}
+
+int32 FAvaViewTarget::GetPriority() const
+{
+	if (::IsValid(CameraPriorityModifier))
+	{
+		return CameraPriorityModifier->GetPriority();
+	}
+	return 0;
+}
+
+const FViewTargetTransitionParams& FAvaViewTarget::GetTransitionParams() const
+{
+	if (::IsValid(CameraPriorityModifier))
+	{
+		return CameraPriorityModifier->GetTransitionParams();
+	}
+
+	static const FViewTargetTransitionParams DefaultTransitionParams;
+	return DefaultTransitionParams;
 }
 
 UAvaCameraSubsystem* UAvaCameraSubsystem::Get(const UObject* InObject)
@@ -56,17 +77,24 @@ void UAvaCameraSubsystem::RegisterScene(const ULevel* InSceneLevel)
 		}
 
 		const UActorModifierCoreStack* ModifierStack = ModifierSubsystem->GetActorModifierStack(Actor);
-		if (!ModifierStack)
-		{
-			continue;
-		}
 
-		if (UAvaCameraPriorityModifier* CameraPriorityModifier = ModifierStack->GetClassModifier<UAvaCameraPriorityModifier>())
+		const UAvaCameraPriorityModifier* CameraPriorityModifier = ModifierStack
+			? ModifierStack->GetClassModifier<UAvaCameraPriorityModifier>()
+			: nullptr;
+
+		if (CameraPriorityModifier)
 		{
 			FAvaViewTarget& Entry = ViewTargets.AddDefaulted_GetRef();
 			Entry.Actor = Actor;
 			Entry.CameraPriorityModifier = CameraPriorityModifier;
 		}
+		// Fallback case where the Actor has a Camera Component but does not have a Priority Modifier
+		else if (Actor->FindComponentByClass<UCameraComponent>())
+		{
+			FAvaViewTarget& Entry = ViewTargets.AddDefaulted_GetRef();
+			Entry.Actor = Actor;
+			Entry.CameraPriorityModifier = nullptr;
+		}
 	}
 
 	// Remove invalid entries
@@ -80,7 +108,7 @@ void UAvaCameraSubsystem::RegisterScene(const ULevel* InSceneLevel)
 	ViewTargets.StableSort(
 		[](const FAvaViewTarget& A, const FAvaViewTarget& B)
 		{
-			return A.CameraPriorityModifier->GetPriority() < B.CameraPriorityModifier->GetPriority();
+			return A.GetPriority() < B.GetPriority();
 		});
 
 	if (!HasCustomViewTargetting(InSceneLevel))
@@ -131,7 +159,7 @@ void UAvaCameraSubsystem::UpdatePlayerControllerViewTarget(const FViewTargetTran
 	const FAvaViewTarget& DesiredViewTarget = ViewTargets.Last();
 
 	// These should've been removed prior to calling this private function
-	check(DesiredViewTarget.Actor && DesiredViewTarget.CameraPriorityModifier);
+	check(DesiredViewTarget.Actor);
 
 	// Current View Target matches the desired. Nothing to do
 	if (DesiredViewTarget.Actor == PlayerController->GetViewTarget())
@@ -141,7 +169,7 @@ void UAvaCameraSubsystem::UpdatePlayerControllerViewTarget(const FViewTargetTran
 
 	const FViewTargetTransitionParams* ViewTargetTransitionParams = InOverrideTransitionParams
 		? InOverrideTransitionParams
-		: &DesiredViewTarget.CameraPriorityModifier->GetTransitionParams();
+		: &DesiredViewTarget.GetTransitionParams();
 
 	PlayerController->SetViewTarget(DesiredViewTarget.Actor, *ViewTargetTransitionParams);
 }
diff --git a/Engine/Plugins/Experimental/Avalanche/Source/AvalancheCamera/Public/AvaCameraSubsystem.h b/Engine/Plugins/Experimental/Avalanche/Source/AvalancheCamera/Public/AvaCameraSubsystem.h
--- a/Engine/Plugins/Experimental/Avalanche/Source/AvalancheCamera/Public/AvaCameraSubsystem.h
+++ b/Engine/Plugins/Experimental/Avalanche/Source/AvalancheCamera/Public/AvaCameraSubsystem.h
@@ -17,15 +17,22 @@ struct FAvaViewTarget
 {
 	GENERATED_BODY()
 
+	/** returns true if the View Target Actor is valid */
 	bool IsValid() const;
 
+	/** Returns the priority of the camera priority modifier of the actor, or 0 if the priority modifier doesn't exist */
+	int32 GetPriority() const;
+
+	/** Returns the Transition Params of the camera priority modifier of the actor, or default transition params if the modifier doesn't exist */
+	const FViewTargetTransitionParams& GetTransitionParams() const;
+
 	/** The View Target Actor */
 	UPROPERTY()
 	TObjectPtr<AActor> Actor;
 
-	/** The Camera Modifier that the View Target Actor has */
+	/** The Camera Modifier that the View Target Actor has, if any */
 	UPROPERTY()
-	TObjectPtr<UAvaCameraPriorityModifier> CameraPriorityModifier;
+	TObjectPtr<const UAvaCameraPriorityModifier> CameraPriorityModifier;
 };
 
 UCLASS(MinimalAPI)
