Motion Design: Added shutdown call on Unload world

diff --git a/Engine/Plugins/Experimental/Avalanche/Source/Avalanche/Private/Framework/AvaGameInstance.cpp b/Engine/Plugins/Experimental/Avalanche/Source/Avalanche/Private/Framework/AvaGameInstance.cpp
--- a/Engine/Plugins/Experimental/Avalanche/Source/Avalanche/Private/Framework/AvaGameInstance.cpp
+++ b/Engine/Plugins/Experimental/Avalanche/Source/Avalanche/Private/Framework/AvaGameInstance.cpp
@@ -347,6 +347,8 @@ void UAvaGameInstance::UnloadWorld()
 		PlayWorld->DestroyWorld(true);
 	}
 
+	Shutdown();
+
 	PlayWorld = nullptr;
 	WorldContext = nullptr;
 }
