[Motion Design] Broadcast Capture Overrun mode default to "skip" instead of "flush".

"Flush" is causing large performance spikes under seemingly "normal" conditions. It seems to be normal for the game thread to overrun the rendering thread. Most motion design graphics are rendering heavy (rather than game heavy). So the default action should be to skip.

Exposing the setting in the media setting at least so it is possible to change it should "flush" really be what is desired.

diff --git a/Engine/Plugins/Experimental/Avalanche/Source/AvalancheMedia/Private/Broadcast/Channel/AvaBroadcastOutputChannel.cpp b/Engine/Plugins/Experimental/Avalanche/Source/AvalancheMedia/Private/Broadcast/Channel/AvaBroadcastOutputChannel.cpp
--- a/Engine/Plugins/Experimental/Avalanche/Source/AvalancheMedia/Private/Broadcast/Channel/AvaBroadcastOutputChannel.cpp
+++ b/Engine/Plugins/Experimental/Avalanche/Source/AvalancheMedia/Private/Broadcast/Channel/AvaBroadcastOutputChannel.cpp
@@ -69,6 +69,19 @@ FAvaBroadcastOutputChannel::FAvaBroadcastOutputChannel(ENoInit NoInit)
 
 namespace UE::AvaMedia::Private
 {
+	EMediaCaptureOverrunAction GetCaptureOverrunAction()
+	{
+		const UAvaMediaSettings& Settings = UAvaMediaSettings::Get();
+		switch (Settings.ChannelOutputOverrunAction)
+		{
+		case EAvaBroadcastOutputOverrunAction::Skip:
+			return EMediaCaptureOverrunAction::Skip;
+		case EAvaBroadcastOutputOverrunAction::Flush:
+		default:
+			return EMediaCaptureOverrunAction::Flush;
+		}
+	}
+	
 	bool IsCapturing(const UMediaCapture* InMediaCapture)
 	{
 		if (IsValid(InMediaCapture))
@@ -724,6 +737,7 @@ bool FAvaBroadcastOutputChannel::StartChannelBroadcast()
 		CaptureOptions.bSkipFrameWhenRunningExpensiveTasks = false;
 		// Allow the formats to be converted if different.
 		CaptureOptions.bConvertToDesiredPixelFormat = true;
+		CaptureOptions.OverrunAction = UE::AvaMedia::Private::GetCaptureOverrunAction();
 		
 		check(RenderTarget);
 		if (MediaCapture->CaptureTextureRenderTarget2D(RenderTarget, CaptureOptions))
diff --git a/Engine/Plugins/Experimental/Avalanche/Source/AvalancheMedia/Public/AvaMediaDefines.h b/Engine/Plugins/Experimental/Avalanche/Source/AvalancheMedia/Public/AvaMediaDefines.h
--- a/Engine/Plugins/Experimental/Avalanche/Source/AvalancheMedia/Public/AvaMediaDefines.h
+++ b/Engine/Plugins/Experimental/Avalanche/Source/AvalancheMedia/Public/AvaMediaDefines.h
@@ -92,6 +92,18 @@ enum class EAvaBroadcastChannelChange : uint8
 };
 ENUM_CLASS_FLAGS(EAvaBroadcastChannelChange);
 
+/**
+ * Action to perform when a media capture overrun (main vs render thread) occurs on a broadcast channel.
+ */
+UENUM()
+enum class EAvaBroadcastOutputOverrunAction : uint8
+{
+	/** Flush rendering thread such that all scheduled commands are executed. */
+	Flush,
+	/** Skip capturing a frame if readback is trailing too much. */
+	Skip
+};
+
 /**
  * The status of a playable object.
  *
diff --git a/Engine/Plugins/Experimental/Avalanche/Source/AvalancheMedia/Public/AvaMediaSettings.h b/Engine/Plugins/Experimental/Avalanche/Source/AvalancheMedia/Public/AvaMediaSettings.h
--- a/Engine/Plugins/Experimental/Avalanche/Source/AvalancheMedia/Public/AvaMediaSettings.h
+++ b/Engine/Plugins/Experimental/Avalanche/Source/AvalancheMedia/Public/AvaMediaSettings.h
@@ -2,6 +2,7 @@
 
 #pragma once
 
+#include "AvaMediaDefines.h"
 #include "Containers/Array.h"
 #include "Containers/EnumAsByte.h"
 #include "Containers/StringFwd.h"
@@ -105,6 +106,12 @@ public:
 	UPROPERTY(config, EditAnywhere, Category = "Broadcast")
 	FIntPoint ChannelDefaultResolution = FIntPoint(1920, 1080);
 
+	/**
+	 * Action to perform when game thread overruns render thread and all frames are in flights being captured / readback. 
+	 */
+	UPROPERTY(config, EditAnywhere, Category = "Broadcast")
+	EAvaBroadcastOutputOverrunAction ChannelOutputOverrunAction = EAvaBroadcastOutputOverrunAction::Skip;
+	
 	/**
 	 * Enables drawing the placeholder widget when there is no Motion Design asset playing.
 	 * If false, the channel is cleared to the background color.
