Fixed calling Shutdown on BeginDestroy

diff --git a/Engine/Plugins/Experimental/Avalanche/Source/Avalanche/Private/Framework/AvaGameInstance.cpp b/Engine/Plugins/Experimental/Avalanche/Source/Avalanche/Private/Framework/AvaGameInstance.cpp
--- a/Engine/Plugins/Experimental/Avalanche/Source/Avalanche/Private/Framework/AvaGameInstance.cpp
+++ b/Engine/Plugins/Experimental/Avalanche/Source/Avalanche/Private/Framework/AvaGameInstance.cpp
@@ -195,7 +195,7 @@ void UAvaGameInstance::RequestUnloadWorld(bool bForceImmediate)
 	{
 		if (bForceImmediate)
 		{
-			UnloadWorld();
+			UnloadWorld(/*bShutdown*/true);
 		}
 		else
 		{
@@ -318,11 +318,11 @@ void UAvaGameInstance::Tick(float DeltaSeconds)
 	}
 	if (bRequestUnloadWorld)
 	{
-		UnloadWorld();
+		UnloadWorld(/*bShutdown*/true);
 	}
 }
 
-void UAvaGameInstance::UnloadWorld()
+void UAvaGameInstance::UnloadWorld(bool bInShutdown)
 {
 	if (bWorldPlaying)
 	{
@@ -347,7 +347,10 @@ void UAvaGameInstance::UnloadWorld()
 		PlayWorld->DestroyWorld(true);
 	}
 
-	Shutdown();
+	if (bInShutdown)
+	{
+		Shutdown();
+	}
 
 	PlayWorld = nullptr;
 	WorldContext = nullptr;
@@ -439,9 +442,9 @@ void UAvaGameInstance::BeginDestroy()
 {
 	FCoreDelegates::OnEnginePreExit.Remove(EnginePreExitHandle);
 	EnginePreExitHandle.Reset();
-	
+
 	EndPlayWorld();
-	UnloadWorld();
-	
+	UnloadWorld(/*bShutdown*/false);
+
 	Super::BeginDestroy();
 }
diff --git a/Engine/Plugins/Experimental/Avalanche/Source/Avalanche/Public/Framework/AvaGameInstance.h b/Engine/Plugins/Experimental/Avalanche/Source/Avalanche/Public/Framework/AvaGameInstance.h
--- a/Engine/Plugins/Experimental/Avalanche/Source/Avalanche/Public/Framework/AvaGameInstance.h
+++ b/Engine/Plugins/Experimental/Avalanche/Source/Avalanche/Public/Framework/AvaGameInstance.h
@@ -103,7 +103,7 @@ public:
 protected:
 	void Tick(float DeltaSeconds);
 
-	void UnloadWorld();
+	void UnloadWorld(bool bInShutdown);
 	void EndPlayWorld();
 	void OnEndFrameTick();
 	void OnEnginePreExit();
