[Motion Design] Rundown Editor: Adding Icon for reuse transition mode in the layer column.

diff --git a/Engine/Plugins/Experimental/Avalanche/Source/AvalancheMediaEditor/Private/AvaMediaEditorStyle.cpp b/Engine/Plugins/Experimental/Avalanche/Source/AvalancheMediaEditor/Private/AvaMediaEditorStyle.cpp
--- a/Engine/Plugins/Experimental/Avalanche/Source/AvalancheMediaEditor/Private/AvaMediaEditorStyle.cpp
+++ b/Engine/Plugins/Experimental/Avalanche/Source/AvalancheMediaEditor/Private/AvaMediaEditorStyle.cpp
@@ -97,6 +97,9 @@ FAvaMediaEditorStyle::FAvaMediaEditorStyle()
 	Set("AvaMediaEditor.MediaOutputLiveWarn"  , new IMAGE_BRUSH_SVG("Icons/MediaIcons/MediaOutputLiveWarn"	, Icon16x16));
 	Set("AvaMediaEditor.MediaOutputError"     , new IMAGE_BRUSH_SVG("Icons/MediaIcons/MediaOutputError"	, Icon16x16));
 
+	// Transition Logic Mode
+	Set("AvaMediaEditor.TransitionModeReuse", new IMAGE_BRUSH_SVG("Icons/MediaIcons/Recycling", Icon16x16));
+	
 	// Motion Design Preview
 	Set("AvaMediaEditor.Checkerboard" , new IMAGE_BRUSH("Images/AvaPreviewCheckerboard", Icon16x16, FLinearColor::White, ESlateBrushTileType::Both));
 
diff --git a/Engine/Plugins/Experimental/Avalanche/Source/AvalancheMediaEditor/Private/Rundown/Pages/Columns/AvaRundownPageTransitionLayerColumn.cpp b/Engine/Plugins/Experimental/Avalanche/Source/AvalancheMediaEditor/Private/Rundown/Pages/Columns/AvaRundownPageTransitionLayerColumn.cpp
--- a/Engine/Plugins/Experimental/Avalanche/Source/AvalancheMediaEditor/Private/Rundown/Pages/Columns/AvaRundownPageTransitionLayerColumn.cpp
+++ b/Engine/Plugins/Experimental/Avalanche/Source/AvalancheMediaEditor/Private/Rundown/Pages/Columns/AvaRundownPageTransitionLayerColumn.cpp
@@ -2,7 +2,10 @@
 
 #include "AvaRundownPageTransitionLayerColumn.h"
 
+#include "AvaMediaEditorStyle.h"
+#include "Rundown/AvaRundown.h"
 #include "Rundown/Pages/PageViews/IAvaRundownPageView.h"
+#include "Widgets/Images/SImage.h"
 #include "Widgets/Layout/SBox.h"
 #include "Widgets/Text/STextBlock.h"
 
@@ -29,14 +32,59 @@ SHeaderRow::FColumn::FArguments FAvaRundownPageTransitionLayerColumn::ConstructH
 	;
 }
 
+namespace UE::AvaMediaEditor::RundownPageTransitionLayerColumn::Private
+{
+	static EVisibility GetTransitionModeVisibility(FAvaRundownPageViewWeak InPageViewWeak)
+	{
+		EVisibility Visibility = EVisibility::Hidden;	//... to preserve alignment.
+		if (const FAvaRundownPageViewPtr PageView = InPageViewWeak.Pin())
+		{
+			const UAvaRundown* Rundown = PageView->GetRundown();
+			if (IsValid(Rundown))
+			{
+				const FAvaRundownPage& Page = Rundown->GetPage(PageView->GetPageId());
+				if (Page.IsValidPage())
+				{
+					// Current logic: show the reuse icon if any of the sub-templates has the reuse mode.
+					const int32 NumTemplates = Page.GetNumTemplates(Rundown);
+					for (int32 TemplateIndex = 0; TemplateIndex < NumTemplates; ++TemplateIndex)
+					{
+						if (Page.GetTransitionMode(Rundown, TemplateIndex) == EAvaTransitionInstancingMode::Reuse)
+						{
+							Visibility = EVisibility::Visible;
+							break;
+						}
+					}
+				}
+			}
+		}
+		return Visibility;
+	}
+}
+
 TSharedRef<SWidget> FAvaRundownPageTransitionLayerColumn::ConstructRowWidget(const FAvaRundownPageViewRef& InPageView, const TSharedPtr<SAvaRundownPageViewRow>& InRow)
 {
-	return SNew(SBox)
+	using namespace UE::AvaMediaEditor::RundownPageTransitionLayerColumn::Private;
+	const FAvaRundownPageViewWeak PageViewWeak = InPageView;
+	
+	TSharedRef<SHorizontalBox> RowWidget = SNew(SHorizontalBox)
+		+ SHorizontalBox::Slot()
+		.AutoWidth()
+		.Padding(UE::AvaRundown::FEditorMetrics::ColumnLeftOffset, 0.f, 0.f, 0.f)
+		[
+			SNew(SImage)
+			.Image(FAvaMediaEditorStyle::Get().GetBrush("AvaMediaEditor.TransitionModeReuse"))
+			.ToolTipText(LOCTEXT("TransitionModeReuseTooltip", "Transition Mode: Reuse - Consecutive pages on the same layer will reuse the same instance."))
+			.Visibility_Static(&GetTransitionModeVisibility, PageViewWeak)
+			.ColorAndOpacity(FSlateColor::UseForeground())
+		]
+		+ SHorizontalBox::Slot()
 		.Padding(UE::AvaRundown::FEditorMetrics::ColumnLeftOffset, 0.f, 0.f, 0.f)
 		[
 			SNew(STextBlock)
 			.Text(InPageView, &IAvaRundownPageView::GetPageTransitionLayerNameText)
 		];
+	return RowWidget;
 }
 
 #undef LOCTEXT_NAMESPACE
