Motion Design: Support for Wait for Layer to support Re-use Levels. Also updated Sub Layer Change Task to support Re-use (and Update RC Values)

#virtualized

diff --git a/Engine/Plugins/Experimental/Avalanche/Source/AvalancheTransition/Private/Tasks/AvaTransitionWaitForLayerTask.cpp b/Engine/Plugins/Experimental/Avalanche/Source/AvalancheTransition/Private/Tasks/AvaTransitionWaitForLayerTask.cpp
--- a/Engine/Plugins/Experimental/Avalanche/Source/AvalancheTransition/Private/Tasks/AvaTransitionWaitForLayerTask.cpp
+++ b/Engine/Plugins/Experimental/Avalanche/Source/AvalancheTransition/Private/Tasks/AvaTransitionWaitForLayerTask.cpp
@@ -2,6 +2,7 @@
 
 #include "Tasks/AvaTransitionWaitForLayerTask.h"
 #include "AvaTransitionLayerUtils.h"
+#include "AvaTransitionTree.h"
 #include "Behavior/AvaTransitionBehaviorInstance.h"
 #include "Engine/Level.h"
 #include "Rendering/AvaTransitionRenderingSubsystem.h"
@@ -55,11 +56,9 @@ EStateTreeRunStatus FAvaTransitionWaitForLayerTask::WaitForLayer(FStateTreeExecu
 
 	if (bIsLayerRunning)
 	{
-		// Hide Level Primitives while Waiting for other Layers
-		if (InstanceData.bHideSceneWhileWaiting && InstanceData.HiddenLevel == nullptr)
+		if (ShouldHideLevel(InContext, InstanceData))
 		{
-			const FAvaTransitionContext& TransitionContext = InContext.GetExternalData(TransitionContextHandle);
-			if (const FAvaTransitionScene* TransitionScene = TransitionContext.GetTransitionScene())
+			if (const FAvaTransitionScene* TransitionScene = InContext.GetExternalData(TransitionContextHandle).GetTransitionScene())
 			{
 				InstanceData.HiddenLevel = TransitionScene->GetLevel();
 				RenderingSubsystem.HideLevel(InstanceData.HiddenLevel);
@@ -74,4 +73,27 @@ EStateTreeRunStatus FAvaTransitionWaitForLayerTask::WaitForLayer(FStateTreeExecu
 	return EStateTreeRunStatus::Succeeded;
 }
 
+bool FAvaTransitionWaitForLayerTask::ShouldHideLevel(const FStateTreeExecutionContext& InContext, const FAvaTransitionWaitForLayerTask::FInstanceDataType& InInstanceData) const
+{
+	if (!InInstanceData.bHideSceneWhileWaiting)
+	{
+		return false;	
+	}
+
+	// If Hidden Level is non-null, it means the level has already been hidden / processed. Skip
+	if (InInstanceData.HiddenLevel != nullptr)
+	{
+		return false;
+	}
+
+	// If Instancing Mode is set to Reuse, and user set to Hide Level regardless, allow Hiding Reused Level
+	const UAvaTransitionTree* TransitionTree = Cast<UAvaTransitionTree>(InContext.GetStateTree());
+	if (TransitionTree && TransitionTree->GetInstancingMode() == EAvaTransitionInstancingMode::Reuse)
+	{
+		return InInstanceData.bHideOnLevelReuse;
+	}
+
+	return true;
+}
+
 #undef LOCTEXT_NAMESPACE
diff --git a/Engine/Plugins/Experimental/Avalanche/Source/AvalancheTransition/Public/Tasks/AvaTransitionWaitForLayerTask.h b/Engine/Plugins/Experimental/Avalanche/Source/AvalancheTransition/Public/Tasks/AvaTransitionWaitForLayerTask.h
--- a/Engine/Plugins/Experimental/Avalanche/Source/AvalancheTransition/Public/Tasks/AvaTransitionWaitForLayerTask.h
+++ b/Engine/Plugins/Experimental/Avalanche/Source/AvalancheTransition/Public/Tasks/AvaTransitionWaitForLayerTask.h
@@ -17,6 +17,10 @@ struct FAvaTransitionWaitForLayerTaskInstanceData : public FAvaTransitionLayerTa
 	UPROPERTY(EditAnywhere, Category="Transition Logic")
 	bool bHideSceneWhileWaiting = true;
 
+	/** Whether to hide level even if set to Re-use. Requires Hide Scene While Waiting to be enabled */
+	UPROPERTY(EditAnywhere, Category="Transition Logic", meta=(EditCondition="bHideSceneWhileWaiting", EditConditionHides))
+	bool bHideOnLevelReuse = false;
+
 	TObjectKey<ULevel> HiddenLevel;
 };
 
@@ -42,5 +46,7 @@ struct AVALANCHETRANSITION_API FAvaTransitionWaitForLayerTask : public FAvaTrans
 
 	EStateTreeRunStatus WaitForLayer(FStateTreeExecutionContext& InContext) const;
 
+	bool ShouldHideLevel(const FStateTreeExecutionContext& InContext, const FAvaTransitionWaitForLayerTask::FInstanceDataType& InInstanceData) const;
+
 	TStateTreeExternalDataHandle<UAvaTransitionRenderingSubsystem> RenderingSubsystemHandle;
 };
