[Motion Design] Follow up Rundown Editor (and Server) Fixes

- (issue reported in ticket) Prevent deleting template if it has playing pages.
- Reset the InstanceIds of a newly pasted template. (found while testing)
- Rundown Server RemovePage and RemoveTemplate commands will return an error message if the operation can't be done.

diff --git a/Engine/Plugins/Experimental/Avalanche/Source/AvalancheMedia/Private/Rundown/AvaRundown.cpp b/Engine/Plugins/Experimental/Avalanche/Source/AvalancheMedia/Private/Rundown/AvaRundown.cpp
--- a/Engine/Plugins/Experimental/Avalanche/Source/AvalancheMedia/Private/Rundown/AvaRundown.cpp
+++ b/Engine/Plugins/Experimental/Avalanche/Source/AvalancheMedia/Private/Rundown/AvaRundown.cpp
@@ -546,6 +546,19 @@ bool UAvaRundown::CanRemovePages(const TArray<int32>& InPageIds) const
 		{
 			return false;
 		}
+
+		// Prevent deletion of templates that have playing page instances.
+		const FAvaRundownPage& Page = GetPage(PageId);
+		if (Page.IsValidPage() && Page.IsTemplate())
+		{
+			for (const int32 InstancePageId : Page.GetInstancedIds())
+			{
+				if (IsPagePlayingOrPreviewing(InstancePageId))
+				{
+					return false;
+				}
+			}
+		}
 	}
 	return InPageIds.Num() > 0;
 }
diff --git a/Engine/Plugins/Experimental/Avalanche/Source/AvalancheMedia/Private/Rundown/AvaRundownServer.cpp b/Engine/Plugins/Experimental/Avalanche/Source/AvalancheMedia/Private/Rundown/AvaRundownServer.cpp
--- a/Engine/Plugins/Experimental/Avalanche/Source/AvalancheMedia/Private/Rundown/AvaRundownServer.cpp
+++ b/Engine/Plugins/Experimental/Avalanche/Source/AvalancheMedia/Private/Rundown/AvaRundownServer.cpp
@@ -1131,8 +1131,14 @@ void FAvaRundownServer::HandleDeletePage(const FAvaRundownDeletePage& InMessage,
 		return;
 	}
 
-	Rundown->RemovePage(InMessage.PageId);
-	SendMessage(InContext->GetSender(), InMessage.RequestId, ELogVerbosity::Log, TEXT("Page %d deleted"), InMessage.PageId);
+	if (Rundown->RemovePage(InMessage.PageId))
+	{
+		SendMessage(InContext->GetSender(), InMessage.RequestId, ELogVerbosity::Log, TEXT("Page %d deleted"), InMessage.PageId);
+	}
+	else
+	{
+		SendMessage(InContext->GetSender(), InMessage.RequestId, ELogVerbosity::Error, TEXT("page %d can't be deleted"), InMessage.PageId);
+	}
 }
 
 void FAvaRundownServer::HandleDeleteTemplate(const FAvaRundownDeleteTemplate& InMessage,
@@ -1163,8 +1169,14 @@ void FAvaRundownServer::HandleDeleteTemplate(const FAvaRundownDeleteTemplate& In
 		return;
 	}
 
-	Rundown->RemovePage(InMessage.PageId);
-	SendMessage(InContext->GetSender(), InMessage.RequestId, ELogVerbosity::Log, TEXT("Page template %d deleted"), InMessage.PageId);
+	if (Rundown->RemovePage(InMessage.PageId))
+	{
+		SendMessage(InContext->GetSender(), InMessage.RequestId, ELogVerbosity::Log, TEXT("Page template %d deleted"), InMessage.PageId);
+	}
+	else
+	{
+		SendMessage(InContext->GetSender(), InMessage.RequestId, ELogVerbosity::Error, TEXT("Page template %d can't be deleted"), InMessage.PageId);
+	}
 }
 
 void FAvaRundownServer::HandleCreateTemplate(const FAvaRundownCreateTemplate& InMessage,
diff --git a/Engine/Plugins/Experimental/Avalanche/Source/AvalancheMedia/Public/Rundown/AvaRundownPage.h b/Engine/Plugins/Experimental/Avalanche/Source/AvalancheMedia/Public/Rundown/AvaRundownPage.h
--- a/Engine/Plugins/Experimental/Avalanche/Source/AvalancheMedia/Public/Rundown/AvaRundownPage.h
+++ b/Engine/Plugins/Experimental/Avalanche/Source/AvalancheMedia/Public/Rundown/AvaRundownPage.h
@@ -95,6 +95,8 @@ public:
 	
 	const TSet<int32>& GetInstancedIds() const { return Instances; }
 
+	void ResetInstancedIds() { Instances.Reset(); }	// For internal use only.
+
 	const FString& GetPageName() const { return PageName; }
 	void SetPageName(const FString& InPageName) { PageName = InPageName; }
 
diff --git a/Engine/Plugins/Experimental/Avalanche/Source/AvalancheMediaEditor/Private/Rundown/AvaRundownEditorUtils.cpp b/Engine/Plugins/Experimental/Avalanche/Source/AvalancheMediaEditor/Private/Rundown/AvaRundownEditorUtils.cpp
--- a/Engine/Plugins/Experimental/Avalanche/Source/AvalancheMediaEditor/Private/Rundown/AvaRundownEditorUtils.cpp
+++ b/Engine/Plugins/Experimental/Avalanche/Source/AvalancheMediaEditor/Private/Rundown/AvaRundownEditorUtils.cpp
@@ -124,7 +124,7 @@ namespace UE::AvaRundownEditor::Utils::Private
 			DestinationPage.SetTemplateId(InTemplateId);
 			return true;
 		}
-		UE_LOG(LogAvaRundown, Error, TEXT("Failed to copy page in plage: page id %d is not found in destination rundown."), InPageId);
+		UE_LOG(LogAvaRundown, Error, TEXT("Failed to copy page in place: page id %d is not found in destination rundown."), InPageId);
 		return false;
 	}
 }
@@ -394,6 +394,7 @@ TArray<int32> UE::AvaRundownEditor::Utils::ImportTemplatePages(UAvaRundown* InRu
 		{
 			OutImportedTemplateIds.Add(SourceTemplateId, ImportedTemplateId);
 			ensure(Private::CopyPageInPlace(InRundown, ImportedTemplateId, SourceTemplate, FAvaRundownPage::InvalidPageId));
+			InRundown->GetPage(ImportedTemplateId).ResetInstancedIds();	// We know this template has no instances yet.
 			OutTemplateIds.Add(ImportedTemplateId);
 		}
 	}
@@ -465,6 +466,7 @@ TArray<int32> UE::AvaRundownEditor::Utils::ImportInstancedPages(
 					const FAvaRundownPage& SourceTemplate = Private::FindPage(InSourceTemplates, SourcePage.GetTemplateId());
 					
 					ensure(Private::CopyPageInPlace(InRundown, ImportedTemplateId, SourceTemplate.IsValidPage() ? SourceTemplate : SourcePage, FAvaRundownPage::InvalidPageId));
+					InRundown->GetPage(ImportedTemplateId).ResetInstancedIds();	// We know this template has no instances yet.
 				}
 				else
 				{
