Avalanche (Matrix crash fix)
- Pre-inversed a rotation matrix that would produce NaNs using InverseTransformVector.

diff --git a/Engine/Plugins/Experimental/Avalanche/Source/AvalancheComponentVisualizers/Private/AvaVisBase.cpp b/Engine/Plugins/Experimental/Avalanche/Source/AvalancheComponentVisualizers/Private/AvaVisBase.cpp
--- a/Engine/Plugins/Experimental/Avalanche/Source/AvalancheComponentVisualizers/Private/AvaVisBase.cpp
+++ b/Engine/Plugins/Experimental/Avalanche/Source/AvalancheComponentVisualizers/Private/AvaVisBase.cpp
@@ -470,13 +470,18 @@ bool FAvaVisualizerBase::IsMouseOverComponent(const UActorComponent* Component,
 		return false;
 	}
 
-	FOrientedBox ActorOrientedBox = FAvaActorUtils::MakeOrientedBox(ActorSizeBox, GetComponentTransform(Component));
+	const FOrientedBox ActorOrientedBox = FAvaActorUtils::MakeOrientedBox(ActorSizeBox, GetComponentTransform(Component));
 
 	// Calculate the rotation of the oriented box and revert it to axis aligned, moving the camera position and look direction along with it
 	// so that we can use the line/box intersection method.
-	const FMatrix RotMatrix(ActorOrientedBox.AxisX, ActorOrientedBox.AxisY, ActorOrientedBox.AxisZ, FVector::ZeroVector);
-
-	FVector2D MouseLocation = {
+	const FMatrix InvRotMatrix = FMatrix(
+		ActorOrientedBox.AxisX, 
+		ActorOrientedBox.AxisY, 
+		ActorOrientedBox.AxisZ, 
+		FVector::ZeroVector
+	).Inverse();
+
+	const FVector2D MouseLocation = {
 		static_cast<double>(InView->CursorPos.X),
 		static_cast<double>(InView->CursorPos.Y)
 	};
@@ -485,11 +490,10 @@ bool FAvaVisualizerBase::IsMouseOverComponent(const UActorComponent* Component,
 
 	InView->DeprojectFVector2D(MouseLocation, WorldOrigin, WorldDirection);
 
-	FVector CameraRotationVector = (WorldOrigin - InView->ViewLocation).GetUnsafeNormal();
-	FVector MouseWorldPosition   = RotMatrix.InverseTransformVector(WorldOrigin - ActorOrientedBox.Center);
-	CameraRotationVector         = RotMatrix.InverseTransformVector(CameraRotationVector);
-	const FVector Extent         = {ActorOrientedBox.ExtentX, ActorOrientedBox.ExtentY, ActorOrientedBox.ExtentZ};
-	const FBox OriginBounds      = FBox(-Extent, Extent);
+	const FVector MouseWorldPosition   = InvRotMatrix.TransformVector(WorldOrigin - ActorOrientedBox.Center);
+	const FVector CameraRotationVector = InvRotMatrix.TransformVector((WorldOrigin - InView->ViewLocation).GetUnsafeNormal());
+	const FVector Extent               = {ActorOrientedBox.ExtentX, ActorOrientedBox.ExtentY, ActorOrientedBox.ExtentZ};
+	const FBox    OriginBounds         = FBox(-Extent, Extent);
 
 	const FVector RayEnd = MouseWorldPosition + CameraRotationVector * 10000.f;
 
