Motion Design: Added Camera Conditional Update View Target

#virtualized

diff --git a/Engine/Plugins/Experimental/Avalanche/Source/Avalanche/Avalanche.Build.cs b/Engine/Plugins/Experimental/Avalanche/Source/Avalanche/Avalanche.Build.cs
--- a/Engine/Plugins/Experimental/Avalanche/Source/Avalanche/Avalanche.Build.cs
+++ b/Engine/Plugins/Experimental/Avalanche/Source/Avalanche/Avalanche.Build.cs
@@ -13,6 +13,7 @@ public class Avalanche : ModuleRules
 			{
 				"AvalancheAttribute",
 				"AvalancheCore",
+				"AvalancheRemoteControl",
 				"AvalancheSceneTree",
 				"AvalancheSequence",
 				"AvalancheTag",
@@ -25,7 +26,7 @@ public class Avalanche : ModuleRules
 				"GeometryScriptingCore",
 				"Niagara",
 				"RawMesh",
-				"Text3D", 
+				"Text3D",
 			}
 		);
 
diff --git a/Engine/Plugins/Experimental/Avalanche/Source/Avalanche/Private/AvaScene.cpp b/Engine/Plugins/Experimental/Avalanche/Source/Avalanche/Private/AvaScene.cpp
--- a/Engine/Plugins/Experimental/Avalanche/Source/Avalanche/Private/AvaScene.cpp
+++ b/Engine/Plugins/Experimental/Avalanche/Source/Avalanche/Private/AvaScene.cpp
@@ -262,6 +262,14 @@ void AAvaScene::RebuildSequenceTree()
 	IAvaSequenceProvider::RebuildSequenceTree();
 }
 
+void AAvaScene::OnValuesApplied_Implementation()
+{
+	if (UAvaCameraSubsystem* CameraSubsystem = UAvaCameraSubsystem::Get(this))
+	{
+		CameraSubsystem->ConditionallyUpdateViewTarget(GetLevel());
+	}
+}
+
 void AAvaScene::PostActorCreated()
 {
 	Super::PostActorCreated();
diff --git a/Engine/Plugins/Experimental/Avalanche/Source/Avalanche/Public/AvaScene.h b/Engine/Plugins/Experimental/Avalanche/Source/Avalanche/Public/AvaScene.h
--- a/Engine/Plugins/Experimental/Avalanche/Source/Avalanche/Public/AvaScene.h
+++ b/Engine/Plugins/Experimental/Avalanche/Source/Avalanche/Public/AvaScene.h
@@ -5,6 +5,7 @@
 #include "AvaSceneTree.h"
 #include "Delegates/Delegate.h"
 #include "GameFramework/Actor.h"
+#include "IAvaRemoteControlInterface.h"
 #include "IAvaSceneInterface.h"
 #include "IAvaSequenceProvider.h"
 #include "UObject/ObjectPtr.h"
@@ -23,7 +24,7 @@ class ISequencer;
 #endif
 
 UCLASS(MinimalAPI, NotPlaceable, Hidden, NotBlueprintable, NotBlueprintType, DisplayName = "Motion Design Scene")
-class AAvaScene : public AActor, public IAvaSequenceProvider, public IAvaSceneInterface, public IAvaViewportDataProvider
+class AAvaScene : public AActor, public IAvaSequenceProvider, public IAvaSceneInterface, public IAvaViewportDataProvider, public IAvaRemoteControlInterface
 {
 	GENERATED_BODY()
 
@@ -78,6 +79,10 @@ public:
 	virtual void RebuildSequenceTree() override;
 	//~ End IAvaSequenceProvider
 
+	//~ Begin IAvaRemoteControlInterface
+	virtual void OnValuesApplied_Implementation() override;
+	//~ End IAvaRemoteControlInterface
+
 	//~ Begin AActor
 	virtual void PostActorCreated() override;
 	virtual void BeginPlay() override;
diff --git a/Engine/Plugins/Experimental/Avalanche/Source/AvalancheCamera/Private/AvaCameraSubsystem.cpp b/Engine/Plugins/Experimental/Avalanche/Source/AvalancheCamera/Private/AvaCameraSubsystem.cpp
--- a/Engine/Plugins/Experimental/Avalanche/Source/AvalancheCamera/Private/AvaCameraSubsystem.cpp
+++ b/Engine/Plugins/Experimental/Avalanche/Source/AvalancheCamera/Private/AvaCameraSubsystem.cpp
@@ -97,24 +97,7 @@ void UAvaCameraSubsystem::RegisterScene(const ULevel* InSceneLevel)
 		}
 	}
 
-	// Remove invalid entries
-	ViewTargets.RemoveAll([](const FAvaViewTarget& InViewTarget)
-		{
-			return !InViewTarget.IsValid();
-		});
-
-	// Sort so that the higher priorities are at the end of the list (end = current)
-	// Stable sorting so that the more recently registered scenes are preferred when the priorities match
-	ViewTargets.StableSort(
-		[](const FAvaViewTarget& A, const FAvaViewTarget& B)
-		{
-			return A.GetPriority() < B.GetPriority();
-		});
-
-	if (!HasCustomViewTargetting(InSceneLevel))
-	{
-		UpdatePlayerControllerViewTarget();
-	}
+	ConditionallyUpdateViewTarget(InSceneLevel);
 }
 
 void UAvaCameraSubsystem::UnregisterScene(const ULevel* InSceneLevel)
@@ -131,10 +114,7 @@ void UAvaCameraSubsystem::UnregisterScene(const ULevel* InSceneLevel)
 			return !InViewTarget.IsValid() || InViewTarget.Actor->GetLevel() == InSceneLevel;
 		});
 
-	if (!HasCustomViewTargetting(InSceneLevel))
-    {
-    	UpdatePlayerControllerViewTarget();
-    }
+	ConditionallyUpdateViewTarget(InSceneLevel);
 }
 
 bool UAvaCameraSubsystem::IsBlendingToViewTarget(const ULevel* InSceneLevel) const
@@ -149,13 +129,27 @@ bool UAvaCameraSubsystem::IsBlendingToViewTarget(const ULevel* InSceneLevel) con
 	return false;
 }
 
-void UAvaCameraSubsystem::UpdatePlayerControllerViewTarget(const FViewTargetTransitionParams* InOverrideTransitionParams) const
+void UAvaCameraSubsystem::UpdatePlayerControllerViewTarget(const FViewTargetTransitionParams* InOverrideTransitionParams)
 {
 	if (!IsValid(PlayerController) || ViewTargets.IsEmpty())
 	{
 		return;
 	}
 
+	// Remove invalid entries
+	ViewTargets.RemoveAll([](const FAvaViewTarget& InViewTarget)
+		{
+			return !InViewTarget.IsValid();
+		});
+
+	// Sort so that the higher priorities are at the end of the list (end = current)
+	// Stable sorting so that the more recently registered scenes are preferred when the priorities match
+	ViewTargets.StableSort(
+		[](const FAvaViewTarget& A, const FAvaViewTarget& B)
+		{
+			return A.GetPriority() < B.GetPriority();
+		});
+
 	const FAvaViewTarget& DesiredViewTarget = ViewTargets.Last();
 
 	// These should've been removed prior to calling this private function
@@ -174,6 +168,16 @@ void UAvaCameraSubsystem::UpdatePlayerControllerViewTarget(const FViewTargetTran
 	PlayerController->SetViewTarget(DesiredViewTarget.Actor, *ViewTargetTransitionParams);
 }
 
+bool UAvaCameraSubsystem::ConditionallyUpdateViewTarget(const ULevel* InSceneLevel)
+{
+	if (!HasCustomViewTargetting(InSceneLevel))
+	{
+		UpdatePlayerControllerViewTarget();
+		return true;
+	}
+	return false;
+}
+
 bool UAvaCameraSubsystem::DoesSupportWorldType(const EWorldType::Type InWorldType) const
 {
 	return InWorldType == EWorldType::GamePreview
diff --git a/Engine/Plugins/Experimental/Avalanche/Source/AvalancheCamera/Private/Transition/AvaTransitionCameraLibrary.cpp b/Engine/Plugins/Experimental/Avalanche/Source/AvalancheCamera/Private/Transition/AvaTransitionCameraLibrary.cpp
new file mode 100644
--- /dev/null
+++ b/Engine/Plugins/Experimental/Avalanche/Source/AvalancheCamera/Private/Transition/AvaTransitionCameraLibrary.cpp
@@ -0,0 +1,39 @@
+// Copyright Epic Games, Inc. All Rights Reserved.
+
+#include "AvaTransitionCameraLibrary.h"
+#include "AvaCameraSubsystem.h"
+#include "AvaTransitionContext.h"
+#include "Behavior/AvaTransitionBehaviorInstanceCache.h"
+#include "Engine/Level.h"
+#include "IAvaTransitionNodeInterface.h"
+
+bool UAvaTransitionCameraLibrary::ConditionallyUpdateViewTarget(UObject* InTransitionNode)
+{
+	IAvaTransitionNodeInterface* NodeInterface = Cast<IAvaTransitionNodeInterface>(InTransitionNode);
+	if (!NodeInterface)
+	{
+		return false;
+	}
+
+	const FAvaTransitionContext* TransitionContext = NodeInterface->GetBehaviorInstanceCache().GetTransitionContext();
+	if (!TransitionContext)
+	{
+		return false;
+	}
+
+	const FAvaTransitionScene* TransitionScene = TransitionContext->GetTransitionScene();
+	if (!TransitionScene)
+	{
+		return false;
+	}
+
+	const ULevel* SceneLevel = TransitionScene->GetLevel();
+
+	UAvaCameraSubsystem* CameraSubsystem = UAvaCameraSubsystem::Get(SceneLevel);
+	if (!CameraSubsystem)
+	{
+		return false;
+	}
+
+	return CameraSubsystem->ConditionallyUpdateViewTarget(SceneLevel);
+}
diff --git a/Engine/Plugins/Experimental/Avalanche/Source/AvalancheCamera/Private/Transition/AvaTransitionCameraLibrary.h b/Engine/Plugins/Experimental/Avalanche/Source/AvalancheCamera/Private/Transition/AvaTransitionCameraLibrary.h
new file mode 100644
--- /dev/null
+++ b/Engine/Plugins/Experimental/Avalanche/Source/AvalancheCamera/Private/Transition/AvaTransitionCameraLibrary.h
@@ -0,0 +1,16 @@
+// Copyright Epic Games, Inc. All Rights Reserved.
+
+#pragma once
+
+#include "Kismet/BlueprintFunctionLibrary.h"
+#include "AvaTransitionCameraLibrary.generated.h"
+
+UCLASS()
+class UAvaTransitionCameraLibrary : public UBlueprintFunctionLibrary
+{
+	GENERATED_BODY()
+
+public:
+	UFUNCTION(BlueprintCallable, Category="Transition Logic|Camera", meta=(DefaultToSelf="InTransitionNode"))
+	static bool ConditionallyUpdateViewTarget(UObject* InTransitionNode);
+};
diff --git a/Engine/Plugins/Experimental/Avalanche/Source/AvalancheCamera/Private/Transition/Tasks/AvaCameraBlendTask.cpp b/Engine/Plugins/Experimental/Avalanche/Source/AvalancheCamera/Private/Transition/Tasks/AvaCameraBlendTask.cpp
--- a/Engine/Plugins/Experimental/Avalanche/Source/AvalancheCamera/Private/Transition/Tasks/AvaCameraBlendTask.cpp
+++ b/Engine/Plugins/Experimental/Avalanche/Source/AvalancheCamera/Private/Transition/Tasks/AvaCameraBlendTask.cpp
@@ -28,7 +28,7 @@ EStateTreeRunStatus FAvaCameraBlendTask::EnterState(FStateTreeExecutionContext&
 
 	const FInstanceDataType& InstanceData = InContext.GetInstanceData(*this);
 
-	const UAvaCameraSubsystem& CameraSubsystem = InContext.GetExternalData(CameraSubsystemHandle);
+	UAvaCameraSubsystem& CameraSubsystem = InContext.GetExternalData(CameraSubsystemHandle);
 
 	const FViewTargetTransitionParams* TransitionParams = nullptr;
 	if (InstanceData.bOverrideTransitionParams)
diff --git a/Engine/Plugins/Experimental/Avalanche/Source/AvalancheCamera/Public/AvaCameraSubsystem.h b/Engine/Plugins/Experimental/Avalanche/Source/AvalancheCamera/Public/AvaCameraSubsystem.h
--- a/Engine/Plugins/Experimental/Avalanche/Source/AvalancheCamera/Public/AvaCameraSubsystem.h
+++ b/Engine/Plugins/Experimental/Avalanche/Source/AvalancheCamera/Public/AvaCameraSubsystem.h
@@ -35,7 +35,7 @@ struct FAvaViewTarget
 	TObjectPtr<const UAvaCameraPriorityModifier> CameraPriorityModifier;
 };
 
-UCLASS(MinimalAPI)
+UCLASS(MinimalAPI, DisplayName="Motion Design Camera Subsystem")
 class UAvaCameraSubsystem : public UWorldSubsystem
 {
 	GENERATED_BODY()
@@ -48,7 +48,9 @@ public:
 
 	bool IsBlendingToViewTarget(const ULevel* InSceneLevel) const;
 
-	void UpdatePlayerControllerViewTarget(const FViewTargetTransitionParams* InOverrideTransitionParams = nullptr) const;
+	void UpdatePlayerControllerViewTarget(const FViewTargetTransitionParams* InOverrideTransitionParams = nullptr);
+
+	AVALANCHECAMERA_API bool ConditionallyUpdateViewTarget(const ULevel* InSceneLevel);
 
 protected:
 	//~ Begin UWorldSubsystem
diff --git a/Engine/Plugins/Experimental/Avalanche/Source/AvalancheRemoteControl/Public/IAvaRemoteControlInterface.h b/Engine/Plugins/Experimental/Avalanche/Source/AvalancheRemoteControl/Public/IAvaRemoteControlInterface.h
--- a/Engine/Plugins/Experimental/Avalanche/Source/AvalancheRemoteControl/Public/IAvaRemoteControlInterface.h
+++ b/Engine/Plugins/Experimental/Avalanche/Source/AvalancheRemoteControl/Public/IAvaRemoteControlInterface.h
@@ -14,7 +14,7 @@ class UAvaRemoteControlInterface : public UInterface
 	GENERATED_BODY()
 };
 
-class AVALANCHEREMOTECONTROL_API IAvaRemoteControlInterface
+class IAvaRemoteControlInterface
 {
 	GENERATED_BODY()
 
@@ -23,6 +23,7 @@ public:
 	 * Called when the remote control values are applied.
 	 */
 	UFUNCTION(BlueprintNativeEvent, Category = "Remote Control")
-	void OnValuesApplied();
+	AVALANCHEREMOTECONTROL_API void OnValuesApplied();
+
 	virtual void OnValuesApplied_Implementation() {}
-};
\ No newline at end of file
+};
diff --git a/Engine/Plugins/Experimental/Avalanche/Source/AvalancheTransition/Private/Behavior/AvaTransitionBehaviorInstance.cpp b/Engine/Plugins/Experimental/Avalanche/Source/AvalancheTransition/Private/Behavior/AvaTransitionBehaviorInstance.cpp
--- a/Engine/Plugins/Experimental/Avalanche/Source/AvalancheTransition/Private/Behavior/AvaTransitionBehaviorInstance.cpp
+++ b/Engine/Plugins/Experimental/Avalanche/Source/AvalancheTransition/Private/Behavior/AvaTransitionBehaviorInstance.cpp
@@ -39,6 +39,15 @@ IAvaTransitionBehavior* FAvaTransitionBehaviorInstance::GetBehavior() const
 	return BehaviorWeak.Get();
 }
 
+const UAvaTransitionTree* FAvaTransitionBehaviorInstance::GetTransitionTree() const
+{
+	if (const IAvaTransitionBehavior* Behavior = GetBehavior())
+	{
+		return Behavior->GetTransitionTree();
+	}
+	return nullptr;
+}
+
 FAvaTagHandle FAvaTransitionBehaviorInstance::GetTransitionLayer() const
 {
 	return TransitionContext.GetTransitionLayer();
diff --git a/Engine/Plugins/Experimental/Avalanche/Source/AvalancheTransition/Private/Execution/AvaTransitionExecutor.cpp b/Engine/Plugins/Experimental/Avalanche/Source/AvalancheTransition/Private/Execution/AvaTransitionExecutor.cpp
--- a/Engine/Plugins/Experimental/Avalanche/Source/AvalancheTransition/Private/Execution/AvaTransitionExecutor.cpp
+++ b/Engine/Plugins/Experimental/Avalanche/Source/AvalancheTransition/Private/Execution/AvaTransitionExecutor.cpp
@@ -6,6 +6,7 @@
 #include "AvaTagList.h"
 #include "AvaTransitionLayer.h"
 #include "AvaTransitionSubsystem.h"
+#include "AvaTransitionTree.h"
 #include "Behavior/IAvaTransitionBehavior.h"
 #include "Engine/World.h"
 #include "Execution/AvaTransitionExecutorBuilder.h"
@@ -50,14 +51,19 @@ void FAvaTransitionExecutor::Setup()
 
 	struct FLayerInfo
 	{
-		// All the Behavior Instances found for a given Layer
-		TArray<FAvaTransitionBehaviorInstance*> BehaviorInstances;
+		// Enter Instances found for a given Layer
+		TArray<FAvaTransitionBehaviorInstance*> EnterInstances;
+
+		// Exit Instances found for a given Layer
+		TArray<FAvaTransitionBehaviorInstance*> ExitInstances;
 
 		// The accumulated Transition Type for a given Layer (e.g. combinations could be In, Out or In | Out)
 		EAvaTransitionType TransitionType;
+
+		FAvaTagHandle TransitionLayer;
 	};
 
-	// Map of BehaviorInstance to the Resolved Tags
+	// Map of the Resolved Tag Layer to the Behavior Instances / Transition type of that Layer
 	TMap<FAvaTag, FLayerInfo> TagLayerInfo;
 	{
 		// Rough Estimate: assume each Instance is in its own single unique tag
@@ -69,8 +75,19 @@ void FAvaTransitionExecutor::Setup()
 			for (const FAvaTag* Tag : Instance.GetTransitionLayer().GetTags())
 			{
 				FLayerInfo& LayerInfo = TagLayerInfo.FindOrAdd(*Tag);
-				LayerInfo.BehaviorInstances.AddUnique(&Instance);
 				LayerInfo.TransitionType |= Instance.GetTransitionType();
+				LayerInfo.TransitionLayer = Instance.GetTransitionLayer();
+
+				switch (Instance.GetTransitionType())
+				{
+				case EAvaTransitionType::In:
+					LayerInfo.EnterInstances.AddUnique(&Instance);
+					break;
+
+				case EAvaTransitionType::Out:
+					LayerInfo.ExitInstances.AddUnique(&Instance);
+					break;
+				}
 			}
 		}
 	}
@@ -81,10 +98,9 @@ void FAvaTransitionExecutor::Setup()
 		const FLayerInfo& LayerInfo = Pair.Value;
 		if (LayerInfo.TransitionType == EAvaTransitionType::In && !EnumHasAnyFlags(LayerInfo.TransitionType, EAvaTransitionType::Out))
 		{
-			check(!LayerInfo.BehaviorInstances.IsEmpty());
 			FAvaTransitionBehaviorInstance& NullInstanceCopy = Instances.Add_GetRef(NullInstance);
 			NullInstanceCopy.SetTransitionType(EAvaTransitionType::Out);
-			NullInstanceCopy.SetOverrideLayer(LayerInfo.BehaviorInstances[0]->GetTransitionLayer());
+			NullInstanceCopy.SetOverrideLayer(LayerInfo.TransitionLayer);
 			NullInstanceCopy.Setup();
 		}
 	}
@@ -93,23 +109,36 @@ void FAvaTransitionExecutor::Setup()
 	// mark them as Needs Discard (this does not mean the scene will be discarded as there could be logic that reverts this flag)
 	for (const TPair<FAvaTag, FLayerInfo>& Pair : TagLayerInfo)
 	{
+		// Skip Layers that have no Enter Transition Instance to replace the existing one 
 		if (!EnumHasAnyFlags(Pair.Value.TransitionType, EAvaTransitionType::In))
 		{
 			continue;
 		}
 
-		for (FAvaTransitionBehaviorInstance* Instance : Pair.Value.BehaviorInstances)
+		for (FAvaTransitionBehaviorInstance* ExitInstance : Pair.Value.ExitInstances)
 		{
-			if (Instance->GetTransitionType() == EAvaTransitionType::In)
-			{
-				continue;
-			}
-
-			Instance->SetTransitionType(EAvaTransitionType::Out);
+			check(ExitInstance->GetTransitionType() == EAvaTransitionType::Out);
 
-			if (FAvaTransitionScene* TransitionScene = Instance->GetTransitionContext().GetTransitionScene())
+			if (FAvaTransitionScene* TransitionScene = ExitInstance->GetTransitionContext().GetTransitionScene())
 			{
-				TransitionScene->SetFlags(EAvaTransitionSceneFlags::NeedsDiscard);
+				bool bSceneReused = false;
+
+				const UAvaTransitionTree* TransitionTree = ExitInstance->GetTransitionTree();
+				if (TransitionTree && TransitionTree->GetInstancingMode() == EAvaTransitionInstancingMode::Reuse)
+				{
+					// Scene reused if the Tree Instancing Mode is set to reuse, and if there is an Enter Instance with matching Tree (i.e. Level)
+					bSceneReused |= Pair.Value.EnterInstances.ContainsByPredicate(
+						[TransitionTree](const FAvaTransitionBehaviorInstance* InEnterInstance)
+						{
+							return InEnterInstance && InEnterInstance->GetTransitionTree() == TransitionTree;
+						});
+				}
+
+				// Only mark for discard if scene not reused
+				if (!bSceneReused)
+				{
+					TransitionScene->SetFlags(EAvaTransitionSceneFlags::NeedsDiscard);
+				}
 			}
 		}
 	}
diff --git a/Engine/Plugins/Experimental/Avalanche/Source/AvalancheTransition/Public/Behavior/AvaTransitionBehaviorInstance.h b/Engine/Plugins/Experimental/Avalanche/Source/AvalancheTransition/Public/Behavior/AvaTransitionBehaviorInstance.h
--- a/Engine/Plugins/Experimental/Avalanche/Source/AvalancheTransition/Public/Behavior/AvaTransitionBehaviorInstance.h
+++ b/Engine/Plugins/Experimental/Avalanche/Source/AvalancheTransition/Public/Behavior/AvaTransitionBehaviorInstance.h
@@ -10,6 +10,7 @@
 
 class FAvaTransitionExecutorBuilder;
 class IAvaTransitionBehavior;
+class UAvaTransitionTree;
 struct FAvaTransitionExecutionContext;
 
 /** Struct containing Instance Data for a Transition Behavior running or about to run */
@@ -37,6 +38,8 @@ struct FAvaTransitionBehaviorInstance
 	/** Gets the Behavior this Instance is based on */
 	AVALANCHETRANSITION_API IAvaTransitionBehavior* GetBehavior() const;
 
+	AVALANCHETRANSITION_API const UAvaTransitionTree* GetTransitionTree() const;
+
 	AVALANCHETRANSITION_API FAvaTagHandle GetTransitionLayer() const;
 
 	AVALANCHETRANSITION_API EAvaTransitionType GetTransitionType() const;
