Motion Design:  Added Sequence Player Actor with spawn option in sequence row right click.
The Sequence Player inherits from LevelSequenceActor that overrides the Player class to the Motion Design Player.

diff --git a/Engine/Plugins/Experimental/Avalanche/Source/AvalancheEditor/Private/Sequencer/AvaSequencerExtension.h b/Engine/Plugins/Experimental/Avalanche/Source/AvalancheEditor/Private/Sequencer/AvaSequencerExtension.h
--- a/Engine/Plugins/Experimental/Avalanche/Source/AvalancheEditor/Private/Sequencer/AvaSequencerExtension.h
+++ b/Engine/Plugins/Experimental/Avalanche/Source/AvalancheEditor/Private/Sequencer/AvaSequencerExtension.h
@@ -53,8 +53,7 @@ public:
 	virtual TSharedPtr<IToolkitHost> GetSequencerToolkitHost() const override;
 	virtual UObject* GetPlaybackContext() const override;
 	virtual bool CanEditOrPlaySequences() const override;
-	virtual bool CanExportSequences() const override { return true; }
-	virtual void ExportSequences(TConstArrayView<UAvaSequence*> InSequencesToExport);
+	virtual void ExportSequences(TConstArrayView<UAvaSequence*> InSequencesToExport) override;
 	//~ End IAvaSequencerProvider
 
 protected:
diff --git a/Engine/Plugins/Experimental/Avalanche/Source/AvalancheSequence/Private/AvaSequencePlaybackActor.cpp b/Engine/Plugins/Experimental/Avalanche/Source/AvalancheSequence/Private/AvaSequencePlaybackActor.cpp
--- a/Engine/Plugins/Experimental/Avalanche/Source/AvalancheSequence/Private/AvaSequencePlaybackActor.cpp
+++ b/Engine/Plugins/Experimental/Avalanche/Source/AvalancheSequence/Private/AvaSequencePlaybackActor.cpp
@@ -460,7 +460,7 @@ UAvaSequencePlayer* AAvaSequencePlaybackActor::GetOrAddSequencePlayer(UAvaSequen
 		{
 			PlayerToGet = NewObject<UAvaSequencePlayer>(this, NAME_None, RF_Transient);
 			ActiveSequencePlayers.Add(PlayerToGet);
-			PlayerToGet->InitSequence(InSequence, this, GetLevel());
+			PlayerToGet->InitSequence(InSequence, this, GetLevel(), FLevelSequenceCameraSettings());
 		}
 
 		return PlayerToGet;
diff --git a/Engine/Plugins/Experimental/Avalanche/Source/AvalancheSequence/Private/AvaSequencePlayer.cpp b/Engine/Plugins/Experimental/Avalanche/Source/AvalancheSequence/Private/AvaSequencePlayer.cpp
--- a/Engine/Plugins/Experimental/Avalanche/Source/AvalancheSequence/Private/AvaSequencePlayer.cpp
+++ b/Engine/Plugins/Experimental/Avalanche/Source/AvalancheSequence/Private/AvaSequencePlayer.cpp
@@ -25,10 +25,10 @@ UAvaSequencePlayer::UAvaSequencePlayer(const FObjectInitializer& InObjectInitial
 	}
 }
 
-void UAvaSequencePlayer::InitSequence(UAvaSequence* InSequence, IAvaSequencePlaybackObject* InPlaybackObject, ULevel* InLevel)
+void UAvaSequencePlayer::InitSequence(UAvaSequence* InSequence, IAvaSequencePlaybackObject* InPlaybackObject, ULevel* InLevel, const FLevelSequenceCameraSettings& InCameraSettings)
 {
 	check(InSequence);
-	Super::Initialize(InSequence, InLevel, FLevelSequenceCameraSettings());
+	Super::Initialize(InSequence, InLevel, InCameraSettings);
 	SequenceController = UAvaSequenceSubsystem::CreateSequenceController(*InSequence, InPlaybackObject);
 	PlaybackObjectWeak = InPlaybackObject;
 	PlaybackLevelWeak  = InLevel;
diff --git a/Engine/Plugins/Experimental/Avalanche/Source/AvalancheSequence/Public/AvaSequencePlayer.h b/Engine/Plugins/Experimental/Avalanche/Source/AvalancheSequence/Public/AvaSequencePlayer.h
--- a/Engine/Plugins/Experimental/Avalanche/Source/AvalancheSequence/Public/AvaSequencePlayer.h
+++ b/Engine/Plugins/Experimental/Avalanche/Source/AvalancheSequence/Public/AvaSequencePlayer.h
@@ -13,6 +13,7 @@ class IAvaSequenceController;
 class IAvaSequencePlaybackObject;
 class UAvaSequence;
 struct FAvaSequencePlayParams;
+struct FLevelSequenceCameraSettings;
 
 UCLASS(Transient, BlueprintType)
 class AVALANCHESEQUENCE_API UAvaSequencePlayer : public ULevelSequencePlayer
@@ -27,7 +28,7 @@ public:
 	static FOnSequenceEvent& OnSequencePaused() { return OnSequencePausedDelegate; }
 	static FOnSequenceEvent& OnSequenceFinished() { return OnSequenceFinishedDelegate; }
 
-	void InitSequence(UAvaSequence* InSequence, IAvaSequencePlaybackObject* InPlaybackObject, ULevel* InLevel);
+	void InitSequence(UAvaSequence* InSequence, IAvaSequencePlaybackObject* InPlaybackObject, ULevel* InLevel, const FLevelSequenceCameraSettings& InCameraSettings);
 
 	UAvaSequence* GetAvaSequence() const;
 
diff --git a/Engine/Plugins/Experimental/Avalanche/Source/AvalancheSequencer/Private/AvaSequenceActorFactory.cpp b/Engine/Plugins/Experimental/Avalanche/Source/AvalancheSequencer/Private/AvaSequenceActorFactory.cpp
deleted file mode 100644
diff --git a/Engine/Plugins/Experimental/Avalanche/Source/AvalancheSequencer/Private/AvaSequenceActorFactory.h b/Engine/Plugins/Experimental/Avalanche/Source/AvalancheSequencer/Private/AvaSequenceActorFactory.h
deleted file mode 100644
diff --git a/Engine/Plugins/Experimental/Avalanche/Source/AvalancheSequencer/Private/AvaSequencer.cpp b/Engine/Plugins/Experimental/Avalanche/Source/AvalancheSequencer/Private/AvaSequencer.cpp
--- a/Engine/Plugins/Experimental/Avalanche/Source/AvalancheSequencer/Private/AvaSequencer.cpp
+++ b/Engine/Plugins/Experimental/Avalanche/Source/AvalancheSequencer/Private/AvaSequencer.cpp
@@ -14,9 +14,9 @@
 #include "CoreGlobals.h"
 #include "DetailsView/SAvaMarkDetails.h"
 #include "DetailsView/Section/AvaSequencePlaybackDetails.h"
-#include "DetailsView/Section/AvaSequencerEaseCurveToolSection.h"
 #include "DetailsView/Section/AvaSequenceSettingsDetails.h"
 #include "DetailsView/Section/AvaSequenceTreeDetails.h"
+#include "DetailsView/Section/AvaSequencerEaseCurveToolSection.h"
 #include "EaseCurveTool/AvaEaseCurveTool.h"
 #include "EaseCurveTool/AvaEaseCurveToolCommands.h"
 #include "Editor/Sequencer/Private/Sequencer.h"
@@ -36,21 +36,22 @@
 #include "MovieScene.h"
 #include "Playback/AvaSequencerCleanView.h"
 #include "Playback/AvaSequencerController.h"
+#include "Player/AvaSequenceActor.h"
 #include "PropertyHandle.h"
 #include "ScopedTransaction.h"
 #include "Selection/AvaEditorSelection.h"
-#include "SequencerCommands.h"
-#include "SequencerSettings.h"
-#include "SequencerUtilities.h"
 #include "SequenceTree/AvaSequenceItem.h"
 #include "SequenceTree/Columns/AvaSequenceNameColumn.h"
 #include "SequenceTree/Columns/AvaSequenceStatusColumn.h"
 #include "SequenceTree/Columns/IAvaSequenceColumn.h"
 #include "SequenceTree/Widgets/SAvaSequenceTree.h"
+#include "SequencerCommands.h"
+#include "SequencerSettings.h"
+#include "SequencerUtilities.h"
 #include "Settings/AvaSequencerSettings.h"
-#include "Sidebar/SidebarDrawerConfig.h"
 #include "Sidebar/SSidebar.h"
 #include "Sidebar/SSidebarContainer.h"
+#include "Sidebar/SidebarDrawerConfig.h"
 #include "ToolMenu.h"
 #include "ToolMenuEntry.h"
 #include "ToolMenuSection.h"
@@ -896,47 +897,51 @@ void FAvaSequencer::DuplicateSequence_Execute()
 	}
 }
 
-bool FAvaSequencer::ExportSequence_IsVisible() const
+bool FAvaSequencer::ExportSequence_CanExecute() const
 {
-	return Provider.CanExportSequences();
+	return SequenceTreeView.IsValid()
+		&& !SequenceTreeView->GetSelectedItems().IsEmpty();
 }
 
-bool FAvaSequencer::ExportSequence_CanExecute() const
+void FAvaSequencer::ExportSequence_Execute()
+{
+	Provider.ExportSequences(GetSelectedSequences());
+}
+
+bool FAvaSequencer::SpawnPlayer_CanExecute() const
 {
 	return SequenceTreeView.IsValid()
 		&& !SequenceTreeView->GetSelectedItems().IsEmpty();
 }
 
-void FAvaSequencer::ExportSequence_Execute()
+void FAvaSequencer::SpawnPlayer_Execute()
 {
-	if (!Provider.CanExportSequences() || !SequenceTreeView.IsValid())
+	if (!GEditor)
 	{
 		return;
 	}
 
-	const TArray<FAvaSequenceItemPtr> SelectedItems = SequenceTreeView->GetSelectedItems();
-	if (SelectedItems.IsEmpty())
+	UWorld* World = Provider.GetPlaybackContext()->GetWorld();
+	if (!World)
 	{
 		return;
 	}
 
-	TArray<UAvaSequence*> SequencesToExport;
-	SequencesToExport.Reserve(SelectedItems.Num());
-
-	for (const FAvaSequenceItemPtr& Item : SelectedItems)
+	UActorFactory* ActorFactory = GEditor->FindActorFactoryForActorClass(AAvaSequenceActor::StaticClass());
+	if (!ensure(ActorFactory))
 	{
-		if (!Item.IsValid())
-		{
-			continue;
-		}
-
-		if (UAvaSequence* Sequence = Item->GetSequence())
-		{
-			SequencesToExport.Add(Sequence);
-		}
+		return;
 	}
 
-	Provider.ExportSequences(SequencesToExport);
+	TArray<UAvaSequence*, TInlineAllocator<1>> Sequences = GetSelectedSequences();
+
+	FScopedTransaction Transaction(LOCTEXT("SpawnSequencePlayers", "Spawn Sequence Players"));
+
+	for (UAvaSequence* Sequence : Sequences)
+	{
+		check(Sequence);
+		GEditor->UseActorFactory(ActorFactory, FAssetData(Sequence), &FTransform::Identity);
+	}
 }
 
 bool FAvaSequencer::DeleteSequence_CanExecute() const
@@ -1682,6 +1687,38 @@ void FAvaSequencer::PostUndo(bool bSuccess)
 	NotifyOnSequenceTreeChanged();
 }
 
+TArray<UAvaSequence*, TInlineAllocator<1>> FAvaSequencer::GetSelectedSequences() const
+{
+	if (!SequenceTreeView.IsValid())
+	{
+		return {};
+	}
+
+	const TArray<FAvaSequenceItemPtr> SelectedItems = SequenceTreeView->GetSelectedItems();
+	if (SelectedItems.IsEmpty())
+	{
+		return {};
+	}
+
+	TArray<UAvaSequence*, TInlineAllocator<1>> SelectedSequences;
+	SelectedSequences.Reserve(SelectedItems.Num());
+
+	for (const FAvaSequenceItemPtr& Item : SelectedItems)
+	{
+		if (!Item.IsValid())
+		{
+			continue;
+		}
+
+		if (UAvaSequence* Sequence = Item->GetSequence())
+		{
+			SelectedSequences.Add(Sequence);
+		}
+	}
+
+	return SelectedSequences;
+}
+
 TSharedPtr<UE::Sequencer::SOutlinerView> FAvaSequencer::GetOutlinerView() const
 {
 	if (OutlinerViewWeak.IsValid())
diff --git a/Engine/Plugins/Experimental/Avalanche/Source/AvalancheSequencer/Private/AvaSequencer.h b/Engine/Plugins/Experimental/Avalanche/Source/AvalancheSequencer/Private/AvaSequencer.h
--- a/Engine/Plugins/Experimental/Avalanche/Source/AvalancheSequencer/Private/AvaSequencer.h
+++ b/Engine/Plugins/Experimental/Avalanche/Source/AvalancheSequencer/Private/AvaSequencer.h
@@ -132,10 +132,12 @@ public:
 	bool DuplicateSequence_CanExecute() const;
 	void DuplicateSequence_Execute();
 
-	bool ExportSequence_IsVisible() const;
 	bool ExportSequence_CanExecute() const;
 	void ExportSequence_Execute();
 
+	bool SpawnPlayer_CanExecute() const;
+	void SpawnPlayer_Execute();
+
 	bool DeleteSequence_CanExecute() const;
 	void DeleteSequence_Execute();
 
@@ -215,6 +217,8 @@ public:
 	// End of FEditorUndoClient
 
 private:
+	TArray<UAvaSequence*, TInlineAllocator<1>> GetSelectedSequences() const;
+
 	TSharedPtr<UE::Sequencer::SOutlinerView> GetOutlinerView() const;
 
 	void InitSequencerCommandList();
diff --git a/Engine/Plugins/Experimental/Avalanche/Source/AvalancheSequencer/Private/Commands/AvaSequencerCommands.cpp b/Engine/Plugins/Experimental/Avalanche/Source/AvalancheSequencer/Private/Commands/AvaSequencerCommands.cpp
--- a/Engine/Plugins/Experimental/Avalanche/Source/AvalancheSequencer/Private/Commands/AvaSequencerCommands.cpp
+++ b/Engine/Plugins/Experimental/Avalanche/Source/AvalancheSequencer/Private/Commands/AvaSequencerCommands.cpp
@@ -48,6 +48,12 @@ void FAvaSequencerCommands::RegisterCommands()
 		, EUserInterfaceActionType::Button
 		, FInputChord());
 
+	UI_COMMAND(SpawnSequencePlayer
+		, "Spawn Sequence Player"
+		, "Spawns a Sequence Player for the selected Sequence"
+		, EUserInterfaceActionType::Button
+		, FInputChord());
+
 	UI_COMMAND(StaggerLayerBars
 		, "Stagger Layer Bars"
 		, "Staggers the selected layer bars"
diff --git a/Engine/Plugins/Experimental/Avalanche/Source/AvalancheSequencer/Private/Commands/AvaSequencerCommands.h b/Engine/Plugins/Experimental/Avalanche/Source/AvalancheSequencer/Private/Commands/AvaSequencerCommands.h
--- a/Engine/Plugins/Experimental/Avalanche/Source/AvalancheSequencer/Private/Commands/AvaSequencerCommands.h
+++ b/Engine/Plugins/Experimental/Avalanche/Source/AvalancheSequencer/Private/Commands/AvaSequencerCommands.h
@@ -32,5 +32,7 @@ public:
 
 	TSharedPtr<FUICommandInfo> ExportSequence;
 
+	TSharedPtr<FUICommandInfo> SpawnSequencePlayer;
+
 	TSharedPtr<FUICommandInfo> StaggerLayerBars;
 };
diff --git a/Engine/Plugins/Experimental/Avalanche/Source/AvalancheSequencer/Private/Player/AvaSequenceActor.cpp b/Engine/Plugins/Experimental/Avalanche/Source/AvalancheSequencer/Private/Player/AvaSequenceActor.cpp
new file mode 100644
--- /dev/null
+++ b/Engine/Plugins/Experimental/Avalanche/Source/AvalancheSequencer/Private/Player/AvaSequenceActor.cpp
@@ -0,0 +1,60 @@
+// Copyright Epic Games, Inc. All Rights Reserved.
+
+#include "AvaSequenceActor.h"
+#include "AvaSequence.h"
+#include "AvaSequencePlayer.h"
+#include "AvaSequenceSubsystem.h"
+
+AAvaSequenceActor::AAvaSequenceActor(const FObjectInitializer& InObjectInitializer)
+	: Super(InObjectInitializer.SetDefaultSubobjectClass(TEXT("AnimationPlayer"), UAvaSequencePlayer::StaticClass()))
+{
+}
+
+void AAvaSequenceActor::Initialize(UAvaSequence* InSequence)
+{
+	if (ensure(!GetSequencePlayer()->IsPlaying()))
+	{
+		LevelSequenceAsset = InSequence;
+		InitSequencePlayer(InSequence);
+	}
+}
+
+void AAvaSequenceActor::PostInitializeComponents()
+{
+	Super::PostInitializeComponents();
+	InitSequencePlayer(Cast<UAvaSequence>(LevelSequenceAsset));
+}
+
+void AAvaSequenceActor::InitSequencePlayer(UAvaSequence* InSequence)
+{
+	if (!InSequence)
+	{
+		return;
+	}
+
+	ULevel* Level = GetLevel();
+	if (!Level)
+	{
+		return;
+	}
+
+	UAvaSequencePlayer* Player = Cast<UAvaSequencePlayer>(GetSequencePlayer());
+	if (!Player)
+	{
+		return;
+	}
+
+	UAvaSequenceSubsystem* SequenceSubsystem = UAvaSequenceSubsystem::Get(this);
+	if (!SequenceSubsystem)
+	{
+		return;
+	}
+
+	IAvaSequencePlaybackObject* PlaybackObject = SequenceSubsystem->FindPlaybackObject(Level);
+	if (!PlaybackObject)
+	{
+		return;
+	}
+
+	Player->InitSequence(InSequence, PlaybackObject, Level, CameraSettings);
+}
diff --git a/Engine/Plugins/Experimental/Avalanche/Source/AvalancheSequencer/Private/Player/AvaSequenceActor.h b/Engine/Plugins/Experimental/Avalanche/Source/AvalancheSequencer/Private/Player/AvaSequenceActor.h
new file mode 100644
--- /dev/null
+++ b/Engine/Plugins/Experimental/Avalanche/Source/AvalancheSequencer/Private/Player/AvaSequenceActor.h
@@ -0,0 +1,27 @@
+// Copyright Epic Games, Inc. All Rights Reserved.
+
+#pragma once
+
+#include "LevelSequenceActor.h"
+#include "AvaSequenceActor.generated.h"
+
+class UAvaSequence;
+class UAvaSequencePlayer;
+
+UCLASS()
+class AAvaSequenceActor : public ALevelSequenceActor
+{
+	GENERATED_BODY()
+
+public:
+	AAvaSequenceActor(const FObjectInitializer& InObjectInitializer);
+
+	void Initialize(UAvaSequence* InSequence);
+
+	//~ Begin AActor
+	virtual void PostInitializeComponents() override;
+	//~ End AActor
+
+private:
+	void InitSequencePlayer(UAvaSequence* InSequence);
+};
diff --git a/Engine/Plugins/Experimental/Avalanche/Source/AvalancheSequencer/Private/Player/AvaSequenceActorFactory.cpp b/Engine/Plugins/Experimental/Avalanche/Source/AvalancheSequencer/Private/Player/AvaSequenceActorFactory.cpp
new file mode 100644
--- /dev/null
+++ b/Engine/Plugins/Experimental/Avalanche/Source/AvalancheSequencer/Private/Player/AvaSequenceActorFactory.cpp
@@ -0,0 +1,78 @@
+// Copyright Epic Games, Inc. All Rights Reserved.
+
+#include "AvaSequenceActorFactory.h"
+#include "AvaSequence.h"
+#include "AvaSequenceActor.h"
+#include "ClassViewerFilter.h"
+#include "ClassViewerModule.h"
+#include "Modules/ModuleManager.h"
+
+#define LOCTEXT_NAMESPACE "AvaSequenceActorFactory"
+
+UAvaSequenceActorFactory::UAvaSequenceActorFactory()
+{
+	DisplayName = LOCTEXT("DisplayName", "Motion Design Sequence");
+	NewActorClass = AAvaSequenceActor::StaticClass();
+}
+
+bool UAvaSequenceActorFactory::CanCreateActorFrom(const FAssetData& InAssetData, FText& OutErrorMessage)
+{
+	FClassViewerModule& ClassViewerModule = FModuleManager::LoadModuleChecked<FClassViewerModule>("ClassViewer");
+
+	if (const TSharedPtr<IClassViewerFilter>& GlobalClassFilter = ClassViewerModule.GetGlobalClassViewerFilter())
+	{
+		TSharedRef<FClassViewerFilterFuncs> ClassFilterFuncs = ClassViewerModule.CreateFilterFuncs();
+		FClassViewerInitializationOptions ClassViewerOptions = {};
+
+		if (!GlobalClassFilter->IsClassAllowed(ClassViewerOptions, AAvaSequenceActor::StaticClass(), ClassFilterFuncs))
+		{
+			return false;
+		}
+	}
+
+	if (UActorFactory::CanCreateActorFrom(InAssetData, OutErrorMessage))
+	{
+		return true;
+	}
+
+	if (InAssetData.IsValid() && !InAssetData.IsInstanceOf(UAvaSequence::StaticClass()))
+	{
+		OutErrorMessage = LOCTEXT("NoSequenceAsset", "A valid Motion Design Sequence must be specified.");
+		return false;
+	}
+
+	return true;
+}
+
+AActor* UAvaSequenceActorFactory::SpawnActor(UObject* InAsset, ULevel* InLevel, const FTransform& InTransform, const FActorSpawnParameters& InSpawnParams)
+{
+	if (AAvaSequenceActor* NewActor = Cast<AAvaSequenceActor>(Super::SpawnActor(InAsset, InLevel, InTransform, InSpawnParams)))
+	{
+		if (UAvaSequence* Sequence = Cast<UAvaSequence>(InAsset))
+		{
+			NewActor->Initialize(Sequence);
+		}
+		return NewActor;
+	}
+	return nullptr;
+}
+
+UObject* UAvaSequenceActorFactory::GetAssetFromActorInstance(AActor* InActorInstance)
+{
+	if (AAvaSequenceActor* SequenceActor = Cast<AAvaSequenceActor>(InActorInstance))
+	{
+		return SequenceActor->GetSequence();
+	}
+	return nullptr;
+}
+
+FString UAvaSequenceActorFactory::GetDefaultActorLabel(UObject* InAsset) const
+{
+	if (UAvaSequence* Sequence = Cast<UAvaSequence>(InAsset))
+	{
+		return Sequence->GetLabel().ToString();
+	}
+	return Super::GetDefaultActorLabel(InAsset);
+}
+
+#undef LOCTEXT_NAMESPACE
diff --git a/Engine/Plugins/Experimental/Avalanche/Source/AvalancheSequencer/Private/Player/AvaSequenceActorFactory.h b/Engine/Plugins/Experimental/Avalanche/Source/AvalancheSequencer/Private/Player/AvaSequenceActorFactory.h
new file mode 100644
--- /dev/null
+++ b/Engine/Plugins/Experimental/Avalanche/Source/AvalancheSequencer/Private/Player/AvaSequenceActorFactory.h
@@ -0,0 +1,23 @@
+// Copyright Epic Games, Inc. All Rights Reserved.
+
+#pragma once
+
+#include "ActorFactories/ActorFactory.h"
+#include "AvaSequenceActorFactory.generated.h"
+
+UCLASS()
+class UAvaSequenceActorFactory : public UActorFactory
+{
+	GENERATED_BODY()
+
+public:
+	UAvaSequenceActorFactory();
+
+protected:
+	//~ Begin UActorFactory
+	virtual bool CanCreateActorFrom(const FAssetData& InAssetData, FText& OutErrorMessage) override;
+	virtual AActor* SpawnActor(UObject* InAsset, ULevel* InLevel, const FTransform& InTransform, const FActorSpawnParameters& InSpawnParams) override;
+	virtual UObject* GetAssetFromActorInstance(AActor* InActorInstance) override;
+	virtual FString GetDefaultActorLabel(UObject* InAsset) const override;
+	//~ End UActorFactory
+};
diff --git a/Engine/Plugins/Experimental/Avalanche/Source/AvalancheSequencer/Private/SequenceTree/AvaSequenceItem.cpp b/Engine/Plugins/Experimental/Avalanche/Source/AvalancheSequencer/Private/SequenceTree/AvaSequenceItem.cpp
--- a/Engine/Plugins/Experimental/Avalanche/Source/AvalancheSequencer/Private/SequenceTree/AvaSequenceItem.cpp
+++ b/Engine/Plugins/Experimental/Avalanche/Source/AvalancheSequencer/Private/SequenceTree/AvaSequenceItem.cpp
@@ -2,13 +2,13 @@
 
 #include "AvaSequenceItem.h"
 #include "AvaSequence.h"
-#include "AvaSequenceActorFactory.h"
 #include "AvaSequencePlaybackObject.h"
 #include "AvaSequencePlayer.h"
 #include "AvaSequencer.h"
 #include "DragDropOps/AvaSequenceItemDragDropOp.h"
 #include "Editor.h"
 #include "Kismet2/BlueprintEditorUtils.h"
+#include "Player/AvaSequenceActorFactory.h"
 #include "ScopedTransaction.h"
 
 #define LOCTEXT_NAMESPACE "AvaSequenceItem"
diff --git a/Engine/Plugins/Experimental/Avalanche/Source/AvalancheSequencer/Private/SequenceTree/DragDropOps/AvaSequenceItemDragDropOp.cpp b/Engine/Plugins/Experimental/Avalanche/Source/AvalancheSequencer/Private/SequenceTree/DragDropOps/AvaSequenceItemDragDropOp.cpp
--- a/Engine/Plugins/Experimental/Avalanche/Source/AvalancheSequencer/Private/SequenceTree/DragDropOps/AvaSequenceItemDragDropOp.cpp
+++ b/Engine/Plugins/Experimental/Avalanche/Source/AvalancheSequencer/Private/SequenceTree/DragDropOps/AvaSequenceItemDragDropOp.cpp
@@ -2,7 +2,7 @@
 
 #include "AvaSequenceItemDragDropOp.h"
 #include "AvaSequence.h"
-#include "AvaSequenceActorFactory.h"
+#include "Player/AvaSequenceActorFactory.h"
 #include "SequenceTree/IAvaSequenceItem.h"
 
 TSharedRef<FAvaSequenceItemDragDropOp> FAvaSequenceItemDragDropOp::New(const FAvaSequenceItemPtr& InSequenceItem, UActorFactory* InActorFactory)
diff --git a/Engine/Plugins/Experimental/Avalanche/Source/AvalancheSequencer/Private/SequenceTree/Widgets/SAvaSequenceTree.cpp b/Engine/Plugins/Experimental/Avalanche/Source/AvalancheSequencer/Private/SequenceTree/Widgets/SAvaSequenceTree.cpp
--- a/Engine/Plugins/Experimental/Avalanche/Source/AvalancheSequencer/Private/SequenceTree/Widgets/SAvaSequenceTree.cpp
+++ b/Engine/Plugins/Experimental/Avalanche/Source/AvalancheSequencer/Private/SequenceTree/Widgets/SAvaSequenceTree.cpp
@@ -172,6 +172,7 @@ TSharedPtr<SWidget> SAvaSequenceTree::OnContextMenuOpening() const
 			, LOCTEXT("ApplyPresetTooltip", "Apply a Preset to the Selected Sequences")
 			, FNewToolMenuDelegate::CreateStatic(&UE::AvaSequencer::Private::GeneratePresetMenu));
 
+		Section.AddMenuEntry(SequencerCommands.SpawnSequencePlayer);
 		Section.AddMenuEntry(SequencerCommands.ExportSequence);
 		Section.AddMenuEntry(GenericCommands.Rename);
 		Section.AddMenuEntry(GenericCommands.Duplicate);
@@ -220,9 +221,11 @@ void SAvaSequenceTree::BindCommands(const TSharedRef<FAvaSequencer>& InSequencer
 
 	CommandList->MapAction(AvaSequencerCommands.ExportSequence
 		, FExecuteAction::CreateSP(InSequencer, &FAvaSequencer::ExportSequence_Execute)
-		, FCanExecuteAction::CreateSP(InSequencer, &FAvaSequencer::ExportSequence_CanExecute)
-		, FIsActionChecked()
-		, FIsActionButtonVisible::CreateSP(InSequencer, &FAvaSequencer::ExportSequence_IsVisible));
+		, FCanExecuteAction::CreateSP(InSequencer, &FAvaSequencer::ExportSequence_CanExecute));
+
+	CommandList->MapAction(AvaSequencerCommands.SpawnSequencePlayer
+		, FExecuteAction::CreateSP(InSequencer, &FAvaSequencer::SpawnPlayer_Execute)
+		, FCanExecuteAction::CreateSP(InSequencer, &FAvaSequencer::SpawnPlayer_CanExecute));
 }
 
 TSharedRef<ITableRow> SAvaSequenceTree::OnGenerateRow(FAvaSequenceItemPtr InListItem
diff --git a/Engine/Plugins/Experimental/Avalanche/Source/AvalancheSequencer/Public/IAvaSequencerProvider.h b/Engine/Plugins/Experimental/Avalanche/Source/AvalancheSequencer/Public/IAvaSequencerProvider.h
--- a/Engine/Plugins/Experimental/Avalanche/Source/AvalancheSequencer/Public/IAvaSequencerProvider.h
+++ b/Engine/Plugins/Experimental/Avalanche/Source/AvalancheSequencer/Public/IAvaSequencerProvider.h
@@ -50,9 +50,6 @@ public:
 	/** Called when a camera cut occurs. */
 	virtual void OnUpdateCameraCut(UObject* InCameraObject, bool bInJumpCut) {};
 
-	/** Called whether this Sequencer Provider can Export Sequences */
-	virtual bool CanExportSequences() const { return false; }
-
 	/** Implementation of converting the Sequences to Level Sequence assets */
 	virtual void ExportSequences(TConstArrayView<UAvaSequence*> InSequencesToExport) {}
 };
