MotionDesign : AvaText3D
- Fix same material assigned on all text sections (front, back, bevel, extrude) and their color not being applied properly

diff --git a/Engine/Plugins/Experimental/Avalanche/Source/AvalancheText/Private/AvaText3DComponent.cpp b/Engine/Plugins/Experimental/Avalanche/Source/AvalancheText/Private/AvaText3DComponent.cpp
--- a/Engine/Plugins/Experimental/Avalanche/Source/AvalancheText/Private/AvaText3DComponent.cpp
+++ b/Engine/Plugins/Experimental/Avalanche/Source/AvalancheText/Private/AvaText3DComponent.cpp
@@ -28,6 +28,11 @@ namespace UE::AvaText::Private
 
 		if (UMaterialInstance* CurrentMaterialInstance = Cast<UMaterialInstance>(Invoke(InGetter, InComponent)))
 		{
+			if (CurrentMaterialInstance == InNewMaterial)
+			{
+				return;
+			}
+
 			InNewMaterial->CopyInterpParameters(CurrentMaterialInstance);
 		}
 
@@ -506,22 +511,31 @@ void UAvaText3DComponent::SetupSolidMaterial()
 		const EAvaTextMaterialFeatures MaterialFeatures = GetMaterialFeaturesFromProperties();
 		const FAvaTextMaterialSettings MaterialSettings(EAvaTextColoringStyle::Solid, MaterialFeatures);
 
+		// Front & Back
+		UMaterialInstanceDynamic* PreviousMaterial = nullptr;
+
 		if (UMaterialInstanceDynamic* const SolidMaterial = GetMIDWithSettings(TextMaterialInstance, MaterialSettings))
 		{
-			UMaterialInstanceDynamic* PreviousMaterial = TextMaterialInstance;
+			PreviousMaterial = TextMaterialInstance;
 			TextMaterialInstance = SolidMaterial;
 			if (PreviousMaterial)
 			{
 				TextMaterialInstance->CopyInterpParameters(PreviousMaterial);
 			}
+		}
 
+		if (UMaterialInstanceDynamic* const SolidMaterial = GetMIDWithSettings(ExtrudeTextMaterialInstance, MaterialSettings))
+		{
 			PreviousMaterial = ExtrudeTextMaterialInstance;
 			ExtrudeTextMaterialInstance = SolidMaterial;
 			if (PreviousMaterial)
 			{
 				ExtrudeTextMaterialInstance->CopyInterpParameters(PreviousMaterial);
 			}
+		}
 
+		if (UMaterialInstanceDynamic* const SolidMaterial = GetMIDWithSettings(BevelTextMaterialInstance, MaterialSettings))
+		{
 			PreviousMaterial = BevelTextMaterialInstance;
 			BevelTextMaterialInstance = SolidMaterial;
 			if (PreviousMaterial)
@@ -532,7 +546,6 @@ void UAvaText3DComponent::SetupSolidMaterial()
 
 		RemoveRefreshReason(EAvaTextRefreshReason::MaterialInstancesChange);
 	}
- 
 	{
 		UE::AvaText::Private::PatchMaterial(
 			this
@@ -1296,7 +1309,7 @@ void UAvaText3DComponent::RefreshExtrudeColor() const
 
 	if (ExtrudeTextMaterialInstance)
 	{
-		ExtrudeTextMaterialInstance->SetVectorParameterValue(UAvaTextMaterialHub::FMaterialStatics::SolidColor_MatParam, Color);
+		ExtrudeTextMaterialInstance->SetVectorParameterValue(UAvaTextMaterialHub::FMaterialStatics::SolidColor_MatParam, ExtrudeColor);
 	}
 }
 
@@ -1309,7 +1322,7 @@ void UAvaText3DComponent::RefreshBevelColor() const
 
 	if (BevelTextMaterialInstance)
 	{
-		BevelTextMaterialInstance->SetVectorParameterValue(UAvaTextMaterialHub::FMaterialStatics::SolidColor_MatParam, Color);
+		BevelTextMaterialInstance->SetVectorParameterValue(UAvaTextMaterialHub::FMaterialStatics::SolidColor_MatParam, BevelColor);
 	}
 }
 
diff --git a/Engine/Plugins/Experimental/Text3D/Source/Text3D/Private/Text3DComponent.cpp b/Engine/Plugins/Experimental/Text3D/Source/Text3D/Private/Text3DComponent.cpp
--- a/Engine/Plugins/Experimental/Text3D/Source/Text3D/Private/Text3DComponent.cpp
+++ b/Engine/Plugins/Experimental/Text3D/Source/Text3D/Private/Text3DComponent.cpp
@@ -1417,7 +1417,12 @@ void UText3DComponent::OnMaterialChanged()
 				continue;
 			}
 
-			StaticMeshComponent->SetMaterial(MaterialIndex, GetMaterial(static_cast<EText3DGroupType>(GroupIndex)));
+			UMaterialInterface* Material = GetMaterial(static_cast<EText3DGroupType>(GroupIndex));
+
+			if (Material != StaticMeshComponent->GetMaterial(MaterialIndex))
+			{
+				StaticMeshComponent->SetMaterial(MaterialIndex, Material);
+			}
 		}
 	}
 }
