Motion Design Transition: Added Default Tree Template settings to build default transitions on new levels

diff --git a/Engine/Plugins/Experimental/Avalanche/Source/AvalancheTransitionEditor/AvalancheTransitionEditor.Build.cs b/Engine/Plugins/Experimental/Avalanche/Source/AvalancheTransitionEditor/AvalancheTransitionEditor.Build.cs
--- a/Engine/Plugins/Experimental/Avalanche/Source/AvalancheTransitionEditor/AvalancheTransitionEditor.Build.cs
+++ b/Engine/Plugins/Experimental/Avalanche/Source/AvalancheTransitionEditor/AvalancheTransitionEditor.Build.cs
@@ -26,6 +26,7 @@ public class AvalancheTransitionEditor : ModuleRules
                 "AssetTools",
                 "AvalancheCore",
                 "AvalancheTransition",
+                "DeveloperSettings",
                 "EditorStyle",
                 "Engine",
                 "InputCore",
diff --git a/Engine/Plugins/Experimental/Avalanche/Source/AvalancheTransitionEditor/Private/AvaTransitionEditorModule.cpp b/Engine/Plugins/Experimental/Avalanche/Source/AvalancheTransitionEditor/Private/AvaTransitionEditorModule.cpp
--- a/Engine/Plugins/Experimental/Avalanche/Source/AvalancheTransitionEditor/Private/AvaTransitionEditorModule.cpp
+++ b/Engine/Plugins/Experimental/Avalanche/Source/AvalancheTransitionEditor/Private/AvaTransitionEditorModule.cpp
@@ -14,6 +14,7 @@
 #include "Styling/AvaTransitionEditorStyle.h"
 #include "ToolMenu.h"
 #include "ToolMenuSection.h"
+#include "Settings/AvaTransitionEditorSettings.h"
 #include "ViewModels/AvaTransitionEditorViewModel.h"
 #include "Widgets/Layout/SBox.h"
 
@@ -126,23 +127,35 @@ void FAvaTransitionEditorModule::ValidateStateTree(UAvaTransitionTree* InTransit
 	// Disable Tree by default if being set up for the first time
 	InTransitionTree->SetEnabled(false);
 
-	UAvaTransitionTreeEditorData* const EditorData = NewObject<UAvaTransitionTreeEditorData>(InTransitionTree, NAME_None, RF_Transactional);
-	check(EditorData);
+	UAvaTransitionTreeEditorData* EditorData;
 
-	EditorData->Schema = NewObject<UAvaTransitionTreeSchema>(EditorData);
+	const UAvaTransitionEditorSettings* TransitionEditorSettings = GetDefault<UAvaTransitionEditorSettings>();
+	check(TransitionEditorSettings);
 
-	InTransitionTree->EditorData = EditorData;
-
-	// Build the Default Tree
-	if (OnBuildDefaultTransitionTree.IsBound())
+	if (UAvaTransitionTreeEditorData* TemplateEditorData = TransitionEditorSettings->LoadDefaultTemplateEditorData())
 	{
-		OnBuildDefaultTransitionTree.Execute(*EditorData);
+		EditorData = DuplicateObject<UAvaTransitionTreeEditorData>(TemplateEditorData, InTransitionTree);
+		check(EditorData);
 	}
 	else
 	{
-		EditorData->AddRootState();	
+		EditorData = NewObject<UAvaTransitionTreeEditorData>(InTransitionTree, NAME_None, RF_Transactional);
+		check(EditorData);
+
+		EditorData->Schema = NewObject<UAvaTransitionTreeSchema>(EditorData);
+
+		if (OnBuildDefaultTransitionTree.IsBound())
+		{
+			OnBuildDefaultTransitionTree.Execute(*EditorData);
+		}
+		else
+		{
+			EditorData->AddRootState();	
+		}
 	}
 
+	InTransitionTree->EditorData = EditorData;
+
 	// Compile in Advanced Mode here so that no new nodes are generated from outside
 	FAvaTransitionCompiler Compiler;
 	Compiler.SetTransitionTree(InTransitionTree);
diff --git a/Engine/Plugins/Experimental/Avalanche/Source/AvalancheTransitionEditor/Private/Settings/AvaTransitionEditorSettings.cpp b/Engine/Plugins/Experimental/Avalanche/Source/AvalancheTransitionEditor/Private/Settings/AvaTransitionEditorSettings.cpp
new file mode 100644
--- /dev/null
+++ b/Engine/Plugins/Experimental/Avalanche/Source/AvalancheTransitionEditor/Private/Settings/AvaTransitionEditorSettings.cpp
@@ -0,0 +1,20 @@
+// Copyright Epic Games, Inc. All Rights Reserved.
+
+#include "AvaTransitionEditorSettings.h"
+#include "AvaTransitionTree.h"
+#include "AvaTransitionTreeEditorData.h"
+
+UAvaTransitionEditorSettings::UAvaTransitionEditorSettings()
+{
+	CategoryName = TEXT("Motion Design");
+	SectionName  = TEXT("Transition Logic");
+}
+
+UAvaTransitionTreeEditorData* UAvaTransitionEditorSettings::LoadDefaultTemplateEditorData() const
+{
+	if (UAvaTransitionTree* TemplateTree = DefaultTemplate.LoadSynchronous())
+	{
+		return Cast<UAvaTransitionTreeEditorData>(TemplateTree->EditorData);
+	}
+	return nullptr;
+}
diff --git a/Engine/Plugins/Experimental/Avalanche/Source/AvalancheTransitionEditor/Private/Settings/AvaTransitionEditorSettings.h b/Engine/Plugins/Experimental/Avalanche/Source/AvalancheTransitionEditor/Private/Settings/AvaTransitionEditorSettings.h
new file mode 100644
--- /dev/null
+++ b/Engine/Plugins/Experimental/Avalanche/Source/AvalancheTransitionEditor/Private/Settings/AvaTransitionEditorSettings.h
@@ -0,0 +1,26 @@
+// Copyright Epic Games, Inc. All Rights Reserved.
+
+#pragma once
+
+#include "Engine/DeveloperSettings.h"
+#include "UObject/SoftObjectPtr.h"
+#include "AvaTransitionEditorSettings.generated.h"
+
+class UAvaTransitionTree;
+class UAvaTransitionTreeEditorData;
+
+UCLASS(config=EditorPerProjectUserSettings, meta=(DisplayName="Transition Logic"))
+class UAvaTransitionEditorSettings : public UDeveloperSettings
+{
+	GENERATED_BODY()
+
+public:
+	UAvaTransitionEditorSettings();
+
+	UAvaTransitionTreeEditorData* LoadDefaultTemplateEditorData() const;
+
+private:
+	/** The template to use when building new transition trees */
+	UPROPERTY(Config, EditAnywhere, Category="Motion Design")
+	TSoftObjectPtr<UAvaTransitionTree> DefaultTemplate;
+};
