[Motion Design] Rundown: Fixing the "missing call to EndPlay" warning.

diff --git a/Engine/Plugins/Experimental/Avalanche/Source/Avalanche/Private/Framework/AvaGameInstance.cpp b/Engine/Plugins/Experimental/Avalanche/Source/Avalanche/Private/Framework/AvaGameInstance.cpp
--- a/Engine/Plugins/Experimental/Avalanche/Source/Avalanche/Private/Framework/AvaGameInstance.cpp
+++ b/Engine/Plugins/Experimental/Avalanche/Source/Avalanche/Private/Framework/AvaGameInstance.cpp
@@ -366,6 +366,11 @@ void UAvaGameInstance::EndPlayWorld()
 	OnEndPlay.Broadcast(this, PlayingChannelName);
 	PlayingChannelName = NAME_None;
 
+	if (PlayWorld)
+	{
+		PlayWorld->EndPlay(EEndPlayReason::LevelTransition);
+	}
+
 	Viewport.Reset();
 	RenderTargetFence.Reset();
 	bIsRenderTargetReady = false;
