MotionDesign : Modifiers
- Improved radial arrange modifiers orientation (look at) code
- Added base rotation to adjust mesh with different forward vector
- Removed plane option and allowing to orient the whole container without modifier adjusting layout for you

diff --git a/Engine/Plugins/Experimental/Avalanche/Source/AvalancheModifiers/Private/AvaModifiersActorUtils.cpp b/Engine/Plugins/Experimental/Avalanche/Source/AvalancheModifiers/Private/AvaModifiersActorUtils.cpp
--- a/Engine/Plugins/Experimental/Avalanche/Source/AvalancheModifiers/Private/AvaModifiersActorUtils.cpp
+++ b/Engine/Plugins/Experimental/Avalanche/Source/AvalancheModifiers/Private/AvaModifiersActorUtils.cpp
@@ -2,7 +2,6 @@
 
 #include "AvaModifiersActorUtils.h"
 
-#include "AvaDefs.h"
 #include "Components/PrimitiveComponent.h"
 #include "GameFramework/Actor.h"
 
@@ -197,34 +196,37 @@ bool FAvaModifiersActorUtils::IsActorNotIsolated(const AActor* InActor)
 
 FRotator FAvaModifiersActorUtils::FindLookAtRotation(const FVector& InEyePosition, const FVector& InTargetPosition, EAvaModifiersAxis InAxis, bool bInFlipAxis)
 {
-	FVector OutDirection = bInFlipAxis ? InEyePosition - InTargetPosition : InTargetPosition - InEyePosition;
-	OutDirection = OutDirection.GetSafeNormal(UE_SMALL_NUMBER, FVector::XAxisVector);
+	const FVector Direction = bInFlipAxis ? (InTargetPosition - InEyePosition).GetSafeNormal() : (InEyePosition - InTargetPosition).GetSafeNormal();
 
-	FMatrix NewRotation = FRotationMatrix::Identity;
+	if (Direction.IsNearlyZero())
+	{
+		return FRotator::ZeroRotator;
+	}
+
+	const FRotator BaseRotation = Direction.Rotation();
 
+	FQuat AxisQuat;
 	switch (InAxis)
 	{
 	case EAvaModifiersAxis::X:
-		NewRotation = FRotationMatrix::MakeFromX(OutDirection);
+		AxisQuat = FQuat::Identity;
 		break;
+
 	case EAvaModifiersAxis::Y:
-		NewRotation = FRotationMatrix::MakeFromY(OutDirection);
+		AxisQuat = FQuat(FVector::ZAxisVector, FMath::DegreesToRadians(90.0f));
 		break;
+
 	case EAvaModifiersAxis::Z:
-		NewRotation = FRotationMatrix::MakeFromZ(OutDirection);
+		AxisQuat = FQuat(FVector::YAxisVector, FMath::DegreesToRadians(-90.0f));
 		break;
-	case EAvaModifiersAxis::None:
-		return NewRotation.Rotator();
-	}
-
-	FRotator Result = NewRotation.Rotator();
 
-	if (bInFlipAxis && InAxis == EAvaModifiersAxis::X)
-	{
-		Result = Result + FRotator(180, 0, 0);
+	default:
+		return FRotator::ZeroRotator;
 	}
 
-	return Result;
+	const FQuat FinalQuat = BaseRotation.Quaternion() * AxisQuat;
+
+	return FinalQuat.Rotator();
 }
 
 bool FAvaModifiersActorUtils::IsActorVisible(const AActor* InActor)
diff --git a/Engine/Plugins/Experimental/Avalanche/Source/AvalancheModifiers/Private/Modifiers/AvaRadialArrangeModifier.cpp b/Engine/Plugins/Experimental/Avalanche/Source/AvalancheModifiers/Private/Modifiers/AvaRadialArrangeModifier.cpp
--- a/Engine/Plugins/Experimental/Avalanche/Source/AvalancheModifiers/Private/Modifiers/AvaRadialArrangeModifier.cpp
+++ b/Engine/Plugins/Experimental/Avalanche/Source/AvalancheModifiers/Private/Modifiers/AvaRadialArrangeModifier.cpp
@@ -2,7 +2,6 @@
 
 #include "Modifiers/AvaRadialArrangeModifier.h"
 
-#include "AvaDefs.h"
 #include "AvaModifiersActorUtils.h"
 #include "GameFramework/Actor.h"
 #include "Shared/AvaTransformModifierShared.h"
@@ -33,7 +32,6 @@ void UAvaRadialArrangeModifier::PostEditChangeProperty(FPropertyChangedEvent& Pr
 
 	static const TSet<FName> PropertiesName =
 	{
-		GET_MEMBER_NAME_CHECKED(UAvaRadialArrangeModifier, Plane),
 		GET_MEMBER_NAME_CHECKED(UAvaRadialArrangeModifier, Count),
 		GET_MEMBER_NAME_CHECKED(UAvaRadialArrangeModifier, Rings),
 		GET_MEMBER_NAME_CHECKED(UAvaRadialArrangeModifier, InnerRadius),
@@ -44,6 +42,7 @@ void UAvaRadialArrangeModifier::PostEditChangeProperty(FPropertyChangedEvent& Pr
 		GET_MEMBER_NAME_CHECKED(UAvaRadialArrangeModifier, bStartFromOuterRadius),
 		GET_MEMBER_NAME_CHECKED(UAvaRadialArrangeModifier, bOrient),
 		GET_MEMBER_NAME_CHECKED(UAvaRadialArrangeModifier, OrientationAxis),
+		GET_MEMBER_NAME_CHECKED(UAvaRadialArrangeModifier, BaseOrientation),
 		GET_MEMBER_NAME_CHECKED(UAvaRadialArrangeModifier, bFlipOrient)
 	};
 
@@ -54,17 +53,6 @@ void UAvaRadialArrangeModifier::PostEditChangeProperty(FPropertyChangedEvent& Pr
 }
 #endif // WITH_EDITOR
 
-void UAvaRadialArrangeModifier::SetPlane(EAvaRadialArrangePlane InPlane)
-{
-	if (Plane == InPlane)
-	{
-		return;
-	}
-
-	Plane = InPlane;
-	MarkModifierDirty();
-}
-
 void UAvaRadialArrangeModifier::SetCount(const int32 InCount)
 {
 	if (Count == InCount)
@@ -175,6 +163,17 @@ void UAvaRadialArrangeModifier::SetOrientationAxis(EAvaModifiersAxis InAxis)
 	MarkModifierDirty();
 }
 
+void UAvaRadialArrangeModifier::SetBaseOrientation(const FRotator& InRotation)
+{
+	if (BaseOrientation.Equals(InRotation))
+	{
+		return;
+	}
+
+	BaseOrientation = InRotation;
+	MarkModifierDirty();
+}
+
 void UAvaRadialArrangeModifier::SetFlipOrient(const bool bInFlipOrient)
 {
 	if (bFlipOrient == bInFlipOrient)
@@ -197,6 +196,22 @@ void UAvaRadialArrangeModifier::OnModifierCDOSetup(FActorModifierCoreMetadata& I
 #endif
 }
 
+void UAvaRadialArrangeModifier::OnModifierAdded(EActorModifierCoreEnableReason InReason)
+{
+	Super::OnModifierAdded(InReason);
+
+	if (InReason == EActorModifierCoreEnableReason::User)
+	{
+		OrientationAxis = EAvaModifiersAxis::X;
+	}
+}
+
+void UAvaRadialArrangeModifier::OnModifiedActorTransformed()
+{
+	// Overwrite parent class behaviour don't do anything when moved
+	// Let user rotate the container and choose the wanted plane
+}
+
 void UAvaRadialArrangeModifier::Apply()
 {
 	AActor* ModifyActor = GetModifiedActor();
@@ -224,7 +239,7 @@ void UAvaRadialArrangeModifier::Apply()
 	const float RadiusDistance = Rings > 1 ? OuterRadius - InnerRadius : 0.0f;
 	const float RadiusDistancePerRing = RadiusDistance / Rings;
 
-	auto CalculateRelativeOffset = [this, RadiusDistancePerRing](float InAngleInDegrees, const int32 InRingIndex) -> FVector
+	auto CalculateRelativeOffset = [this, RadiusDistancePerRing](float InAngleInDegrees, const int32 InRingIndex)->FVector
 	{
 		InAngleInDegrees = FRotator::NormalizeAxis(InAngleInDegrees);
 
@@ -235,22 +250,7 @@ void UAvaRadialArrangeModifier::Apply()
 		const float RingStartOffset = RadiusDistancePerRing * InRingIndex;
 		const float ChildRadius = InnerRadius + RingStartOffset;
 
-		FVector Offset = FVector::ZeroVector;
-
-		switch (Plane)
-		{
-			case EAvaRadialArrangePlane::XY:
-				Offset = FVector(ChildRadius * SlotCos, ChildRadius * SlotSin, 0);
-				break;
-			case EAvaRadialArrangePlane::YZ:
-				Offset = FVector(0, ChildRadius * SlotCos, ChildRadius * SlotSin);
-				break;
-			case EAvaRadialArrangePlane::XZ:
-				Offset = FVector(ChildRadius * SlotCos, 0, ChildRadius * SlotSin);
-				break;
-		}
-
-		return Offset;
+		return FVector(ChildRadius * SlotCos, ChildRadius * SlotSin, 0);
 	};
 
 	UAvaTransformModifierShared* LayoutShared = GetShared<UAvaTransformModifierShared>(true);
@@ -350,12 +350,12 @@ void UAvaRadialArrangeModifier::Apply()
 
 		if (bOrient)
 		{
-			const FVector EyePosition = AttachedActor->GetActorLocation();
-			const FVector TargetPosition = ModifyActor->GetActorLocation();
+			const FVector EyePosition = RelativeOffset;
+			const FVector TargetPosition = FVector::ZeroVector;
 
-			const FRotator NewRotation = FAvaModifiersActorUtils::FindLookAtRotation(EyePosition, TargetPosition, OrientationAxis, bFlipOrient);
+			FRotator NewRotation = BaseOrientation + FAvaModifiersActorUtils::FindLookAtRotation(EyePosition, TargetPosition, OrientationAxis, bFlipOrient);
 
-			AttachedActor->SetActorRotation(NewRotation);
+			AttachedActor->SetActorRelativeRotation(NewRotation);
 		}
 		// Restore original rotation
 		else if (FAvaTransformSharedActorState* ActorState = LayoutShared->FindActorState(AttachedActor))
diff --git a/Engine/Plugins/Experimental/Avalanche/Source/AvalancheModifiers/Public/Modifiers/AvaRadialArrangeModifier.h b/Engine/Plugins/Experimental/Avalanche/Source/AvalancheModifiers/Public/Modifiers/AvaRadialArrangeModifier.h
--- a/Engine/Plugins/Experimental/Avalanche/Source/AvalancheModifiers/Public/Modifiers/AvaRadialArrangeModifier.h
+++ b/Engine/Plugins/Experimental/Avalanche/Source/AvalancheModifiers/Public/Modifiers/AvaRadialArrangeModifier.h
@@ -40,12 +40,6 @@ class UAvaRadialArrangeModifier : public UAvaArrangeBaseModifier
 	GENERATED_BODY()
 
 public:
-	AVALANCHEMODIFIERS_API void SetPlane(EAvaRadialArrangePlane InPlane);
-	EAvaRadialArrangePlane GetPlane() const
-	{
-		return Plane;
-	}
-
 	/** Sets the number of child elements to use in the arrangement. Children whose index is greater than or equal to this value will be hidden. */
 	UFUNCTION(BlueprintCallable, Category="Motion Design|Modifiers|RadialArrange")
 	AVALANCHEMODIFIERS_API void SetCount(const int32 InCount);
@@ -156,6 +150,17 @@ public:
 		return OrientationAxis;
 	}
 
+	/** Sets the base rotation */
+	UFUNCTION(BlueprintCallable, Category="Motion Design|Modifiers|RadialArrange")
+	AVALANCHEMODIFIERS_API void SetBaseOrientation(const FRotator& InRotation);
+
+	/** Gets the base rotation */
+	UFUNCTION(BlueprintPure, Category="Motion Design|Modifiers|RadialArrange")
+	const FRotator& GetBaseOrientation() const
+	{
+		return BaseOrientation;
+	}
+
 	/** If true, will flip the center orientation to face outwards. */
 	UFUNCTION(BlueprintCallable, Category="Motion Design|Modifiers|RadialArrange")
 	AVALANCHEMODIFIERS_API void SetFlipOrient(const bool bInFlipOrient);
@@ -177,13 +182,11 @@ protected:
 
 	//~ Begin UActorModifierCoreBase
 	virtual void OnModifierCDOSetup(FActorModifierCoreMetadata& InMetadata) override;
+	virtual void OnModifierAdded(EActorModifierCoreEnableReason InReason) override;
+	virtual void OnModifiedActorTransformed() override;
 	virtual void Apply() override;
 	//~ End UActorModifierCoreBase
 
-	/** Base plane for the radial arrangement */
-	UPROPERTY(EditInstanceOnly, Setter, Getter, Category="RadialArrange", meta=(AllowPrivateAccess="true"))
-	EAvaRadialArrangePlane Plane = EAvaRadialArrangePlane::YZ;
-
 	/** The number of child elements to limit in the arrangement, or -1 if unlimited. Children whose index is greater than or equal to this value will be hidden. */
 	UPROPERTY(EditInstanceOnly, Setter, Getter, Category="RadialArrange", meta=(ClampMin="-1", UIMin="-1", AllowPrivateAccess="true"))
 	int32 Count = -1;
@@ -214,21 +217,25 @@ protected:
 
 	/** If true, will arrange the child elements starting from the outer radius and moving to the inner radius. Has no effect if only using one ring. */
 	UPROPERTY(EditInstanceOnly, Setter="SetStartFromOuterRadius", Getter="GetStartFromOuterRadius", Category="RadialArrange", meta=(AllowPrivateAccess="true"))
-	bool bStartFromOuterRadius;
+	bool bStartFromOuterRadius = false;
 
 	/** If true, will orient the selected axis torwards the center. */
 	UPROPERTY(EditInstanceOnly, Setter="SetOrient", Getter="GetOrient", Category="RadialArrange", meta=(AllowPrivateAccess="true"))
-	bool bOrient;
+	bool bOrient = false;
 
 	/** The axis to look at the center. */
 	UPROPERTY(EditInstanceOnly, Setter, Getter, Category="RadialArrange", meta=(EditCondition="bOrient", EditConditionHides, AllowPrivateAccess="true"))
 	EAvaModifiersAxis OrientationAxis = EAvaModifiersAxis::None;
 
+	/** Base rotation added on top of the orientation rotation computed */
+	UPROPERTY(EditInstanceOnly, Setter, Getter, Category="RadialArrange", meta=(EditCondition="bOrient", EditConditionHides, AllowPrivateAccess="true"))
+	FRotator BaseOrientation = FRotator::ZeroRotator;
+
 	UE_DEPRECATED(5.5, "Use OrientationAxis instead")
 	UPROPERTY(meta=(DeprecatedProperty, DeprecationMessage="Use OrientationAxis instead"))
 	EAvaAxis OrientAxis;
 
 	/** If true, will flip the orientation axis to the opposite direction. */
 	UPROPERTY(EditInstanceOnly, Setter="SetFlipOrient", Getter="GetFlipOrient", Category="RadialArrange", meta=(EditCondition="bOrient", EditConditionHides, AllowPrivateAccess="true"))
-	bool bFlipOrient;
+	bool bFlipOrient = false;
 };
