MotionDesign : Modifiers
- Improving actor transform save and restore system for modifiers to only save and restore what has changed

diff --git a/Engine/Plugins/Experimental/Avalanche/Source/AvalancheModifiers/Private/Modifiers/AvaRadialArrangeModifier.cpp b/Engine/Plugins/Experimental/Avalanche/Source/AvalancheModifiers/Private/Modifiers/AvaRadialArrangeModifier.cpp
--- a/Engine/Plugins/Experimental/Avalanche/Source/AvalancheModifiers/Private/Modifiers/AvaRadialArrangeModifier.cpp
+++ b/Engine/Plugins/Experimental/Avalanche/Source/AvalancheModifiers/Private/Modifiers/AvaRadialArrangeModifier.cpp
@@ -343,7 +343,7 @@ void UAvaRadialArrangeModifier::Apply()
 		}
 
 		// Track this actor layout state
-		LayoutShared->SaveActorState(this, AttachedActor);
+		LayoutShared->SaveActorState(this, AttachedActor, EAvaTransformSharedModifierState::LocationRotation);
 
 		const FVector RelativeOffset = CalculateRelativeOffset(SlotAngle, RingIndex);
 		AttachedActor->SetActorRelativeLocation(RelativeOffset);
@@ -358,9 +358,9 @@ void UAvaRadialArrangeModifier::Apply()
 			AttachedActor->SetActorRelativeRotation(NewRotation);
 		}
 		// Restore original rotation
-		else if (FAvaTransformSharedActorState* ActorState = LayoutShared->FindActorState(AttachedActor))
+		else
 		{
-			AttachedActor->SetActorRotation(ActorState->ActorTransform.GetRotation());
+			LayoutShared->RestoreActorState(this, AttachedActor, EAvaTransformSharedModifierState::Rotation);
 		}
 	}
 
diff --git a/Engine/Plugins/Experimental/Avalanche/Source/AvalancheModifiers/Private/Shared/AvaTransformModifierShared.cpp b/Engine/Plugins/Experimental/Avalanche/Source/AvalancheModifiers/Private/Shared/AvaTransformModifierShared.cpp
--- a/Engine/Plugins/Experimental/Avalanche/Source/AvalancheModifiers/Private/Shared/AvaTransformModifierShared.cpp
+++ b/Engine/Plugins/Experimental/Avalanche/Source/AvalancheModifiers/Private/Shared/AvaTransformModifierShared.cpp
@@ -5,18 +5,37 @@
 #include "Modifiers/AvaBaseModifier.h"
 #include "GameFramework/Actor.h"
 
-void FAvaTransformSharedModifierState::Save(const AActor* InActor)
+void FAvaTransformSharedModifierState::Save(const AActor* InActor, EAvaTransformSharedModifierState InSaveState)
 {
 	if (const UAvaBaseModifier* Modifier = ModifierWeak.Get())
 	{
 		if (InActor)
 		{
-			ActorTransform = InActor->GetActorTransform();
+			if (EnumHasAnyFlags(InSaveState, EAvaTransformSharedModifierState::Location)
+				&& !EnumHasAnyFlags(SaveState, EAvaTransformSharedModifierState::Location))
+			{
+				ActorTransform.SetLocation(InActor->GetActorLocation());
+				EnumAddFlags(SaveState, EAvaTransformSharedModifierState::Location);
+			}
+
+			if (EnumHasAnyFlags(InSaveState, EAvaTransformSharedModifierState::Rotation)
+				&& !EnumHasAnyFlags(SaveState, EAvaTransformSharedModifierState::Rotation))
+			{
+				ActorTransform.SetRotation(InActor->GetActorRotation().Quaternion());
+				EnumAddFlags(SaveState, EAvaTransformSharedModifierState::Rotation);
+			}
+
+			if (EnumHasAnyFlags(InSaveState, EAvaTransformSharedModifierState::Scale)
+				&& !EnumHasAnyFlags(SaveState, EAvaTransformSharedModifierState::Scale))
+			{
+				ActorTransform.SetScale3D(InActor->GetActorScale3D());
+				EnumAddFlags(SaveState, EAvaTransformSharedModifierState::Scale);
+			}
 		}
 	}
 }
 
-void FAvaTransformSharedModifierState::Restore(AActor* InActor, EAvaTransformSharedModifier InRestoreState) const
+void FAvaTransformSharedModifierState::Restore(AActor* InActor, EAvaTransformSharedModifierState InRestoreState)
 {
 	if (const UAvaBaseModifier* Modifier = ModifierWeak.Get())
 	{
@@ -24,18 +43,21 @@ void FAvaTransformSharedModifierState::Restore(AActor* InActor, EAvaTransformSha
 		{
 			FTransform RestoreTransform = ActorTransform;
 			const FTransform& CurrentActorTransform = InActor->GetActorTransform();
-			
-			if (!EnumHasAnyFlags(InRestoreState, EAvaTransformSharedModifier::Location))
+
+			if (!EnumHasAnyFlags(InRestoreState, EAvaTransformSharedModifierState::Location)
+				|| !EnumHasAnyFlags(SaveState, EAvaTransformSharedModifierState::Location))
 			{
 				RestoreTransform.SetLocation(CurrentActorTransform.GetLocation());
 			}
 
-			if (!EnumHasAnyFlags(InRestoreState, EAvaTransformSharedModifier::Rotation))
+			if (!EnumHasAnyFlags(InRestoreState, EAvaTransformSharedModifierState::Rotation)
+				|| !EnumHasAnyFlags(SaveState, EAvaTransformSharedModifierState::Rotation))
 			{
 				RestoreTransform.SetRotation(CurrentActorTransform.GetRotation());
 			}
 
-			if (!EnumHasAnyFlags(InRestoreState, EAvaTransformSharedModifier::Scale))
+			if (!EnumHasAnyFlags(InRestoreState, EAvaTransformSharedModifierState::Scale)
+				|| !EnumHasAnyFlags(SaveState, EAvaTransformSharedModifierState::Scale))
 			{
 				RestoreTransform.SetScale3D(CurrentActorTransform.GetScale3D());
 			}
@@ -44,36 +66,60 @@ void FAvaTransformSharedModifierState::Restore(AActor* InActor, EAvaTransformSha
 			{
 				InActor->SetActorTransform(RestoreTransform);
 			}
+
+			EnumRemoveFlags(SaveState, InRestoreState);
 		}
 	}
 }
 
-void FAvaTransformSharedActorState::Save()
+void FAvaTransformSharedActorState::Save(EAvaTransformSharedModifierState InSaveState)
 {
 	if (const AActor* Actor = ActorWeak.Get())
 	{
-		ActorTransform = Actor->GetActorTransform();
+		if (EnumHasAnyFlags(InSaveState, EAvaTransformSharedModifierState::Location)
+			&& !EnumHasAnyFlags(SaveState, EAvaTransformSharedModifierState::Location))
+		{
+			ActorTransform.SetLocation(Actor->GetActorLocation());
+			EnumAddFlags(SaveState, EAvaTransformSharedModifierState::Location);
+		}
+
+		if (EnumHasAnyFlags(InSaveState, EAvaTransformSharedModifierState::Rotation)
+			&& !EnumHasAnyFlags(SaveState, EAvaTransformSharedModifierState::Rotation))
+		{
+			ActorTransform.SetRotation(Actor->GetActorRotation().Quaternion());
+			EnumAddFlags(SaveState, EAvaTransformSharedModifierState::Rotation);
+		}
+
+		if (EnumHasAnyFlags(InSaveState, EAvaTransformSharedModifierState::Scale)
+			&& !EnumHasAnyFlags(SaveState, EAvaTransformSharedModifierState::Scale))
+		{
+			ActorTransform.SetScale3D(Actor->GetActorScale3D());
+			EnumAddFlags(SaveState, EAvaTransformSharedModifierState::Scale);
+		}
 	}
 }
 
-void FAvaTransformSharedActorState::Restore(EAvaTransformSharedModifier InRestoreState) const
+void FAvaTransformSharedActorState::Restore(EAvaTransformSharedModifierState InRestoreState)
 {
 	if (AActor* Actor = ActorWeak.Get())
 	{
 		FTransform RestoreTransform = ActorTransform;
 		const FTransform& CurrentActorTransform = Actor->GetActorTransform();
-			
-		if (!EnumHasAnyFlags(InRestoreState, EAvaTransformSharedModifier::Location))
+
+		if (!EnumHasAnyFlags(InRestoreState, EAvaTransformSharedModifierState::Location)
+			|| !EnumHasAnyFlags(SaveState, EAvaTransformSharedModifierState::Location))
 		{
 			RestoreTransform.SetLocation(CurrentActorTransform.GetLocation());
 		}
 
-		if (!EnumHasAnyFlags(InRestoreState, EAvaTransformSharedModifier::Rotation))
+		if (!EnumHasAnyFlags(InRestoreState, EAvaTransformSharedModifierState::Rotation)
+			|| !EnumHasAnyFlags(SaveState, EAvaTransformSharedModifierState::Rotation))
 		{
 			RestoreTransform.SetRotation(CurrentActorTransform.GetRotation());
 		}
 
-		if (!EnumHasAnyFlags(InRestoreState, EAvaTransformSharedModifier::Scale))
+		if (!EnumHasAnyFlags(InRestoreState, EAvaTransformSharedModifierState::Scale)
+			|| !EnumHasAnyFlags(SaveState, EAvaTransformSharedModifierState::Scale))
 		{
 			RestoreTransform.SetScale3D(CurrentActorTransform.GetScale3D());
 		}
@@ -82,10 +128,12 @@ void FAvaTransformSharedActorState::Restore(EAvaTransformSharedModifier InRestor
 		{
 			Actor->SetActorTransform(RestoreTransform);
 		}
+
+		EnumRemoveFlags(SaveState, InRestoreState);
 	}
 }
 
-void UAvaTransformModifierShared::SaveActorState(UAvaBaseModifier* InModifierContext, AActor* InActor)
+void UAvaTransformModifierShared::SaveActorState(UAvaBaseModifier* InModifierContext, AActor* InActor, EAvaTransformSharedModifierState InSaveState)
 {
 	if (!IsValid(InActor))
 	{
@@ -93,23 +141,13 @@ void UAvaTransformModifierShared::SaveActorState(UAvaBaseModifier* InModifierCon
 	}
 
 	FAvaTransformSharedActorState& ActorState = ActorStates.FindOrAdd(FAvaTransformSharedActorState(InActor));
+	ActorState.Save(InSaveState);
 
-	if (ActorState.ModifierStates.IsEmpty())
-	{
-		ActorState.Save();
-	}
-
-	bool bAlreadyInSet = false;
-	FAvaTransformSharedModifierState& ModifierState = ActorState.ModifierStates.FindOrAdd(FAvaTransformSharedModifierState(InModifierContext), &bAlreadyInSet);
-
-	if (!bAlreadyInSet)
-	{
-		ModifierState.Save(InActor);
-		ActorState.ModifierStates.Add(ModifierState);
-	}
+	FAvaTransformSharedModifierState& ModifierState = ActorState.ModifierStates.FindOrAdd(FAvaTransformSharedModifierState(InModifierContext));
+	ModifierState.Save(InActor, InSaveState);
 }
 
-void UAvaTransformModifierShared::RestoreActorState(UAvaBaseModifier* InModifierContext, AActor* InActor, EAvaTransformSharedModifier InRestoreTransform)
+void UAvaTransformModifierShared::RestoreActorState(UAvaBaseModifier* InModifierContext, AActor* InActor, EAvaTransformSharedModifierState InRestoreState)
 {
 	if (!IsValid(InActor))
 	{
@@ -122,20 +160,24 @@ void UAvaTransformModifierShared::RestoreActorState(UAvaBaseModifier* InModifier
 		return;
 	}
 
-	const FAvaTransformSharedModifierState* ActorModifierState = ActorState->ModifierStates.Find(FAvaTransformSharedModifierState(InModifierContext));
+	FAvaTransformSharedModifierState* ActorModifierState = ActorState->ModifierStates.Find(FAvaTransformSharedModifierState(InModifierContext));
 	if (!ActorModifierState)
 	{
 		return;
 	}
 
 	// restore modifier state and remove it
-	ActorModifierState->Restore(InActor, InRestoreTransform);
-	ActorState->ModifierStates.Remove(*ActorModifierState);
+	ActorModifierState->Restore(InActor, InRestoreState);
+
+	if (ActorModifierState->SaveState == EAvaTransformSharedModifierState::None)
+	{
+		ActorState->ModifierStates.Remove(*ActorModifierState);
+	}
 
 	// Restore original actor state and remove it
 	if (ActorState->ModifierStates.IsEmpty())
 	{
-		ActorState->Restore(InRestoreTransform);
+		ActorState->Restore(EAvaTransformSharedModifierState::All);
 		ActorStates.Remove(*ActorState);
 	}
 }
@@ -153,7 +195,7 @@ FAvaTransformSharedActorState* UAvaTransformModifierShared::FindActorState(AActo
 TSet<FAvaTransformSharedActorState*> UAvaTransformModifierShared::FindActorsState(UAvaBaseModifier* InModifierContext)
 {
 	TSet<FAvaTransformSharedActorState*> ModifierActorStates;
-	
+
 	for (FAvaTransformSharedActorState& ActorState : ActorStates)
 	{
 		if (ActorState.ModifierStates.Contains(FAvaTransformSharedModifierState(InModifierContext)))
@@ -161,16 +203,16 @@ TSet<FAvaTransformSharedActorState*> UAvaTransformModifierShared::FindActorsStat
 			ModifierActorStates.Add(&ActorState);
 		}
 	}
-	
+
 	return ModifierActorStates;
 }
 
-void UAvaTransformModifierShared::RestoreActorsState(UAvaBaseModifier* InModifierContext, const TSet<AActor*>* InActors, EAvaTransformSharedModifier InRestoreTransform)
+void UAvaTransformModifierShared::RestoreActorsState(UAvaBaseModifier* InModifierContext, const TSet<AActor*>* InActors, EAvaTransformSharedModifierState InRestoreState)
 {
 	const FAvaTransformSharedModifierState SearchModifierState(InModifierContext);
 	TSet<AActor*> LinkedModifierActors;
 	TSet<UActorModifierCoreBase*> LinkedActorModifiers;
-	
+
 	for (const FAvaTransformSharedActorState& ActorState : ActorStates)
 	{
 		AActor* Actor = ActorState.ActorWeak.Get();
@@ -188,7 +230,7 @@ void UAvaTransformModifierShared::RestoreActorsState(UAvaBaseModifier* InModifie
 		{
 			continue;
 		}
-		
+
 		// Collect actors affected by modifier
 		LinkedModifierActors.Add(Actor);
 
@@ -209,16 +251,16 @@ void UAvaTransformModifierShared::RestoreActorsState(UAvaBaseModifier* InModifie
 	// Restore actor state
 	for (AActor* Actor : LinkedModifierActors)
 	{
-		RestoreActorState(InModifierContext, Actor, InRestoreTransform);
+		RestoreActorState(InModifierContext, Actor, InRestoreState);
 	}
 }
 
-void UAvaTransformModifierShared::RestoreActorsState(UAvaBaseModifier* InModifierContext, const TSet<TWeakObjectPtr<AActor>>& InActors, EAvaTransformSharedModifier InRestoreTransform)
+void UAvaTransformModifierShared::RestoreActorsState(UAvaBaseModifier* InModifierContext, const TSet<TWeakObjectPtr<AActor>>& InActors, EAvaTransformSharedModifierState InRestoreState)
 {
 	TSet<AActor*> Actors;
 	Algo::Transform(InActors, Actors, [](const TWeakObjectPtr<AActor>& InActor)->AActor*{ return InActor.Get(); });
 
-	RestoreActorsState(InModifierContext, &Actors, InRestoreTransform);
+	RestoreActorsState(InModifierContext, &Actors, InRestoreState);
 }
 
 bool UAvaTransformModifierShared::IsActorStateSaved(UAvaBaseModifier* InModifierContext, AActor* InActor)
@@ -227,14 +269,14 @@ bool UAvaTransformModifierShared::IsActorStateSaved(UAvaBaseModifier* InModifier
 	{
 		return ActorState->ModifierStates.Contains(FAvaTransformSharedModifierState(InModifierContext));
 	}
-	
+
 	return false;
 }
 
 bool UAvaTransformModifierShared::IsActorsStateSaved(UAvaBaseModifier* InModifierContext)
 {
 	const FAvaTransformSharedModifierState ModifierState(InModifierContext);
-	
+
 	for (const FAvaTransformSharedActorState& ActorState : ActorStates)
 	{
 		if (ActorState.ModifierStates.Contains(ModifierState))
@@ -242,7 +284,7 @@ bool UAvaTransformModifierShared::IsActorsStateSaved(UAvaBaseModifier* InModifie
 			return true;
 		}
 	}
-	
+
 	return false;
 }
 
diff --git a/Engine/Plugins/Experimental/Avalanche/Source/AvalancheModifiers/Public/Shared/AvaTransformModifierShared.h b/Engine/Plugins/Experimental/Avalanche/Source/AvalancheModifiers/Public/Shared/AvaTransformModifierShared.h
--- a/Engine/Plugins/Experimental/Avalanche/Source/AvalancheModifiers/Public/Shared/AvaTransformModifierShared.h
+++ b/Engine/Plugins/Experimental/Avalanche/Source/AvalancheModifiers/Public/Shared/AvaTransformModifierShared.h
@@ -11,8 +11,9 @@
 class AActor;
 class UAvaBaseModifier;
 
+/** Enumerates all transform state that can be saved */
 UENUM()
-enum class EAvaTransformSharedModifier : uint8
+enum class EAvaTransformSharedModifierState : uint8
 {
 	None,
 	Location = 1 << 0,
@@ -35,10 +36,10 @@ struct FAvaTransformSharedModifierState
 	{}
 
 	/** Save this modifier state if valid */
-	void Save(const AActor* InActor);
+	void Save(const AActor* InActor, EAvaTransformSharedModifierState InSaveState);
 
 	/** Restore this modifier state if valid */
-	void Restore(AActor* InActor, EAvaTransformSharedModifier InRestoreState) const;
+	void Restore(AActor* InActor, EAvaTransformSharedModifierState InRestoreState);
 
 	friend uint32 GetTypeHash(const FAvaTransformSharedModifierState& InItem)
 	{
@@ -50,11 +51,17 @@ struct FAvaTransformSharedModifierState
 		return ModifierWeak == Other.ModifierWeak;
 	}
 
+	/** Modifier applying the transform change */
 	UPROPERTY()
 	TWeakObjectPtr<UAvaBaseModifier> ModifierWeak;
 
+	/** Pre modifier transform saved */
 	UPROPERTY()
 	FTransform ActorTransform;
+
+	/** Used to restore only what has changed */
+	UPROPERTY()
+	EAvaTransformSharedModifierState SaveState = EAvaTransformSharedModifierState::All;
 };
 
 USTRUCT()
@@ -68,10 +75,10 @@ struct FAvaTransformSharedActorState
 	{}
 
 	/** Save this actor state if valid */
-	void Save();
+	void Save(EAvaTransformSharedModifierState InSaveState);
 
 	/** Restore this actor state if valid */
-	void Restore(EAvaTransformSharedModifier InRestoreState) const;
+	void Restore(EAvaTransformSharedModifierState InRestoreState);
 
 	friend uint32 GetTypeHash(const FAvaTransformSharedActorState& InItem)
 	{
@@ -94,6 +101,10 @@ struct FAvaTransformSharedActorState
 	/** Pre state transform saved */
 	UPROPERTY()
 	FTransform ActorTransform;
+
+	/** Used to restore only what has changed */
+	UPROPERTY()
+	EAvaTransformSharedModifierState SaveState = EAvaTransformSharedModifierState::All;
 };
 
 /**
@@ -109,10 +120,10 @@ class UAvaTransformModifierShared : public UActorModifierCoreSharedObject
 
 public:
 	/** Save actor state, adds it if it is not tracked */
-	void SaveActorState(UAvaBaseModifier* InModifierContext, AActor* InActor);
+	void SaveActorState(UAvaBaseModifier* InModifierContext, AActor* InActor, EAvaTransformSharedModifierState InSaveState = EAvaTransformSharedModifierState::All);
 
 	/** Restore actor state, removes it if no other modifier track that actor state */
-	void RestoreActorState(UAvaBaseModifier* InModifierContext, AActor* InActor, EAvaTransformSharedModifier InRestoreTransform = EAvaTransformSharedModifier::All);
+	void RestoreActorState(UAvaBaseModifier* InModifierContext, AActor* InActor, EAvaTransformSharedModifierState InRestoreState = EAvaTransformSharedModifierState::All);
 
 	/** Get the actor state for a specific actor */
 	FAvaTransformSharedActorState* FindActorState(AActor* InActor);
@@ -121,10 +132,10 @@ public:
 	TSet<FAvaTransformSharedActorState*> FindActorsState(UAvaBaseModifier* InModifierContext);
 
 	/** Restore all actors states linked to this modifier */
-	void RestoreActorsState(UAvaBaseModifier* InModifierContext, const TSet<AActor*>* InActors = nullptr, EAvaTransformSharedModifier InRestoreTransform = EAvaTransformSharedModifier::All);
+	void RestoreActorsState(UAvaBaseModifier* InModifierContext, const TSet<AActor*>* InActors = nullptr, EAvaTransformSharedModifierState InRestoreState = EAvaTransformSharedModifierState::All);
 
 	/** Restore all specified actors linked to this modifier */
-	void RestoreActorsState(UAvaBaseModifier* InModifierContext, const TSet<TWeakObjectPtr<AActor>>& InActors, EAvaTransformSharedModifier InRestoreTransform = EAvaTransformSharedModifier::All);
+	void RestoreActorsState(UAvaBaseModifier* InModifierContext, const TSet<TWeakObjectPtr<AActor>>& InActors, EAvaTransformSharedModifierState InRestoreState = EAvaTransformSharedModifierState::All);
 
 	/** Returns true, if this modifier is tracking this actor */
 	bool IsActorStateSaved(UAvaBaseModifier* InModifierContext, AActor* InActor);
