Typo fix.
Deprecation of RemoveOnActorDestroyededHandler in favor of RemoveOnActorDestroyedHandler

diff --git a/Engine/Plugins/ChaosVD/Source/ChaosVD/Private/ChaosVDScene.cpp b/Engine/Plugins/ChaosVD/Source/ChaosVD/Private/ChaosVDScene.cpp
--- a/Engine/Plugins/ChaosVD/Source/ChaosVD/Private/ChaosVDScene.cpp
+++ b/Engine/Plugins/ChaosVD/Source/ChaosVD/Private/ChaosVDScene.cpp
@@ -119,7 +119,7 @@ void FChaosVDScene::DeInitialize()
 
 	if (PhysicsVDWorld)
 	{
-		PhysicsVDWorld->RemoveOnActorDestroyededHandler(ActorDestroyedHandle);
+		PhysicsVDWorld->RemoveOnActorDestroyedHandler(ActorDestroyedHandle);
 
 		PhysicsVDWorld->DestroyWorld(true);
 		GEngine->DestroyWorldContext(PhysicsVDWorld);
diff --git a/Engine/Plugins/Experimental/Avalanche/Source/AvalancheModifiers/Private/Extensions/AvaSceneTreeUpdateModifierExtension.cpp b/Engine/Plugins/Experimental/Avalanche/Source/AvalancheModifiers/Private/Extensions/AvaSceneTreeUpdateModifierExtension.cpp
--- a/Engine/Plugins/Experimental/Avalanche/Source/AvalancheModifiers/Private/Extensions/AvaSceneTreeUpdateModifierExtension.cpp
+++ b/Engine/Plugins/Experimental/Avalanche/Source/AvalancheModifiers/Private/Extensions/AvaSceneTreeUpdateModifierExtension.cpp
@@ -79,7 +79,7 @@ void FAvaSceneTreeUpdateModifierExtension::OnExtensionEnabled(EActorModifierCore
 	}
 
 	// When actor are destroyed in world
-	World->RemoveOnActorDestroyededHandler(WorldActorDestroyedDelegate);
+	World->RemoveOnActorDestroyedHandler(WorldActorDestroyedDelegate);
 	WorldActorDestroyedDelegate = World->AddOnActorDestroyedHandler(FOnActorDestroyed::FDelegate::CreateSP(this, &FAvaSceneTreeUpdateModifierExtension::OnWorldActorDestroyed));
 
 	// Used to detect visibility changes in siblings
@@ -116,7 +116,7 @@ void FAvaSceneTreeUpdateModifierExtension::OnExtensionDisabled(EActorModifierCor
 		return;
 	}
 
-	World->RemoveOnActorDestroyededHandler(WorldActorDestroyedDelegate);
+	World->RemoveOnActorDestroyedHandler(WorldActorDestroyedDelegate);
 
 #if WITH_EDITOR
 	if (UAvaOutlinerSubsystem* const OutlinerSubsystem = World->GetSubsystem<UAvaOutlinerSubsystem>())
diff --git a/Engine/Plugins/Experimental/EditorDataStorage/Source/TedsCore/Private/TypedElementDatabaseCompatibility.cpp b/Engine/Plugins/Experimental/EditorDataStorage/Source/TedsCore/Private/TypedElementDatabaseCompatibility.cpp
--- a/Engine/Plugins/Experimental/EditorDataStorage/Source/TedsCore/Private/TypedElementDatabaseCompatibility.cpp
+++ b/Engine/Plugins/Experimental/EditorDataStorage/Source/TedsCore/Private/TypedElementDatabaseCompatibility.cpp
@@ -99,7 +99,7 @@ void UEditorDataStorageCompatibility::Deinitialize()
 	
 	for (TPair<UWorld*, FDelegateHandle>& It : ActorDestroyedDelegateHandles)
 	{
-		It.Key->RemoveOnActorDestroyededHandler(It.Value);
+		It.Key->RemoveOnActorDestroyedHandler(It.Value);
 	}
 
 	FWorldDelegates::OnPreWorldFinishDestroy.Remove(PreWorldFinishDestroyDelegateHandle);
@@ -1183,7 +1183,7 @@ void UEditorDataStorageCompatibility::OnPreWorldFinishDestroy(UWorld* World)
 	FDelegateHandle Handle;
 	if (ActorDestroyedDelegateHandles.RemoveAndCopyValue(World, Handle))
 	{
-		World->RemoveOnActorDestroyededHandler(Handle);
+		World->RemoveOnActorDestroyedHandler(Handle);
 	}
 }
 
diff --git a/Engine/Source/Runtime/Engine/Classes/Engine/World.h b/Engine/Source/Runtime/Engine/Classes/Engine/World.h
--- a/Engine/Source/Runtime/Engine/Classes/Engine/World.h
+++ b/Engine/Source/Runtime/Engine/Classes/Engine/World.h
@@ -2716,7 +2716,9 @@ public:
 	FDelegateHandle AddOnActorDestroyedHandler(const FOnActorDestroyed::FDelegate& InHandler) const;
 
 	/** Remove a listener for OnActorDestroyed events */
+	UE_DEPRECATED(5.5, "Use the typo corrected version RemoveOnActorDestroyedHandler")
 	void RemoveOnActorDestroyededHandler(FDelegateHandle InHandle) const;
+	void RemoveOnActorDestroyedHandler(FDelegateHandle InHandle) const;
 
 	/** Add a listener for OnPostRegisterAllActorComponents events */
 	FDelegateHandle AddOnPostRegisterAllActorComponentsHandler(const FOnPostRegisterAllActorComponents::FDelegate& InHandler) const;
diff --git a/Engine/Source/Runtime/Engine/Private/DemoNetDriver.cpp b/Engine/Source/Runtime/Engine/Private/DemoNetDriver.cpp
--- a/Engine/Source/Runtime/Engine/Private/DemoNetDriver.cpp
+++ b/Engine/Source/Runtime/Engine/Private/DemoNetDriver.cpp
@@ -925,7 +925,7 @@ void UDemoNetDriver::FinishDestroy()
 	
 	if (World)
 	{
-		World->RemoveOnActorDestroyededHandler(DelegateHandleActorPreDestroy);
+		World->RemoveOnActorDestroyedHandler(DelegateHandleActorPreDestroy);
 	}
 
 	Super::FinishDestroy();
@@ -5659,7 +5659,7 @@ void UDemoNetDriver::SetWorld(UWorld* InWorld)
 {
 	if (World)
 	{
-		World->RemoveOnActorDestroyededHandler(DelegateHandleActorPreDestroy);
+		World->RemoveOnActorDestroyedHandler(DelegateHandleActorPreDestroy);
 	}
 
 	Super::SetWorld(InWorld);
diff --git a/Engine/Source/Runtime/Engine/Private/World.cpp b/Engine/Source/Runtime/Engine/Private/World.cpp
--- a/Engine/Source/Runtime/Engine/Private/World.cpp
+++ b/Engine/Source/Runtime/Engine/Private/World.cpp
@@ -5903,6 +5903,11 @@ FDelegateHandle UWorld::AddOnActorDestroyedHandler(const FOnActorDestroyed::FDel
 }
 
 void UWorld::RemoveOnActorDestroyededHandler(FDelegateHandle InHandle) const
+{
+	RemoveOnActorDestroyedHandler(InHandle);
+}
+
+void UWorld::RemoveOnActorDestroyedHandler(FDelegateHandle InHandle) const
 {
 	OnActorDestroyed.Remove(InHandle);
 }
