Motion Design; Added Transition Attribute API to control State Task

#virtualized

diff --git a/Engine/Plugins/Experimental/Avalanche/Source/Avalanche/Private/AvaSceneState.cpp b/Engine/Plugins/Experimental/Avalanche/Source/Avalanche/Private/AvaSceneState.cpp
--- a/Engine/Plugins/Experimental/Avalanche/Source/Avalanche/Private/AvaSceneState.cpp
+++ b/Engine/Plugins/Experimental/Avalanche/Source/Avalanche/Private/AvaSceneState.cpp
@@ -1,6 +1,8 @@
 // Copyright Epic Games, Inc. All Rights Reserved.
 
 #include "AvaSceneState.h"
+
+#include "AvaNameAttribute.h"
 #include "AvaSceneSettings.h"
 #include "AvaTagHandle.h"
 #include "AvaTagHandleContainer.h"
@@ -69,3 +71,65 @@ bool UAvaSceneState::ContainsTagAttribute(const FAvaTagHandle& InTagHandle) cons
 			return TagAttribute && TagAttribute->ContainsTag(InTagHandle);
 		});
 }
+
+bool UAvaSceneState::AddNameAttribute(FName InName)
+{
+	if (InName.IsNone())
+	{
+		return false;
+	}
+
+	// Return true if already existing
+	if (ContainsNameAttribute(InName))
+	{
+		return true;
+	}
+
+	UAvaNameAttribute* NameAttribute = NewObject<UAvaNameAttribute>(this, NAME_None, RF_Transient);
+	check(NameAttribute);
+	NameAttribute->Name = InName;
+
+	SceneAttributes.Add(NameAttribute);
+	return true;
+}
+
+bool UAvaSceneState::RemoveNameAttribute(FName InName)
+{
+	uint32 NamesCleared = 0;
+
+	for (TArray<TObjectPtr<UAvaAttribute>>::TIterator Iter(SceneAttributes); Iter; ++Iter)
+	{
+		UAvaNameAttribute* NameAttribute = Cast<UAvaNameAttribute>(*Iter);
+		if (!NameAttribute || NameAttribute->Name != InName)
+		{
+			continue;
+		}
+
+		// Attempt to clear the given name for the attribute.
+		// If this owns the attribute, remove from the Scene Attributes list
+		if (NameAttribute->GetOuter() == this)
+		{
+			Iter.RemoveCurrent();
+		}
+		// Do not remove the attribute itself from the list if it's an external attribute as it could still have valid tags,
+		// or later have valid tags (due to a dynamic change)
+		else
+		{
+			NameAttribute->Name = NAME_None;
+		}
+
+		++NamesCleared;
+	}
+
+	return NamesCleared > 0;
+}
+
+bool UAvaSceneState::ContainsNameAttribute(FName InName) const
+{
+	return SceneAttributes.ContainsByPredicate(
+		[InName](UAvaAttribute* InAttribute)
+		{
+			UAvaNameAttribute* NameAttribute = Cast<UAvaNameAttribute>(InAttribute);
+			return NameAttribute && NameAttribute->Name == InName;
+		});
+}
diff --git a/Engine/Plugins/Experimental/Avalanche/Source/Avalanche/Private/Transition/AvaTransitionAttributeLibrary.cpp b/Engine/Plugins/Experimental/Avalanche/Source/Avalanche/Private/Transition/AvaTransitionAttributeLibrary.cpp
new file mode 100644
--- /dev/null
+++ b/Engine/Plugins/Experimental/Avalanche/Source/Avalanche/Private/Transition/AvaTransitionAttributeLibrary.cpp
@@ -0,0 +1,97 @@
+// Copyright Epic Games, Inc. All Rights Reserved.
+
+#include "AvaTransitionAttributeLibrary.h"
+#include "AvaSceneState.h"
+#include "AvaSceneSubsystem.h"
+#include "AvaTransitionContext.h"
+#include "AvaTransitionLayer.h"
+#include "AvaTransitionLayerUtils.h"
+#include "AvaTransitionSubsystem.h"
+#include "Behavior/AvaTransitionBehaviorInstance.h"
+#include "Behavior/AvaTransitionBehaviorInstanceCache.h"
+#include "IAvaSceneInterface.h"
+#include "IAvaTransitionNodeInterface.h"
+
+namespace UE::Ava::Private
+{
+	UAvaSceneState* GetSceneState(UObject* InTransitionNode)
+	{
+		const IAvaTransitionNodeInterface* NodeInterface = Cast<IAvaTransitionNodeInterface>(InTransitionNode);
+		if (!NodeInterface)
+		{
+			return nullptr;
+		}
+
+		const FAvaTransitionContext* TransitionContext = NodeInterface->GetBehaviorInstanceCache().GetTransitionContext();
+		if (!TransitionContext)
+		{
+			return nullptr;
+		}
+
+		const FAvaTransitionScene* TransitionScene = TransitionContext->GetTransitionScene();
+		if (!TransitionScene)
+		{
+			return nullptr;
+		}
+
+		if (IAvaSceneInterface* SceneInterface = UAvaSceneSubsystem::FindSceneInterface(TransitionScene->GetLevel()))
+		{
+			return SceneInterface->GetSceneState();
+		}
+		return nullptr;
+	}
+}
+
+bool UAvaTransitionAttributeLibrary::AddTagAttribute(UObject* InTransitionNode, const FAvaTagHandle& InTagHandle)
+{
+	if (UAvaSceneState* SceneState = UE::Ava::Private::GetSceneState(InTransitionNode))
+	{
+		return SceneState->AddTagAttribute(InTagHandle);
+	}
+	return false;
+}
+
+bool UAvaTransitionAttributeLibrary::RemoveTagAttribute(UObject* InTransitionNode, const FAvaTagHandle& InTagHandle)
+{
+	if (UAvaSceneState* SceneState = UE::Ava::Private::GetSceneState(InTransitionNode))
+	{
+		return SceneState->RemoveTagAttribute(InTagHandle);
+	}
+	return false;
+}
+
+bool UAvaTransitionAttributeLibrary::ContainsTagAttribute(UObject* InTransitionNode, const FAvaTagHandle& InTagHandle)
+{
+	if (UAvaSceneState* SceneState = UE::Ava::Private::GetSceneState(InTransitionNode))
+	{
+		return SceneState->ContainsTagAttribute(InTagHandle);
+	}
+	return false;
+}
+
+bool UAvaTransitionAttributeLibrary::AddNameAttribute(UObject* InTransitionNode, FName InName)
+{
+	if (UAvaSceneState* SceneState = UE::Ava::Private::GetSceneState(InTransitionNode))
+	{
+		return SceneState->AddNameAttribute(InName);
+	}
+	return false;
+}
+
+bool UAvaTransitionAttributeLibrary::RemoveNameAttribute(UObject* InTransitionNode, FName InName)
+{
+	if (UAvaSceneState* SceneState = UE::Ava::Private::GetSceneState(InTransitionNode))
+	{
+		return SceneState->RemoveNameAttribute(InName);
+	}
+	return false;
+}
+
+bool UAvaTransitionAttributeLibrary::ContainsNameAttribute(UObject* InTransitionNode, FName InName)
+{
+	if (UAvaSceneState* SceneState = UE::Ava::Private::GetSceneState(InTransitionNode))
+	{
+		return SceneState->ContainsNameAttribute(InName);
+	}
+	return false;
+}
diff --git a/Engine/Plugins/Experimental/Avalanche/Source/Avalanche/Private/Transition/AvaTransitionAttributeLibrary.h b/Engine/Plugins/Experimental/Avalanche/Source/Avalanche/Private/Transition/AvaTransitionAttributeLibrary.h
new file mode 100644
--- /dev/null
+++ b/Engine/Plugins/Experimental/Avalanche/Source/Avalanche/Private/Transition/AvaTransitionAttributeLibrary.h
@@ -0,0 +1,36 @@
+// Copyright Epic Games, Inc. All Rights Reserved.
+
+#pragma once
+
+#include "Kismet/BlueprintFunctionLibrary.h"
+#include "AvaTransitionAttributeLibrary.generated.h"
+
+enum class EAvaTransitionLayerCompareType : uint8;
+enum class EAvaTransitionSceneType : uint8;
+struct FAvaTagHandle;
+struct FAvaTagHandleContainer;
+
+UCLASS()
+class UAvaTransitionAttributeLibrary : public UBlueprintFunctionLibrary
+{
+	GENERATED_BODY()
+
+public:
+	UFUNCTION(BlueprintCallable, Category="Scene Attributes", meta=(DefaultToSelf="InTransitionNode"))
+	static bool AddTagAttribute(UObject* InTransitionNode, const FAvaTagHandle& InTagHandle);
+
+	UFUNCTION(BlueprintCallable, Category="Scene Attributes", meta=(DefaultToSelf="InTransitionNode"))
+	static bool RemoveTagAttribute(UObject* InTransitionNode, const FAvaTagHandle& InTagHandle);
+
+	UFUNCTION(BlueprintCallable, Category="Scene Attributes", meta=(DefaultToSelf="InTransitionNode"))
+	static bool ContainsTagAttribute(UObject* InTransitionNode, const FAvaTagHandle& InTagHandle);
+
+	UFUNCTION(BlueprintCallable, Category="Scene Attributes", meta=(DefaultToSelf="InTransitionNode"))
+	static bool AddNameAttribute(UObject* InTransitionNode, FName InName);
+
+	UFUNCTION(BlueprintCallable, Category="Scene Attributes", meta=(DefaultToSelf="InTransitionNode"))
+	static bool RemoveNameAttribute(UObject* InTransitionNode, FName InName);
+
+	UFUNCTION(BlueprintCallable, Category="Scene Attributes", meta=(DefaultToSelf="InTransitionNode"))
+	static bool ContainsNameAttribute(UObject* InTransitionNode, FName InName);
+};
diff --git a/Engine/Plugins/Experimental/Avalanche/Source/Avalanche/Public/AvaSceneState.h b/Engine/Plugins/Experimental/Avalanche/Source/Avalanche/Public/AvaSceneState.h
--- a/Engine/Plugins/Experimental/Avalanche/Source/Avalanche/Public/AvaSceneState.h
+++ b/Engine/Plugins/Experimental/Avalanche/Source/Avalanche/Public/AvaSceneState.h
@@ -12,7 +12,7 @@ class UAvaAttribute;
 class UAvaSceneSettings;
 struct FAvaTagHandle;
 
-/** Object providing State information of the Scene */
+/** Object providing attribute information of the Scene */
 UCLASS(MinimalAPI)
 class UAvaSceneState : public UObject
 {
@@ -22,13 +22,15 @@ public:
 	void Initialize(UAvaSceneSettings* InSceneSettings);
 
 	AVALANCHE_API bool AddTagAttribute(const FAvaTagHandle& InTagHandle);
-
 	AVALANCHE_API bool RemoveTagAttribute(const FAvaTagHandle& InTagHandle);
-
 	AVALANCHE_API bool ContainsTagAttribute(const FAvaTagHandle& InTagHandle) const;
 
+	AVALANCHE_API bool AddNameAttribute(FName InName);
+	AVALANCHE_API bool RemoveNameAttribute(FName InName);
+	AVALANCHE_API bool ContainsNameAttribute(FName InName) const;
+
 private:
-	/** In-play Scene Attributes. Starts as being the Scene Setting Attributes but can added to / removed from while in-play */
+	/** In-play Scene Attributes. Can be added to / removed from while in-play */
 	UPROPERTY(Transient)
 	TArray<TObjectPtr<UAvaAttribute>> SceneAttributes;
 };
diff --git a/Engine/Plugins/Experimental/Avalanche/Source/AvalancheAttribute/Private/AvaNameAttribute.cpp b/Engine/Plugins/Experimental/Avalanche/Source/AvalancheAttribute/Private/AvaNameAttribute.cpp
new file mode 100644
--- /dev/null
+++ b/Engine/Plugins/Experimental/Avalanche/Source/AvalancheAttribute/Private/AvaNameAttribute.cpp
@@ -0,0 +1,17 @@
+// Copyright Epic Games, Inc. All Rights Reserved.
+
+#include "AvaNameAttribute.h"
+
+#define LOCTEXT_NAMESPACE "AvaNameAttribute"
+
+FText UAvaNameAttribute::GetDisplayName() const
+{
+	return FText::Format(LOCTEXT("DisplayName", "Name: {0}"), FText::FromName(Name));
+}
+
+void UAvaNameAttribute::SetName(FName InName)
+{
+	Name = InName;
+}
+
+#undef LOCTEXT_NAMESPACE
diff --git a/Engine/Plugins/Experimental/Avalanche/Source/AvalancheAttribute/Public/AvaNameAttribute.h b/Engine/Plugins/Experimental/Avalanche/Source/AvalancheAttribute/Public/AvaNameAttribute.h
new file mode 100644
--- /dev/null
+++ b/Engine/Plugins/Experimental/Avalanche/Source/AvalancheAttribute/Public/AvaNameAttribute.h
@@ -0,0 +1,24 @@
+// Copyright Epic Games, Inc. All Rights Reserved.
+
+#pragma once
+
+#include "AvaAttribute.h"
+#include "AvaNameAttribute.generated.h"
+
+/** Attribute that holds a name */
+UCLASS(MinimalAPI, DisplayName="Name Attribute")
+class UAvaNameAttribute : public UAvaAttribute
+{
+	GENERATED_BODY()
+
+public:
+	//~ Begin UAvaAttribute
+	AVALANCHEATTRIBUTE_API virtual FText GetDisplayName() const override;
+	//~ End UAvaAttribute
+
+	UFUNCTION()
+	AVALANCHEATTRIBUTE_API void SetName(FName InName);
+
+	UPROPERTY(EditAnywhere, Setter, Category="Attributes")
+	FName Name;
+};
diff --git a/Engine/Plugins/Experimental/Avalanche/Source/AvalancheSequence/Public/AvaSequencePlaybackObject.h b/Engine/Plugins/Experimental/Avalanche/Source/AvalancheSequence/Public/AvaSequencePlaybackObject.h
--- a/Engine/Plugins/Experimental/Avalanche/Source/AvalancheSequence/Public/AvaSequencePlaybackObject.h
+++ b/Engine/Plugins/Experimental/Avalanche/Source/AvalancheSequence/Public/AvaSequencePlaybackObject.h
@@ -19,7 +19,7 @@ namespace UE::MovieScene
 	struct FOnCameraCutUpdatedParams;
 }
 
-UINTERFACE(MinimalAPI, meta = (CannotImplementInterfaceInBlueprint))
+UINTERFACE(MinimalAPI, BlueprintType, meta = (CannotImplementInterfaceInBlueprint))
 class UAvaSequencePlaybackObject : public UInterface
 {
 	GENERATED_BODY()
diff --git a/Engine/Plugins/Experimental/Avalanche/Source/AvalancheTransition/Private/AvaTransitionLibrary.cpp b/Engine/Plugins/Experimental/Avalanche/Source/AvalancheTransition/Private/AvaTransitionLibrary.cpp
new file mode 100644
--- /dev/null
+++ b/Engine/Plugins/Experimental/Avalanche/Source/AvalancheTransition/Private/AvaTransitionLibrary.cpp
@@ -0,0 +1,176 @@
+// Copyright Epic Games, Inc. All Rights Reserved.
+
+#include "AvaTransitionLibrary.h"
+#include "AvaTransitionContext.h"
+#include "AvaTransitionLayer.h"
+#include "AvaTransitionLayerUtils.h"
+#include "AvaTransitionSubsystem.h"
+#include "Behavior/AvaTransitionBehaviorInstance.h"
+#include "Behavior/AvaTransitionBehaviorInstanceCache.h"
+#include "Engine/Level.h"
+#include "Engine/World.h"
+#include "IAvaTransitionNodeInterface.h"
+#include "UObject/Package.h"
+
+bool UAvaTransitionLibrary::IsTransitionActiveInLayer(UObject* InTransitionNode
+	, EAvaTransitionComparisonResult InSceneComparisonType
+	, EAvaTransitionLayerCompareType InLayerComparisonType
+	, const FAvaTagHandleContainer& InSpecificLayers)
+{
+	const IAvaTransitionNodeInterface* NodeInterface = Cast<IAvaTransitionNodeInterface>(InTransitionNode);
+	if (!NodeInterface)
+	{
+		return false;
+	}
+
+	const FAvaTransitionContext* TransitionContext = NodeInterface->GetBehaviorInstanceCache().GetTransitionContext();
+	if (!TransitionContext)
+	{
+		return false;
+	}
+
+	const FAvaTransitionScene* TransitionScene = TransitionContext->GetTransitionScene();
+	if (!TransitionScene)
+	{
+		return false;
+	}
+
+	ULevel* TransitionLevel = TransitionScene->GetLevel();
+	if (!TransitionLevel || !TransitionLevel->OwningWorld)
+	{
+		return false;
+	}
+
+	UAvaTransitionSubsystem* TransitionSubsystem = TransitionLevel->OwningWorld->GetSubsystem<UAvaTransitionSubsystem>();
+	if (!TransitionSubsystem)
+	{
+		return false;
+	}
+
+	const FAvaTransitionLayerComparator Comparator = FAvaTransitionLayerUtils::BuildComparator(*TransitionContext
+		, InLayerComparisonType
+		, InSpecificLayers);
+
+	const TArray<const FAvaTransitionBehaviorInstance*> BehaviorInstances = FAvaTransitionLayerUtils::QueryBehaviorInstances(*TransitionSubsystem, Comparator);
+	if (BehaviorInstances.IsEmpty())
+	{
+		return false;
+	}
+
+	bool bHasMatchingInstance = BehaviorInstances.ContainsByPredicate(
+		[TransitionScene, InSceneComparisonType](const FAvaTransitionBehaviorInstance* InInstance)
+		{
+			const FAvaTransitionScene* OtherTransitionScene = InInstance->GetTransitionContext().GetTransitionScene();
+
+			const EAvaTransitionComparisonResult ComparisonResult = OtherTransitionScene
+				? TransitionScene->Compare(*OtherTransitionScene)
+				: EAvaTransitionComparisonResult::None;
+
+			return ComparisonResult == InSceneComparisonType;
+		});
+
+	return bHasMatchingInstance;
+}
+
+EAvaTransitionType UAvaTransitionLibrary::GetTransitionType(UObject* InTransitionNode)
+{
+	const IAvaTransitionNodeInterface* NodeInterface = Cast<IAvaTransitionNodeInterface>(InTransitionNode);
+	if (!NodeInterface)
+	{
+		return EAvaTransitionType::None;
+	}
+
+	const FAvaTransitionContext* TransitionContext = NodeInterface->GetBehaviorInstanceCache().GetTransitionContext();
+	if (!TransitionContext)
+	{
+		return EAvaTransitionType::None;
+	}
+
+	return TransitionContext->GetTransitionType();
+}
+
+bool UAvaTransitionLibrary::IsOtherSceneTransitioning(UObject* InTransitionNode, const TSoftObjectPtr<UWorld>& InSceneAsset)
+{
+	const IAvaTransitionNodeInterface* NodeInterface = Cast<IAvaTransitionNodeInterface>(InTransitionNode);
+	if (!NodeInterface)
+	{
+		return false;
+	}
+
+	const FAvaTransitionContext* TransitionContext = NodeInterface->GetBehaviorInstanceCache().GetTransitionContext();
+	if (!TransitionContext)
+	{
+		return false;
+	}
+
+	const FAvaTransitionScene* TransitionScene = TransitionContext->GetTransitionScene();
+	if (!TransitionScene)
+	{
+		return false;
+	}
+
+	ULevel* TransitionLevel = TransitionScene->GetLevel();
+	if (!TransitionLevel || !TransitionLevel->OwningWorld)
+	{
+		return false;
+	}
+
+	UAvaTransitionSubsystem* TransitionSubsystem = TransitionLevel->OwningWorld->GetSubsystem<UAvaTransitionSubsystem>();
+	if (!TransitionSubsystem)
+	{
+		return false;
+	}
+
+	const FAvaTransitionLayerComparator Comparator = FAvaTransitionLayerUtils::BuildComparator(*TransitionContext
+		, EAvaTransitionLayerCompareType::Different
+		, FAvaTagHandleContainer());
+
+	TArray<const FAvaTransitionBehaviorInstance*> BehaviorInstances = FAvaTransitionLayerUtils::QueryBehaviorInstances(*TransitionSubsystem, Comparator);
+	if (BehaviorInstances.IsEmpty())
+	{
+		return false;
+	}
+
+	bool bFoundInstance = BehaviorInstances.ContainsByPredicate(
+		[TransitionScene, &InSceneAsset](const FAvaTransitionBehaviorInstance* InInstance)
+		{
+			const FAvaTransitionScene* const OtherTransitionScene = InInstance->GetTransitionContext().GetTransitionScene();
+
+			// ignore check on scenes marked as needing discard
+			if (!OtherTransitionScene || OtherTransitionScene->HasAnyFlags(EAvaTransitionSceneFlags::NeedsDiscard))
+			{
+				return false;
+			}
+
+			const ULevel* const OtherLevel = OtherTransitionScene->GetLevel();
+			if (!OtherLevel)
+			{
+				return false;
+			}
+
+			const UPackage* const OtherPackage = OtherLevel->GetPackage();
+			if (!OtherPackage)
+			{
+				return false;
+			}
+
+			// Remove the /Temp from the Package Name
+			FString PackageName = OtherPackage->GetName();
+			if (PackageName.StartsWith(TEXT("/Temp")))
+			{
+				PackageName.RightChopInline(-1 + sizeof(TEXT("/Temp")) / sizeof(TCHAR));
+			}
+
+			// Remove the _LevelInstance_[Num] from the Package Name
+			const int32 LevelInstancePosition = PackageName.Find(TEXT("_LevelInstance_"), ESearchCase::IgnoreCase, ESearchDir::FromEnd);
+			if (LevelInstancePosition != INDEX_NONE)
+			{
+				PackageName.LeftChopInline(PackageName.Len() - LevelInstancePosition);
+			}
+
+			const FString SceneAssetName = InSceneAsset.GetLongPackageName();
+			return PackageName == SceneAssetName;
+		});
+
+	return bFoundInstance;
+}
diff --git a/Engine/Plugins/Experimental/Avalanche/Source/AvalancheTransition/Private/AvaTransitionLibrary.h b/Engine/Plugins/Experimental/Avalanche/Source/AvalancheTransition/Private/AvaTransitionLibrary.h
new file mode 100644
--- /dev/null
+++ b/Engine/Plugins/Experimental/Avalanche/Source/AvalancheTransition/Private/AvaTransitionLibrary.h
@@ -0,0 +1,30 @@
+// Copyright Epic Games, Inc. All Rights Reserved.
+
+#pragma once
+
+#include "Kismet/BlueprintFunctionLibrary.h"
+#include "AvaTransitionLibrary.generated.h"
+
+enum class EAvaTransitionComparisonResult : uint8;
+enum class EAvaTransitionLayerCompareType : uint8;
+enum class EAvaTransitionType : uint8;
+struct FAvaTagHandleContainer;
+
+UCLASS()
+class UAvaTransitionLibrary : public UBlueprintFunctionLibrary
+{
+	GENERATED_BODY()
+
+public:
+	UFUNCTION(BlueprintCallable, Category="Transition Logic", meta=(DefaultToSelf="InTransitionNode"))
+	static bool IsTransitionActiveInLayer(UObject* InTransitionNode
+		, EAvaTransitionComparisonResult InSceneComparisonType
+		, EAvaTransitionLayerCompareType InLayerComparisonType
+		, const FAvaTagHandleContainer& InSpecificLayers);
+
+	UFUNCTION(BlueprintCallable, Category="Transition Logic", meta=(DefaultToSelf="InTransitionNode"))
+	static EAvaTransitionType GetTransitionType(UObject* InTransitionNode);
+
+	UFUNCTION(BlueprintCallable, Category="Transition Logic", meta=(DefaultToSelf="InTransitionNode"))
+	static bool IsOtherSceneTransitioning(UObject* InTransitionNode, const TSoftObjectPtr<UWorld>& InSceneAsset);
+};
