Motion Design
- When Command Lists are exported from a module, the Get() method is replaced with GetInternal() for within the module and GetExternal() for outside the module, which ensures the Command List is registered.

diff --git a/Engine/Plugins/Experimental/Avalanche/Source/AvalancheEditor/Private/Scene/AvaSceneRigExtension.cpp b/Engine/Plugins/Experimental/Avalanche/Source/AvalancheEditor/Private/Scene/AvaSceneRigExtension.cpp
--- a/Engine/Plugins/Experimental/Avalanche/Source/AvalancheEditor/Private/Scene/AvaSceneRigExtension.cpp
+++ b/Engine/Plugins/Experimental/Avalanche/Source/AvalancheEditor/Private/Scene/AvaSceneRigExtension.cpp
@@ -43,8 +43,8 @@ void FAvaSceneRigExtension::Deactivate()
 void FAvaSceneRigExtension::BindCommands(const TSharedRef<FUICommandList>& InCommandList)
 {
 	FAvaEditorExtension::BindCommands(InCommandList);
-	
-	const FAvaSceneRigEditorCommands& SceneRigEditorCommands = FAvaSceneRigEditorCommands::Get();
+
+	const FAvaSceneRigEditorCommands& SceneRigEditorCommands = FAvaSceneRigEditorCommands::GetExternal();
 	
 	InCommandList->MapAction(
 		SceneRigEditorCommands.PromptToSaveSceneRigFromOutlinerItems,
@@ -88,7 +88,7 @@ void FAvaSceneRigExtension::ExtendOutlinerItemContextMenu(UToolMenu* const InToo
 	const TArray<AActor*> SelectedActors = OutlinerItemsToActors(ItemsContextWeak.Get()->GetItems(), true);
 	if (!SelectedActors.IsEmpty() && UAvaSceneRigSubsystem::AreActorsSupported(SelectedActors))
 	{
-		const FAvaSceneRigEditorCommands& SceneRigEditorCommands = FAvaSceneRigEditorCommands::Get();
+		const FAvaSceneRigEditorCommands& SceneRigEditorCommands = FAvaSceneRigEditorCommands::GetExternal();
 
 		FToolMenuSection& NewSection = InToolMenu->AddSection(TEXT("SceneRig"), LOCTEXT("SceneRig", "Scene Rig")
 			, FToolMenuInsert(TEXT("ContextActions"), EToolMenuInsertType::After));
@@ -107,7 +107,7 @@ void FAvaSceneRigExtension::CreateSubMenu(UToolMenu* const InToolMenu)
 	InToolMenu->AddDynamicSection(TEXT("SceneRig"), FNewToolMenuDelegate::CreateSPLambda(this, [this](UToolMenu* const InToolMenu)
 		{
 			const TSharedPtr<FUICommandList> CommandList = CommandListWeak.Pin();
-			const FAvaSceneRigEditorCommands& SceneRigEditorCommands = FAvaSceneRigEditorCommands::Get();
+			const FAvaSceneRigEditorCommands& SceneRigEditorCommands = FAvaSceneRigEditorCommands::GetExternal();
 
 			FToolMenuSection& NewSection = InToolMenu->AddSection(TEXT("SceneRig"), LOCTEXT("SceneRigActions", "Scene Rig"));
 
diff --git a/Engine/Plugins/Experimental/Avalanche/Source/AvalancheEditor/Private/Viewport/AvaLevelViewportExtension.cpp b/Engine/Plugins/Experimental/Avalanche/Source/AvalancheEditor/Private/Viewport/AvaLevelViewportExtension.cpp
--- a/Engine/Plugins/Experimental/Avalanche/Source/AvalancheEditor/Private/Viewport/AvaLevelViewportExtension.cpp
+++ b/Engine/Plugins/Experimental/Avalanche/Source/AvalancheEditor/Private/Viewport/AvaLevelViewportExtension.cpp
@@ -130,7 +130,7 @@ void FAvaLevelViewportExtension::BindCommands(const TSharedRef<FUICommandList>&
 	InCommandList->MapAction(FAvaEditorCommands::Get().SwitchViewports
 		, FExecuteAction::CreateSP(this, &FAvaLevelViewportExtension::OnSwitchViewports));
 
-	const FAvaLevelViewportCommands& LevelViewportCommands = FAvaLevelViewportCommands::Get();
+	const FAvaLevelViewportCommands& LevelViewportCommands = FAvaLevelViewportCommands::GetExternal();
 
 	InCommandList->MapAction(LevelViewportCommands.ResetLocation
 		, FExecuteAction::CreateSP(this, &FAvaLevelViewportExtension::ExecuteResetLocation));
diff --git a/Engine/Plugins/Experimental/Avalanche/Source/AvalancheLevelViewport/Private/AvaLevelViewportCommands.cpp b/Engine/Plugins/Experimental/Avalanche/Source/AvalancheLevelViewport/Private/AvaLevelViewportCommands.cpp
--- a/Engine/Plugins/Experimental/Avalanche/Source/AvalancheLevelViewport/Private/AvaLevelViewportCommands.cpp
+++ b/Engine/Plugins/Experimental/Avalanche/Source/AvalancheLevelViewport/Private/AvaLevelViewportCommands.cpp
@@ -5,6 +5,21 @@
 
 #define LOCTEXT_NAMESPACE "AvaLevelViewportCommands"
 
+const FAvaLevelViewportCommands& FAvaLevelViewportCommands::GetExternal()
+{
+	if (!IsRegistered())
+	{
+		Register();
+	}
+
+	return Get();
+}
+
+const FAvaLevelViewportCommands& FAvaLevelViewportCommands::GetInternal()
+{
+	return Get();
+}
+
 FAvaLevelViewportCommands::FAvaLevelViewportCommands()
 	: TCommands<FAvaLevelViewportCommands>(
 		TEXT("AvaLevelViewport")
diff --git a/Engine/Plugins/Experimental/Avalanche/Source/AvalancheLevelViewport/Private/AvaViewportCameraHistory.cpp b/Engine/Plugins/Experimental/Avalanche/Source/AvalancheLevelViewport/Private/AvaViewportCameraHistory.cpp
--- a/Engine/Plugins/Experimental/Avalanche/Source/AvalancheLevelViewport/Private/AvaViewportCameraHistory.cpp
+++ b/Engine/Plugins/Experimental/Avalanche/Source/AvalancheLevelViewport/Private/AvaViewportCameraHistory.cpp
@@ -51,7 +51,7 @@ void FAvaViewportCameraHistory::BindCommands()
 		return;
 	}
 	
-	const FAvaLevelViewportCommands& AvaLevelViewportCommands = FAvaLevelViewportCommands::Get();
+	const FAvaLevelViewportCommands& AvaLevelViewportCommands = FAvaLevelViewportCommands::GetInternal();
 	const TSharedRef<FUICommandList> LevelEditorActions = LevelEditorModule->GetGlobalLevelEditorActions();
 
 	LevelEditorActions->MapAction(AvaLevelViewportCommands.CameraTransformUndo
@@ -73,7 +73,7 @@ void FAvaViewportCameraHistory::UnbindCommands()
 		return;
 	}
 	
-	const FAvaLevelViewportCommands& AvaLevelViewportCommands = FAvaLevelViewportCommands::Get();
+	const FAvaLevelViewportCommands& AvaLevelViewportCommands = FAvaLevelViewportCommands::GetInternal();
 	const TSharedRef<FUICommandList> LevelEditorActions = LevelEditorModule->GetGlobalLevelEditorActions();
 
 	LevelEditorActions->UnmapAction(AvaLevelViewportCommands.CameraTransformUndo);
diff --git a/Engine/Plugins/Experimental/Avalanche/Source/AvalancheLevelViewport/Private/SAvaLevelViewport.cpp b/Engine/Plugins/Experimental/Avalanche/Source/AvalancheLevelViewport/Private/SAvaLevelViewport.cpp
--- a/Engine/Plugins/Experimental/Avalanche/Source/AvalancheLevelViewport/Private/SAvaLevelViewport.cpp
+++ b/Engine/Plugins/Experimental/Avalanche/Source/AvalancheLevelViewport/Private/SAvaLevelViewport.cpp
@@ -97,7 +97,7 @@ void SAvaLevelViewport::BindCommands()
 	CommandListRef.UnmapAction(EditorViewportCommands.Right);
 	CommandListRef.UnmapAction(EditorViewportCommands.Bottom);
 
-	const FAvaLevelViewportCommands& AvaLevelViewportCommands = FAvaLevelViewportCommands::Get();
+	const FAvaLevelViewportCommands& AvaLevelViewportCommands = FAvaLevelViewportCommands::GetInternal();
 
 	// Viewport
 	CommandListRef.MapAction(
@@ -611,7 +611,7 @@ void SAvaLevelViewport::ActivateCamera(TWeakObjectPtr<ACameraActor> InCamera)
 
 	ViewportClient->SetViewTarget(InCamera);
 
-	const FAvaLevelViewportCommands& AvaLevelViewportCommands = FAvaLevelViewportCommands::Get();
+	const FAvaLevelViewportCommands& AvaLevelViewportCommands = FAvaLevelViewportCommands::GetInternal();
 	ViewportClient->GetZoomController()->Reset();
 
 	if (VirtualSizeAspectRatioState == EAvaViewportVirtualSizeAspectRatioState::LockedToCamera)
diff --git a/Engine/Plugins/Experimental/Avalanche/Source/AvalancheLevelViewport/Private/SAvaLevelViewport_Menus.cpp b/Engine/Plugins/Experimental/Avalanche/Source/AvalancheLevelViewport/Private/SAvaLevelViewport_Menus.cpp
--- a/Engine/Plugins/Experimental/Avalanche/Source/AvalancheLevelViewport/Private/SAvaLevelViewport_Menus.cpp
+++ b/Engine/Plugins/Experimental/Avalanche/Source/AvalancheLevelViewport/Private/SAvaLevelViewport_Menus.cpp
@@ -107,7 +107,7 @@ TSharedPtr<SWidget> SAvaLevelViewport::MakeViewportToolbar()
 
 TSharedRef<SWidget> SAvaLevelViewport::OnExtendLevelEditorViewportToolbarForChildActorLock(FWeakObjectPtr InExtensionContext)
 {
-	const FAvaLevelViewportCommands& Commands = FAvaLevelViewportCommands::Get();
+	const FAvaLevelViewportCommands& Commands = FAvaLevelViewportCommands::GetInternal();
 
 	return SNew(SEditorViewportToolBarButton)
 		.ButtonType(EUserInterfaceActionType::Check)
@@ -168,7 +168,7 @@ void SAvaLevelViewport::FillCameraMenu(UToolMenu* InMenu)
 		return;
 	}
 
-	const FAvaLevelViewportCommands& LevelViewportCommands = FAvaLevelViewportCommands::Get();
+	const FAvaLevelViewportCommands& LevelViewportCommands = FAvaLevelViewportCommands::GetInternal();
 
 	TArray<TWeakObjectPtr<ACameraActor>> CameraActors;
 	for (ACameraActor* const CameraActor : TActorRange<ACameraActor>(World))
@@ -373,7 +373,7 @@ void SAvaLevelViewport::AddVirtualSizeMenuEntries(UToolMenu* InMenu)
 
 void SAvaLevelViewport::AddVirtualSizeDefaultEntries(FToolMenuSection& InSection)
 {
-	const FAvaLevelViewportCommands& LevelViewportCommands = FAvaLevelViewportCommands::Get();
+	const FAvaLevelViewportCommands& LevelViewportCommands = FAvaLevelViewportCommands::GetInternal();
 
 	InSection.AddMenuEntry("NoSize", LevelViewportCommands.VirtualSizeDisable, INVTEXT("-"));
 	InSection.AddMenuEntry("1920x1080", LevelViewportCommands.VirtualSize1920x1080, INVTEXT("1920 x 1080"));
@@ -383,7 +383,7 @@ void SAvaLevelViewport::AddVirtualSizeSizeSettings(FToolMenuSection& InSection)
 {
 	using namespace UE::AvaLevelViewport::Private;
 
-	const FAvaLevelViewportCommands& LevelViewportCommands = FAvaLevelViewportCommands::Get();
+	const FAvaLevelViewportCommands& LevelViewportCommands = FAvaLevelViewportCommands::GetInternal();
 
 	InSection.AddMenuEntry("UseUnlockedAspectRatio", LevelViewportCommands.VirtualSizeAspectRatioUnlocked, LOCTEXT("VirtualSizeUnlockedAspectRatio", "Free Aspect Ratio"));
 	InSection.AddMenuEntry("UseLockedAspectRatio", LevelViewportCommands.VirtualSizeAspectRatioLocked, LOCTEXT("VirtualSizeLockedAspectRatio", "Locked Aspect Ratio"));
@@ -744,7 +744,7 @@ void SAvaLevelViewport::AddGuidePresetSavedMenu(UToolMenu* InMenu)
 
 void SAvaLevelViewport::AddCameraZoomMenuEntries(UToolMenu* InMenu)
 {
-	const FAvaLevelViewportCommands& LevelViewportCommands = FAvaLevelViewportCommands::Get();
+	const FAvaLevelViewportCommands& LevelViewportCommands = FAvaLevelViewportCommands::GetInternal();
 
 	FToolMenuSection& PanSection = InMenu->AddSection(TEXT("CameraPanSection"), LOCTEXT("CameraPanSection", "Pan"));
 	PanSection.AddMenuEntry("CameraPanLeft",       LevelViewportCommands.CameraPanLeft,       LOCTEXT("CameraPanLeft",       "Pan Left"));
diff --git a/Engine/Plugins/Experimental/Avalanche/Source/AvalancheLevelViewport/Private/Widgets/SAvaLevelViewportGuide.cpp b/Engine/Plugins/Experimental/Avalanche/Source/AvalancheLevelViewport/Private/Widgets/SAvaLevelViewportGuide.cpp
--- a/Engine/Plugins/Experimental/Avalanche/Source/AvalancheLevelViewport/Private/Widgets/SAvaLevelViewportGuide.cpp
+++ b/Engine/Plugins/Experimental/Avalanche/Source/AvalancheLevelViewport/Private/Widgets/SAvaLevelViewportGuide.cpp
@@ -179,7 +179,7 @@ namespace UE::AvaViewport::Private
 			return;
 		}
 
-		const FAvaLevelViewportCommands& ViewportCommands = FAvaLevelViewportCommands::Get();
+		const FAvaLevelViewportCommands& ViewportCommands = FAvaLevelViewportCommands::GetInternal();
 
 		FToolMenuSection& Section = InToolMenu->AddSection("Guide", LOCTEXT("GuideOptions", "Guide Options"));
 		Section.AddMenuEntryWithCommandList(ViewportCommands.ToggleGuideEnabled, Guide->GetGetCommandList());
@@ -839,7 +839,7 @@ void SAvaLevelViewportGuide::OpenRightClickMenu()
 
 void SAvaLevelViewportGuide::BindCommands()
 {
-	const FAvaLevelViewportCommands& ViewportCommands = FAvaLevelViewportCommands::Get();
+	const FAvaLevelViewportCommands& ViewportCommands = FAvaLevelViewportCommands::GetInternal();
 
 	CommandList = MakeShared<FUICommandList>();
 
diff --git a/Engine/Plugins/Experimental/Avalanche/Source/AvalancheLevelViewport/Private/Widgets/SAvaLevelViewportStatusBarButtons.cpp b/Engine/Plugins/Experimental/Avalanche/Source/AvalancheLevelViewport/Private/Widgets/SAvaLevelViewportStatusBarButtons.cpp
--- a/Engine/Plugins/Experimental/Avalanche/Source/AvalancheLevelViewport/Private/Widgets/SAvaLevelViewportStatusBarButtons.cpp
+++ b/Engine/Plugins/Experimental/Avalanche/Source/AvalancheLevelViewport/Private/Widgets/SAvaLevelViewportStatusBarButtons.cpp
@@ -254,7 +254,7 @@ void SAvaLevelViewportStatusBarButtons::PopulateActorButtons(TSharedPtr<SHorizon
 {
 	using namespace UE::AvaLevelViewport::Private;
 
-	const FAvaLevelViewportCommands& CommandsRef = FAvaLevelViewportCommands::Get();
+	const FAvaLevelViewportCommands& CommandsRef = FAvaLevelViewportCommands::GetInternal();
 
 	InContainer->AddSlot()
 		.AutoWidth()
@@ -292,7 +292,7 @@ void SAvaLevelViewportStatusBarButtons::PopulateViewportButtons(TSharedPtr<SHori
 {
 	using namespace UE::AvaLevelViewport::Private;
 
-	const FAvaLevelViewportCommands& CommandsRef = FAvaLevelViewportCommands::Get();
+	const FAvaLevelViewportCommands& CommandsRef = FAvaLevelViewportCommands::GetInternal();
 	const FLevelViewportCommands& ViewportActionsRef = FLevelViewportCommands::Get();
 
 	InContainer->AddSlot()
@@ -550,7 +550,7 @@ TSharedRef<SWidget> SAvaLevelViewportStatusBarButtons::GetPostProcessMenuContent
 
 	using namespace UE::Ava::LevelViewportStatusBarButtons::Private;
 
-	const FAvaLevelViewportCommands& CommandsRef = FAvaLevelViewportCommands::Get();
+	const FAvaLevelViewportCommands& CommandsRef = FAvaLevelViewportCommands::GetInternal();
 
 	FToolMenuSection& EffectsSection = ContextMenu->FindOrAddSection("Effects", LOCTEXT("Effects", "Effects"));
 
@@ -891,7 +891,7 @@ TSharedRef<SWidget> SAvaLevelViewportStatusBarButtons::GetSnappingMenuContent()
 
 	FToolMenuSection& SnapToSection = ContextMenu->FindOrAddSection("SnapTo", LOCTEXT("SnapTo", "Snap To"));
 
-	const FAvaLevelViewportCommands& CommandsRef = FAvaLevelViewportCommands::Get();
+	const FAvaLevelViewportCommands& CommandsRef = FAvaLevelViewportCommands::GetInternal();
 
 	SnapToSection.AddEntry(FToolMenuEntry::InitMenuEntryWithCommandList(
 		CommandsRef.ToggleGridSnapping,
@@ -1061,7 +1061,7 @@ TSharedRef<SWidget> SAvaLevelViewportStatusBarButtons::GetGridMenuContent() cons
 
 	FToolMenuSection& GridSection = ContextMenu->FindOrAddSection("Grid", LOCTEXT("Grid", "Grid"));
 
-	const FAvaLevelViewportCommands& CommandsRef = FAvaLevelViewportCommands::Get();
+	const FAvaLevelViewportCommands& CommandsRef = FAvaLevelViewportCommands::GetInternal();
 
 	const FAvaLevelViewportGuideFrameAndWidget FrameAndWidget(ViewportFrameWeak);
 
diff --git a/Engine/Plugins/Experimental/Avalanche/Source/AvalancheLevelViewport/Private/Widgets/SAvaLevelViewportStatusBarTransformSettings.cpp b/Engine/Plugins/Experimental/Avalanche/Source/AvalancheLevelViewport/Private/Widgets/SAvaLevelViewportStatusBarTransformSettings.cpp
--- a/Engine/Plugins/Experimental/Avalanche/Source/AvalancheLevelViewport/Private/Widgets/SAvaLevelViewportStatusBarTransformSettings.cpp
+++ b/Engine/Plugins/Experimental/Avalanche/Source/AvalancheLevelViewport/Private/Widgets/SAvaLevelViewportStatusBarTransformSettings.cpp
@@ -138,7 +138,7 @@ TSharedRef<SHorizontalBox> SAvaLevelViewportStatusBarTransformSettings::CreateLo
 {
 	using namespace UE::AvaLevelViewport::Private;
 
-	const FAvaLevelViewportCommands& CommandsRef = FAvaLevelViewportCommands::Get();
+	const FAvaLevelViewportCommands& CommandsRef = FAvaLevelViewportCommands::GetInternal();
 
 	return SNew(SHorizontalBox)
 
@@ -178,7 +178,7 @@ TSharedRef<SHorizontalBox> SAvaLevelViewportStatusBarTransformSettings::CreateRo
 {
 	using namespace UE::AvaLevelViewport::Private;
 
-	const FAvaLevelViewportCommands& CommandsRef = FAvaLevelViewportCommands::Get();
+	const FAvaLevelViewportCommands& CommandsRef = FAvaLevelViewportCommands::GetInternal();
 
 	return SNew(SHorizontalBox)
 
@@ -217,7 +217,7 @@ TSharedRef<SHorizontalBox> SAvaLevelViewportStatusBarTransformSettings::CreateSc
 {
 	using namespace UE::AvaLevelViewport::Private;
 
-	const FAvaLevelViewportCommands& CommandsRef = FAvaLevelViewportCommands::Get();
+	const FAvaLevelViewportCommands& CommandsRef = FAvaLevelViewportCommands::GetInternal();
 
 	return SNew(SHorizontalBox)
 
diff --git a/Engine/Plugins/Experimental/Avalanche/Source/AvalancheLevelViewport/Public/AvaLevelViewportCommands.h b/Engine/Plugins/Experimental/Avalanche/Source/AvalancheLevelViewport/Public/AvaLevelViewportCommands.h
--- a/Engine/Plugins/Experimental/Avalanche/Source/AvalancheLevelViewport/Public/AvaLevelViewportCommands.h
+++ b/Engine/Plugins/Experimental/Avalanche/Source/AvalancheLevelViewport/Public/AvaLevelViewportCommands.h
@@ -4,23 +4,19 @@
 
 #include "Framework/Commands/Commands.h"
 
-class AVALANCHELEVELVIEWPORT_API FAvaLevelViewportCommands : public TCommands<FAvaLevelViewportCommands>
+class FAvaLevelViewportCommands : public TCommands<FAvaLevelViewportCommands>
 {
 public:
+	AVALANCHELEVELVIEWPORT_API static const FAvaLevelViewportCommands& GetExternal();
+
+	static const FAvaLevelViewportCommands& GetInternal();
+
 	FAvaLevelViewportCommands();
 
 	//~ Begin TCommands
 	virtual void RegisterCommands() override;
 	//~ End TCommands
 
-	void RegisterViewportCommands();
-	void RegisterCameraCommands();
-	void RegisterGridCommands();
-	void RegisterSnappingCommands();
-	void RegisterVirtualSizeCommands();
-	void RegisterGuideCommands();
-	void RegisterTransformCommands();
-
 	// Viewport
 	TSharedPtr<FUICommandInfo> ToggleOverlay;
 	TSharedPtr<FUICommandInfo> ToggleBoundingBoxes;
@@ -82,4 +78,16 @@ public:
 	TSharedPtr<FUICommandInfo> ResetRotation;
 	TSharedPtr<FUICommandInfo> ResetScale;
 	TSharedPtr<FUICommandInfo> ResetTransform;
+
+private:
+	// Make this unavailable to the public
+	using TCommands<FAvaLevelViewportCommands>::Get;
+
+	void RegisterViewportCommands();
+	void RegisterCameraCommands();
+	void RegisterGridCommands();
+	void RegisterSnappingCommands();
+	void RegisterVirtualSizeCommands();
+	void RegisterGuideCommands();
+	void RegisterTransformCommands();
 };
diff --git a/Engine/Plugins/Experimental/Avalanche/Source/AvalancheSceneRigEditor/Private/AvaSceneRigEditorCommands.cpp b/Engine/Plugins/Experimental/Avalanche/Source/AvalancheSceneRigEditor/Private/AvaSceneRigEditorCommands.cpp
--- a/Engine/Plugins/Experimental/Avalanche/Source/AvalancheSceneRigEditor/Private/AvaSceneRigEditorCommands.cpp
+++ b/Engine/Plugins/Experimental/Avalanche/Source/AvalancheSceneRigEditor/Private/AvaSceneRigEditorCommands.cpp
@@ -4,6 +4,21 @@
 
 #define LOCTEXT_NAMESPACE "AvaSceneRigEditorCommands"
 
+const FAvaSceneRigEditorCommands& FAvaSceneRigEditorCommands::GetExternal()
+{
+	if (!IsRegistered())
+	{
+		Register();
+	}
+
+	return Get();
+}
+
+const FAvaSceneRigEditorCommands& FAvaSceneRigEditorCommands::GetInternal()
+{
+	return Get();
+}
+
 void FAvaSceneRigEditorCommands::RegisterCommands()
 {
 	UI_COMMAND(PromptToSaveSceneRigFromOutlinerItems
diff --git a/Engine/Plugins/Experimental/Avalanche/Source/AvalancheSceneRigEditor/Public/AvaSceneRigEditorCommands.h b/Engine/Plugins/Experimental/Avalanche/Source/AvalancheSceneRigEditor/Public/AvaSceneRigEditorCommands.h
--- a/Engine/Plugins/Experimental/Avalanche/Source/AvalancheSceneRigEditor/Public/AvaSceneRigEditorCommands.h
+++ b/Engine/Plugins/Experimental/Avalanche/Source/AvalancheSceneRigEditor/Public/AvaSceneRigEditorCommands.h
@@ -5,9 +5,13 @@
 #include "Framework/Commands/Commands.h"
 #include "Styling/AppStyle.h"
 
-class AVALANCHESCENERIGEDITOR_API FAvaSceneRigEditorCommands : public TCommands<FAvaSceneRigEditorCommands>
+class FAvaSceneRigEditorCommands : public TCommands<FAvaSceneRigEditorCommands>
 {
 public:
+	AVALANCHESCENERIGEDITOR_API static const FAvaSceneRigEditorCommands& GetExternal();
+
+	static const FAvaSceneRigEditorCommands& GetInternal();
+
 	FAvaSceneRigEditorCommands()
 		: TCommands<FAvaSceneRigEditorCommands>(TEXT("AvaSceneRigCommands")
 		, NSLOCTEXT("AvaSceneRigCommands", "AvaSceneRigCommands", "Avalanche Scene Rig Commands")
@@ -23,4 +27,8 @@ public:
 
 	TSharedPtr<FUICommandInfo> AddOutlinerItemsToSceneRig;
 	TSharedPtr<FUICommandInfo> RemoveOutlinerItemsToSceneRig;
+
+private:
+	// Make this unavailable to the public
+	using TCommands<FAvaSceneRigEditorCommands>::Get;
 };
