MotionDesign : Modifiers
- Fixed Mirror and PlaneCut modifier preview plane not showing when preview is enabled

diff --git a/Engine/Plugins/Experimental/Avalanche/Source/AvalancheModifiers/Private/AvaModifiersPreviewPlane.cpp b/Engine/Plugins/Experimental/Avalanche/Source/AvalancheModifiers/Private/AvaModifiersPreviewPlane.cpp
--- a/Engine/Plugins/Experimental/Avalanche/Source/AvalancheModifiers/Private/AvaModifiersPreviewPlane.cpp
+++ b/Engine/Plugins/Experimental/Avalanche/Source/AvalancheModifiers/Private/AvaModifiersPreviewPlane.cpp
@@ -24,14 +24,17 @@ void FAvaModifierPreviewPlane::Create(USceneComponent* InActorComponent)
 	{
 		return;
 	}
-	
+
 	AActor* Actor = InActorComponent->GetOwner();
-	
+
 	PreviewComponent = NewObject<UStaticMeshComponent>(Actor);
-	PreviewComponent->AttachToComponent(InActorComponent, FAttachmentTransformRules::SnapToTargetNotIncludingScale);
+	PreviewComponent->OnComponentCreated();
+	PreviewComponent->SetupAttachment(InActorComponent);
 	PreviewComponent->RegisterComponent();
+
 	PreviewComponent->SetHiddenInGame(true);
-	PreviewComponent->SetBoundsScale(0.f);
+	PreviewComponent->SetBoundsScale(UE_KINDA_SMALL_NUMBER);
+	PreviewComponent->Bounds = FBox(ForceInitToZero);
 
 #if WITH_EDITOR
 	PreviewComponent->SetIsVisualizationComponent(true);
@@ -70,7 +73,7 @@ void FAvaModifierPreviewPlane::Hide() const
 	{
 		return;
 	}
-	
+
 	PreviewComponent->SetVisibility(false, false);
 	PreviewComponent->SetStaticMesh(nullptr);
 }
@@ -92,7 +95,7 @@ UStaticMesh* FAvaModifierPreviewPlane::LoadPreviewResource() const
 {
 	// get material asset
 	static const FString AssetPath = TEXT("/Script/Engine.StaticMesh'/Avalanche/EditorResources/SM_TwoSidedPlane.SM_TwoSidedPlane'");
-	
+
 	UStaticMesh* LoadedMesh = FindObject<UStaticMesh>(nullptr, *AssetPath);
 	if (!LoadedMesh)
 	{
