Motion Design: Renamed Mask Apply/set to Target/Source

diff --git a/Engine/Plugins/Experimental/Avalanche/Source/AvalancheMask/Private/Mask2D/AvaMask2DReadModifier.cpp b/Engine/Plugins/Experimental/Avalanche/Source/AvalancheMask/Private/Mask2D/AvaMask2DReadModifier.cpp
--- a/Engine/Plugins/Experimental/Avalanche/Source/AvalancheMask/Private/Mask2D/AvaMask2DReadModifier.cpp
+++ b/Engine/Plugins/Experimental/Avalanche/Source/AvalancheMask/Private/Mask2D/AvaMask2DReadModifier.cpp
@@ -154,7 +154,7 @@ void UAvaMask2DReadModifier::OnModifierCDOSetup(FActorModifierCoreMetadata& InMe
 	InMetadata.DisallowAfter(TEXT("MaskWrite"));
 	InMetadata.DisallowBefore(TEXT("MaskWrite"));
 #if WITH_EDITOR
-	InMetadata.SetDisplayName(FText::FromString(TEXT("Mask (Apply)")));
+	InMetadata.SetDisplayName(FText::FromString(TEXT("Mask (Target)")));
 	InMetadata.SetDescription(LOCTEXT("ModifierDescription", "Allows to use a custom mask texture on attached actors materials"));
 #endif
 }
diff --git a/Engine/Plugins/Experimental/Avalanche/Source/AvalancheMask/Private/Mask2D/AvaMask2DWriteModifier.cpp b/Engine/Plugins/Experimental/Avalanche/Source/AvalancheMask/Private/Mask2D/AvaMask2DWriteModifier.cpp
--- a/Engine/Plugins/Experimental/Avalanche/Source/AvalancheMask/Private/Mask2D/AvaMask2DWriteModifier.cpp
+++ b/Engine/Plugins/Experimental/Avalanche/Source/AvalancheMask/Private/Mask2D/AvaMask2DWriteModifier.cpp
@@ -112,7 +112,7 @@ void UAvaMask2DWriteModifier::OnModifierCDOSetup(FActorModifierCoreMetadata& InM
 	InMetadata.DisallowAfter(TEXT("MaskRead"));
 	InMetadata.DisallowBefore(TEXT("MaskRead"));
 #if WITH_EDITOR
-	InMetadata.SetDisplayName(FText::FromString(TEXT("Mask (Set)")));
+	InMetadata.SetDisplayName(FText::FromString(TEXT("Mask (Source)")));
 	InMetadata.SetDescription(LOCTEXT("ModifierDescription", "Allows to use a custom mask texture on attached actors materials"));
 #endif
 }
