Motion Design: Added option to hide actors while waiting for the Delay Task to complete + fixed an issue with TL Rendering Subsystem not supporting multiple places calling Hide/Show

diff --git a/Engine/Plugins/Experimental/Avalanche/Source/AvalancheTransition/Private/Rendering/AvaTransitionRenderingSubsystem.cpp b/Engine/Plugins/Experimental/Avalanche/Source/AvalancheTransition/Private/Rendering/AvaTransitionRenderingSubsystem.cpp
--- a/Engine/Plugins/Experimental/Avalanche/Source/AvalancheTransition/Private/Rendering/AvaTransitionRenderingSubsystem.cpp
+++ b/Engine/Plugins/Experimental/Avalanche/Source/AvalancheTransition/Private/Rendering/AvaTransitionRenderingSubsystem.cpp
@@ -8,7 +8,7 @@
 
 void UAvaTransitionRenderingSubsystem::ShowLevel(TObjectKey<ULevel> InLevel)
 {
-	HiddenLevels.Remove(InLevel);
+	HiddenLevels.RemoveSingle(InLevel);
 }
 
 void UAvaTransitionRenderingSubsystem::HideLevel(TObjectKey<ULevel> InLevel)
@@ -19,9 +19,18 @@ void UAvaTransitionRenderingSubsystem::HideLevel(TObjectKey<ULevel> InLevel)
 void UAvaTransitionRenderingSubsystem::SetupView(FSceneView& InView)
 {
 	HiddenPrimitives.Reset();
+	ProcessedLevels.Reset();
 
 	for (const TObjectKey<ULevel>& LevelKey : HiddenLevels)
 	{
+		bool bAlreadyProcessedLevel;
+		ProcessedLevels.Add(LevelKey, &bAlreadyProcessedLevel);
+
+		if (bAlreadyProcessedLevel)
+		{
+			continue;
+		}
+
 		const ULevel* Level = LevelKey.ResolveObjectPtr();
 		if (!Level)
 		{
diff --git a/Engine/Plugins/Experimental/Avalanche/Source/AvalancheTransition/Private/Tasks/AvaTransitionDelayTask.cpp b/Engine/Plugins/Experimental/Avalanche/Source/AvalancheTransition/Private/Tasks/AvaTransitionDelayTask.cpp
--- a/Engine/Plugins/Experimental/Avalanche/Source/AvalancheTransition/Private/Tasks/AvaTransitionDelayTask.cpp
+++ b/Engine/Plugins/Experimental/Avalanche/Source/AvalancheTransition/Private/Tasks/AvaTransitionDelayTask.cpp
@@ -1,8 +1,14 @@
 // Copyright Epic Games, Inc. All Rights Reserved.
 
 #include "Tasks/AvaTransitionDelayTask.h"
+#include "AvaTransitionContext.h"
+#include "AvaTransitionScene.h"
+#include "AvaTransitionTree.h"
 #include "AvaTransitionUtils.h"
+#include "Engine/Level.h"
+#include "Rendering/AvaTransitionRenderingSubsystem.h"
 #include "StateTreeExecutionContext.h"
+#include "StateTreeLinker.h"
 
 #define LOCTEXT_NAMESPACE "AvaTransitionDelayTask"
 
@@ -19,6 +25,13 @@ FText FAvaTransitionDelayTask::GetDescription(const FGuid& InId, FStateTreeDataV
 }
 #endif
 
+bool FAvaTransitionDelayTask::Link(FStateTreeLinker& InLinker)
+{
+	Super::Link(InLinker);
+	InLinker.LinkExternalData(RenderingSubsystemHandle);
+	return true;
+}
+
 void FAvaTransitionDelayTask::PostLoad(FStateTreeDataView InInstanceDataView)
 {
 	Super::PostLoad(InInstanceDataView);
@@ -38,26 +51,60 @@ EStateTreeRunStatus FAvaTransitionDelayTask::EnterState(FStateTreeExecutionConte
 {
 	FInstanceDataType& InstanceData = InContext.GetInstanceData(*this);
 	InstanceData.RemainingTime = InstanceData.Duration;
+	return WaitForDelayCompletion(InContext, InstanceData);
+}
+
+EStateTreeRunStatus FAvaTransitionDelayTask::Tick(FStateTreeExecutionContext& InContext, const float InDeltaTime) const
+{
+	FInstanceDataType& InstanceData = InContext.GetInstanceData(*this);
+	InstanceData.RemainingTime -= InDeltaTime;
+	return WaitForDelayCompletion(InContext, InstanceData);
+}
 
-	if (InstanceData.RemainingTime <= 0.f)
+EStateTreeRunStatus FAvaTransitionDelayTask::WaitForDelayCompletion(FStateTreeExecutionContext& InContext, FInstanceDataType& InInstanceData) const
+{
+	UAvaTransitionRenderingSubsystem& RenderingSubsystem = InContext.GetExternalData(RenderingSubsystemHandle);
+
+	if (InInstanceData.RemainingTime <= 0.f)
 	{
+		// Restore Level Visibility
+		RenderingSubsystem.ShowLevel(InInstanceData.HiddenLevel);
 		return EStateTreeRunStatus::Succeeded;
 	}
 
+	if (ShouldHideLevel(InContext, InInstanceData))
+	{
+		if (const FAvaTransitionScene* TransitionScene = InContext.GetExternalData(TransitionContextHandle).GetTransitionScene())
+		{
+			InInstanceData.HiddenLevel = TransitionScene->GetLevel();
+			RenderingSubsystem.HideLevel(InInstanceData.HiddenLevel);
+		}
+	}
+
 	return EStateTreeRunStatus::Running;
 }
 
-EStateTreeRunStatus FAvaTransitionDelayTask::Tick(FStateTreeExecutionContext& InContext, const float InDeltaTime) const
+bool FAvaTransitionDelayTask::ShouldHideLevel(const FStateTreeExecutionContext& InContext, const FInstanceDataType& InInstanceData) const
 {
-	FInstanceDataType& InstanceData = InContext.GetInstanceData(*this);
+	if (InInstanceData.HideMode == EAvaTransitionLevelHideMode::NoHide)
+	{
+		return false;	
+	}
 
-	InstanceData.RemainingTime -= InDeltaTime;
+	// If Hidden Level is non-null, it means the level has already been hidden / processed. Skip
+	if (InInstanceData.HiddenLevel != nullptr)
+	{
+		return false;
+	}
 
-	if (InstanceData.RemainingTime <= 0.f)
+	// If Instancing Mode is set to Reuse, and user set to Hide Level regardless, allow Hiding Reused Level
+	const UAvaTransitionTree* TransitionTree = Cast<UAvaTransitionTree>(InContext.GetStateTree());
+	if (TransitionTree && TransitionTree->GetInstancingMode() == EAvaTransitionInstancingMode::Reuse)
 	{
-		return EStateTreeRunStatus::Succeeded;
+		return InInstanceData.HideMode == EAvaTransitionLevelHideMode::AlwaysHide;
 	}
-	return EStateTreeRunStatus::Running;
+
+	return true;
 }
 
 #undef LOCTEXT_NAMESPACE
diff --git a/Engine/Plugins/Experimental/Avalanche/Source/AvalancheTransition/Private/Tasks/AvaTransitionWaitForLayerTask.cpp b/Engine/Plugins/Experimental/Avalanche/Source/AvalancheTransition/Private/Tasks/AvaTransitionWaitForLayerTask.cpp
--- a/Engine/Plugins/Experimental/Avalanche/Source/AvalancheTransition/Private/Tasks/AvaTransitionWaitForLayerTask.cpp
+++ b/Engine/Plugins/Experimental/Avalanche/Source/AvalancheTransition/Private/Tasks/AvaTransitionWaitForLayerTask.cpp
@@ -75,9 +75,9 @@ EStateTreeRunStatus FAvaTransitionWaitForLayerTask::WaitForLayer(FStateTreeExecu
 
 bool FAvaTransitionWaitForLayerTask::ShouldHideLevel(const FStateTreeExecutionContext& InContext, const FAvaTransitionWaitForLayerTask::FInstanceDataType& InInstanceData) const
 {
-	if (!InInstanceData.bHideSceneWhileWaiting)
+	if (InInstanceData.HideMode == EAvaTransitionLevelHideMode::NoHide)
 	{
-		return false;	
+		return false;
 	}
 
 	// If Hidden Level is non-null, it means the level has already been hidden / processed. Skip
@@ -90,7 +90,7 @@ bool FAvaTransitionWaitForLayerTask::ShouldHideLevel(const FStateTreeExecutionCo
 	const UAvaTransitionTree* TransitionTree = Cast<UAvaTransitionTree>(InContext.GetStateTree());
 	if (TransitionTree && TransitionTree->GetInstancingMode() == EAvaTransitionInstancingMode::Reuse)
 	{
-		return InInstanceData.bHideOnLevelReuse;
+		return InInstanceData.HideMode == EAvaTransitionLevelHideMode::AlwaysHide;
 	}
 
 	return true;
diff --git a/Engine/Plugins/Experimental/Avalanche/Source/AvalancheTransition/Public/AvaTransitionEnums.h b/Engine/Plugins/Experimental/Avalanche/Source/AvalancheTransition/Public/AvaTransitionEnums.h
--- a/Engine/Plugins/Experimental/Avalanche/Source/AvalancheTransition/Public/AvaTransitionEnums.h
+++ b/Engine/Plugins/Experimental/Avalanche/Source/AvalancheTransition/Public/AvaTransitionEnums.h
@@ -67,4 +67,17 @@ enum class EAvaTransitionInstancingMode : uint8
 	New,
 	/** If previous instance exists, reuse it for the entering behavior. */
 	Reuse
-};
\ No newline at end of file
+};
+
+UENUM(BlueprintType, DisplayName="Motion Design Transition Level Hide Mode")
+enum class EAvaTransitionLevelHideMode : uint8
+{
+	/** Do not hide the Actors of the Level */
+	NoHide,
+
+	/** Hide the Level's Actors unless it's a Level that is re-used */
+	HideUnlessReuse,
+
+	/** Always hide the Level's Actors even if it's re-used */
+	AlwaysHide,
+};
diff --git a/Engine/Plugins/Experimental/Avalanche/Source/AvalancheTransition/Public/Rendering/AvaTransitionRenderingSubsystem.h b/Engine/Plugins/Experimental/Avalanche/Source/AvalancheTransition/Public/Rendering/AvaTransitionRenderingSubsystem.h
--- a/Engine/Plugins/Experimental/Avalanche/Source/AvalancheTransition/Public/Rendering/AvaTransitionRenderingSubsystem.h
+++ b/Engine/Plugins/Experimental/Avalanche/Source/AvalancheTransition/Public/Rendering/AvaTransitionRenderingSubsystem.h
@@ -2,6 +2,7 @@
 
 #pragma once
 
+#include "Containers/Array.h"
 #include "Containers/Set.h"
 #include "PrimitiveComponentId.h"
 #include "Subsystems/WorldSubsystem.h"
@@ -23,7 +24,11 @@ public:
 	void SetupView(FSceneView& InView);
 
 private:
-	TSet<TObjectKey<ULevel>> HiddenLevels;
+	/** List of levels that will be hidden. Same level can be repeated to indicate multiple places are actively trying to have the level hidden */
+	TArray<TObjectKey<ULevel>> HiddenLevels;
 
 	TSet<FPrimitiveComponentId> HiddenPrimitives;
+
+	/** Temporary Member Variable to keep Track of the Processed Levels when Iterating Hidden Levels */
+	TSet<TObjectKey<ULevel>> ProcessedLevels;
 };
diff --git a/Engine/Plugins/Experimental/Avalanche/Source/AvalancheTransition/Public/Tasks/AvaTransitionDelayTask.h b/Engine/Plugins/Experimental/Avalanche/Source/AvalancheTransition/Public/Tasks/AvaTransitionDelayTask.h
--- a/Engine/Plugins/Experimental/Avalanche/Source/AvalancheTransition/Public/Tasks/AvaTransitionDelayTask.h
+++ b/Engine/Plugins/Experimental/Avalanche/Source/AvalancheTransition/Public/Tasks/AvaTransitionDelayTask.h
@@ -2,9 +2,14 @@
 
 #pragma once
 
+#include "AvaTransitionEnums.h"
 #include "AvaTransitionTask.h"
+#include "UObject/ObjectKey.h"
 #include "AvaTransitionDelayTask.generated.h"
 
+class UAvaTransitionRenderingSubsystem;
+class ULevel;
+
 USTRUCT()
 struct FAvaTransitionDelayTaskInstanceData
 {
@@ -14,8 +19,14 @@ struct FAvaTransitionDelayTaskInstanceData
 	UPROPERTY(EditAnywhere, Category="Parameter", meta=(ClampMin="0.0"))
 	float Duration = 0.5f;
 
+	/** Hide mode to use while the Wait is taking place */
+	UPROPERTY(EditAnywhere, Category="Transition Logic")
+	EAvaTransitionLevelHideMode HideMode = EAvaTransitionLevelHideMode::NoHide;
+
 	/** Internal countdown in seconds. */
 	float RemainingTime = 0.f;
+
+	TObjectKey<ULevel> HiddenLevel;
 };
 
 USTRUCT(DisplayName="Delay", Category="Transition Logic")
@@ -39,6 +50,7 @@ struct AVALANCHETRANSITION_API FAvaTransitionDelayTask : public FAvaTransitionTa
 	virtual FText GetDescription(const FGuid& InId, FStateTreeDataView InInstanceDataView, const IStateTreeBindingLookup& InBindingLookup, EStateTreeNodeFormatting InFormatting) const override;
 #endif
 	virtual const UStruct* GetInstanceDataType() const override { return FInstanceDataType::StaticStruct(); }
+	virtual bool Link(FStateTreeLinker& InLinker) override;
 	virtual void PostLoad(FStateTreeDataView InInstanceDataView) override;
 	//~ End FStateTreeNodeBase
 
@@ -47,7 +59,13 @@ struct AVALANCHETRANSITION_API FAvaTransitionDelayTask : public FAvaTransitionTa
 	virtual EStateTreeRunStatus Tick(FStateTreeExecutionContext& InContext, const float InDeltaTime) const override;
 	//~ End FStateTreeTaskBase
 
+	EStateTreeRunStatus WaitForDelayCompletion(FStateTreeExecutionContext& InContext, FInstanceDataType& InInstanceData) const;
+
+	bool ShouldHideLevel(const FStateTreeExecutionContext& InContext, const FInstanceDataType& InInstanceData) const;
+
 	UE_DEPRECATED(5.5, "Duration has been moved to Instance Data")
 	UPROPERTY(meta = (DeprecatedProperty, DeprecationMessage = "Use the Instance Data Duration instead"))
 	float Duration_DEPRECATED = -1.f;
+
+	TStateTreeExternalDataHandle<UAvaTransitionRenderingSubsystem> RenderingSubsystemHandle;
 };
diff --git a/Engine/Plugins/Experimental/Avalanche/Source/AvalancheTransition/Public/Tasks/AvaTransitionWaitForLayerTask.h b/Engine/Plugins/Experimental/Avalanche/Source/AvalancheTransition/Public/Tasks/AvaTransitionWaitForLayerTask.h
--- a/Engine/Plugins/Experimental/Avalanche/Source/AvalancheTransition/Public/Tasks/AvaTransitionWaitForLayerTask.h
+++ b/Engine/Plugins/Experimental/Avalanche/Source/AvalancheTransition/Public/Tasks/AvaTransitionWaitForLayerTask.h
@@ -14,12 +14,9 @@ struct FAvaTransitionWaitForLayerTaskInstanceData : public FAvaTransitionLayerTa
 {
 	GENERATED_BODY()
 
+	/** Hide mode to use while the Wait is taking place */
 	UPROPERTY(EditAnywhere, Category="Transition Logic")
-	bool bHideSceneWhileWaiting = true;
-
-	/** Whether to hide level even if set to Re-use. Requires Hide Scene While Waiting to be enabled */
-	UPROPERTY(EditAnywhere, Category="Transition Logic", meta=(EditCondition="bHideSceneWhileWaiting", EditConditionHides))
-	bool bHideOnLevelReuse = false;
+	EAvaTransitionLevelHideMode HideMode = EAvaTransitionLevelHideMode::HideUnlessReuse;
 
 	TObjectKey<ULevel> HiddenLevel;
 };
