MotionDesign : Modifiers
- Removing auto size ability to follow reference actor to avoid performance impact when marked dirty,
auto size is now a geometry modifier, auto follow modifier will handle the layout arrangement and padding instead

diff --git a/Engine/Plugins/Experimental/Avalanche/Source/AvalancheModifiers/Private/Modifiers/AvaAutoSizeModifier.cpp b/Engine/Plugins/Experimental/Avalanche/Source/AvalancheModifiers/Private/Modifiers/AvaAutoSizeModifier.cpp
--- a/Engine/Plugins/Experimental/Avalanche/Source/AvalancheModifiers/Private/Modifiers/AvaAutoSizeModifier.cpp
+++ b/Engine/Plugins/Experimental/Avalanche/Source/AvalancheModifiers/Private/Modifiers/AvaAutoSizeModifier.cpp
@@ -2,7 +2,6 @@
 
 #include "Modifiers/AvaAutoSizeModifier.h"
 
-#include "AvaActorUtils.h"
 #include "Async/Async.h"
 #include "AvaModifiersActorUtils.h"
 #include "Components/DynamicMeshComponent.h"
@@ -25,18 +24,7 @@ bool UAvaAutoSizeModifier::IsModifierDirtyable() const
 
 	const FBox ReferenceActorLocalBounds = FAvaModifiersActorUtils::GetActorsBounds(TrackedActor, bIncludeChildren, true);
 
-	const FVector FollowAxisVector = FAvaModifiersActorUtils::GetVectorAxis(FollowedAxis);
-
-	// If we follow reference actor check bounds
-	if (!FollowAxisVector.IsNearlyZero())
-	{
-		if (ReferenceActorLocalBounds.Equals(CachedReferenceBounds, 0.01))
-		{
-			return Super::IsModifierDirtyable();
-		}
-	}
-	// If we don't follow reference actor check size
-	else if (ReferenceActorLocalBounds.GetSize().Equals(CachedReferenceBounds.GetSize(), 0.01))
+	if (ReferenceActorLocalBounds.GetSize().Equals(CachedReferenceBounds.GetSize(), 0.01))
 	{
 		return Super::IsModifierDirtyable();
 	}
@@ -50,9 +38,9 @@ void UAvaAutoSizeModifier::OnModifierCDOSetup(FActorModifierCoreMetadata& InMeta
 
 	InMetadata.AllowTick(true);
 	InMetadata.SetName(TEXT("AutoSize"));
-	InMetadata.SetCategory(TEXT("Layout"));
+	InMetadata.SetCategory(TEXT("Geometry"));
 #if WITH_EDITOR
-	InMetadata.SetDescription(LOCTEXT("ModifierDescription", "The modified actor will be resized to act as a background for the specified actor"));
+	InMetadata.SetDescription(LOCTEXT("ModifierDescription", "The modified actor will be resized to act as a background for a reference actor"));
 #endif
 
 	InMetadata.SetCompatibilityRule([](const AActor* InActor)->bool
@@ -74,22 +62,7 @@ void UAvaAutoSizeModifier::OnModifierCDOSetup(FActorModifierCoreMetadata& InMeta
 
 void UAvaAutoSizeModifier::OnModifiedActorTransformed()
 {
-	Super::OnModifiedActorTransformed();
-
-	const AActor* ActorModified = GetModifiedActor();
-	const AActor* TrackedActor = ReferenceActor.ReferenceActorWeak.Get();
-	if (!TrackedActor || !ActorModified)
-	{
-		return;
-	}
-
-	// Compare current location and previous followed location but only component of followed axis to allow movements
-	if (FAvaModifiersActorUtils::IsAxisVectorEquals(ActorModified->GetActorLocation(), CachedFollowLocation, FollowedAxis))
-	{
-		return;
-	}
-
-	MarkModifierDirty();
+	// Do nothing when moved
 }
 
 void UAvaAutoSizeModifier::OnModifierAdded(EActorModifierCoreEnableReason InReason)
@@ -178,13 +151,12 @@ void UAvaAutoSizeModifier::Apply()
 	FVector ModifiedActorBoundsExtent;
 	ModifiedActorBounds.GetCenterAndExtents(ModifiedActorBoundsOrigin, ModifiedActorBoundsExtent);
 
-	const FMargin& DesiredPadding = Padding;
 	FVector ReferenceBoundsExtent = CachedReferenceBounds.GetExtent();
 
 	// Add padding only if there is content inside otherwise don't
 	if (!CachedReferenceBounds.GetExtent().IsNearlyZero())
 	{
-		ReferenceBoundsExtent += FVector(0, DesiredPadding.Right + DesiredPadding.Left, DesiredPadding.Top + DesiredPadding.Bottom);
+		ReferenceBoundsExtent += FVector(0, PaddingHorizontal, PaddingVertical);
 	}
 
 	// Avoid division by 0
@@ -197,29 +169,6 @@ void UAvaAutoSizeModifier::Apply()
 	const float ScaleRatioZ = ReferenceBoundsExtent.Z/ModifiedActorBoundsExtent.Z;
 	const FVector ScaleRatio = FVector(1.0f, ScaleRatioY, ScaleRatioZ);
 
-	// Move only if there is content inside and at least an axis is followed
-	const FVector FollowAxisVector = FAvaModifiersActorUtils::GetVectorAxis(FollowedAxis);
-	if (!FollowAxisVector.IsNearlyZero() && !CachedReferenceBounds.GetSize().IsNearlyZero())
-	{
-		const FVector ModifiedActorPivotToBoundsOffset = ModifiedActorBoundsOrigin - CurrentlyModifiedActor->GetActorLocation();
-		const FVector ReferenceActorPivotToBoundsOffset = CachedReferenceBounds.GetCenter() - TrackedActor->GetActorLocation();
-		const FVector PaddingOffset = FVector(0.f, DesiredPadding.Right - DesiredPadding.Left, DesiredPadding.Top - DesiredPadding.Bottom);
-
-		const FVector DesiredLocationOffset =
-			// Add margin padding
-			PaddingOffset
-			// Remove pivot actor location difference for modified actor
-			- ModifiedActorPivotToBoundsOffset
-			// add pivot offset difference for reference actor
-			+ ReferenceActorPivotToBoundsOffset;
-
-		CachedFollowLocation =
-			CurrentlyModifiedActor->GetActorLocation()
-			+ (TrackedActor->GetActorLocation() - CurrentlyModifiedActor->GetActorLocation() + DesiredLocationOffset) * FollowAxisVector;
-
-		CurrentlyModifiedActor->SetActorLocation(CachedFollowLocation);
-	}
-
 	// check if we can use Shape2DDynMesh SetSize2D for this modified mesh.
 	// this allows us to keep properly scaled corner bevels and slants
 	if (UAvaShape2DDynMeshBase* const Shape2DDynMesh = ShapeDynMesh2DWeak.Get())
@@ -293,6 +242,16 @@ void UAvaAutoSizeModifier::PostLoad()
 		bDeprecatedPropertiesMigrated = true;
 	}
 
+	const FVector2D Padding2D(Padding.Left + Padding.Right, Padding.Top + Padding.Bottom);
+	if (!Padding2D.IsNearlyZero()
+		&& PaddingHorizontal == 0
+		&& PaddingVertical == 0)
+	{
+		Padding = FMargin(0);
+		PaddingHorizontal = Padding2D.X;
+		PaddingVertical = Padding2D.Y;
+	}
+
 	Super::PostLoad();
 }
 
@@ -304,7 +263,8 @@ void UAvaAutoSizeModifier::PostEditChangeProperty(FPropertyChangedEvent& Propert
 	const FName MemberName = PropertyChangedEvent.GetMemberPropertyName();
 
 	static const FName ReferenceActorPropertyName = GET_MEMBER_NAME_CHECKED(UAvaAutoSizeModifier, ReferenceActor);
-	static const FName PaddingPropertyName = GET_MEMBER_NAME_CHECKED(UAvaAutoSizeModifier, Padding);
+	static const FName PaddingHorizontalPropertyName = GET_MEMBER_NAME_CHECKED(UAvaAutoSizeModifier, PaddingHorizontal);
+	static const FName PaddingVerticalPropertyName = GET_MEMBER_NAME_CHECKED(UAvaAutoSizeModifier, PaddingVertical);
 	static const FName IncludeChildrenName = GET_MEMBER_NAME_CHECKED(UAvaAutoSizeModifier, bIncludeChildren);
 	static const FName FitModePropertyName = GET_MEMBER_NAME_CHECKED(UAvaAutoSizeModifier, FitMode);
 
@@ -312,7 +272,8 @@ void UAvaAutoSizeModifier::PostEditChangeProperty(FPropertyChangedEvent& Propert
 	{
 		OnReferenceActorChanged();
 	}
-	else if (MemberName == PaddingPropertyName
+	else if (MemberName == PaddingHorizontalPropertyName
+		|| MemberName == PaddingVerticalPropertyName
 		|| MemberName == FitModePropertyName
 		|| MemberName == IncludeChildrenName)
 	{
@@ -332,25 +293,26 @@ void UAvaAutoSizeModifier::SetReferenceActor(const FAvaSceneTreeActor& InReferen
 	OnReferenceActorChanged();
 }
 
-void UAvaAutoSizeModifier::SetFollowedAxis(int32 InFollowedAxis)
+void UAvaAutoSizeModifier::SetPaddingHorizontal(double InPadding)
 {
-	if (FollowedAxis == InFollowedAxis)
+	if (FMath::IsNearlyEqual(InPadding, PaddingHorizontal))
 	{
 		return;
 	}
 
-	FollowedAxis = InFollowedAxis;
+	PaddingHorizontal = InPadding;
 	MarkModifierDirty();
 }
 
-void UAvaAutoSizeModifier::SetPadding(const FMargin& InPadding)
+void UAvaAutoSizeModifier::SetPaddingVertical(double InPadding)
 {
-	if (Padding != InPadding)
+	if (FMath::IsNearlyEqual(InPadding, PaddingVertical))
 	{
-		Padding = InPadding;
-
-		MarkModifierDirty();
+		return;
 	}
+
+	PaddingVertical = InPadding;
+	MarkModifierDirty();
 }
 
 void UAvaAutoSizeModifier::SetFitMode(const EAvaAutoSizeFitMode InFitMode)
diff --git a/Engine/Plugins/Experimental/Avalanche/Source/AvalancheModifiers/Public/Modifiers/AvaAutoFollowModifier.h b/Engine/Plugins/Experimental/Avalanche/Source/AvalancheModifiers/Public/Modifiers/AvaAutoFollowModifier.h
--- a/Engine/Plugins/Experimental/Avalanche/Source/AvalancheModifiers/Public/Modifiers/AvaAutoFollowModifier.h
+++ b/Engine/Plugins/Experimental/Avalanche/Source/AvalancheModifiers/Public/Modifiers/AvaAutoFollowModifier.h
@@ -180,15 +180,15 @@ protected:
 	FAvaAnchorAlignment LocalAlignment;
 
 	/** The distance from this actor to the followed actor. */
-	UPROPERTY(EditInstanceOnly, Setter="SetDefaultDistance", Getter="GetDefaultDistance", Interp, Category="AutoFollow", meta=(AllowPrivateAccess="true"))
+	UPROPERTY(EditInstanceOnly, Setter="SetDefaultDistance", Getter="GetDefaultDistance", Interp, DisplayName="Start Padding", Category="AutoFollow", meta=(AllowPrivateAccess="true"))
 	FVector DefaultDistance;
 
 	/** The maximum distance from this actor to the followed actor. */
-	UPROPERTY(EditInstanceOnly, Setter="SetMaxDistance", Getter="GetMaxDistance", Interp, Category="AutoFollow", meta=(AllowPrivateAccess="true"))
+	UPROPERTY(EditInstanceOnly, Setter="SetMaxDistance", Getter="GetMaxDistance", Interp, DisplayName="End Padding", Category="AutoFollow", meta=(AllowPrivateAccess="true"))
 	FVector MaxDistance;
 
 	/** Percent % progress from the maximum distance to the default distance. */
-	UPROPERTY(EditInstanceOnly, Setter="SetProgress", Getter="GetProgress", Interp, Category="AutoFollow", meta=(ClampMin="0.0", UIMin="0.0", ClampMax="100.0", UIMax="100.0", AllowPrivateAccess="true"))
+	UPROPERTY(EditInstanceOnly, Setter="SetProgress", Getter="GetProgress", Interp, DisplayName="Padding Progress", Category="AutoFollow", meta=(ClampMin="0.0", UIMin="0.0", ClampMax="100.0", UIMax="100.0", AllowPrivateAccess="true"))
 	FVector Progress;
 
 private:
diff --git a/Engine/Plugins/Experimental/Avalanche/Source/AvalancheModifiers/Public/Modifiers/AvaAutoSizeModifier.h b/Engine/Plugins/Experimental/Avalanche/Source/AvalancheModifiers/Public/Modifiers/AvaAutoSizeModifier.h
--- a/Engine/Plugins/Experimental/Avalanche/Source/AvalancheModifiers/Public/Modifiers/AvaAutoSizeModifier.h
+++ b/Engine/Plugins/Experimental/Avalanche/Source/AvalancheModifiers/Public/Modifiers/AvaAutoSizeModifier.h
@@ -4,7 +4,6 @@
 
 #include "AvaDefs.h"
 #include "AvaGeometryBaseModifier.h"
-#include "AvaModifiersActorUtils.h"
 #include "Extensions/AvaRenderStateUpdateModifierExtension.h"
 #include "Extensions/AvaSceneTreeUpdateModifierExtension.h"
 #include "Extensions/AvaTransformUpdateModifierExtension.h"
@@ -23,7 +22,7 @@ enum class EAvaAutoSizeFitMode : uint8
 };
 
 /**
- * Adapts the modified actor geometry size/scale and position so that it acts as a background for a specified actor
+ * Adapts the modified actor geometry size/scale to match reference actor bounds and act as a background
  */
 UCLASS(MinimalAPI, BlueprintType)
 class UAvaAutoSizeModifier : public UAvaGeometryBaseModifier
@@ -44,23 +43,21 @@ public:
 	}
 
 	UFUNCTION(BlueprintCallable, Category="Motion Design|Modifiers|AutoSize")
-	AVALANCHEMODIFIERS_API void SetFollowedAxis(int32 InFollowedAxis);
+	AVALANCHEMODIFIERS_API void SetPaddingHorizontal(double InPadding);
 
 	UFUNCTION(BlueprintPure, Category="Motion Design|Modifiers|AutoSize")
-	int32 GetFollowedAxis() const
+	double GetPaddingHorizontal() const
 	{
-		return FollowedAxis;
+		return PaddingHorizontal;
 	}
 
-	/** Sets the actor affecting the modifier. This is user selectable if the Reference Container is set to "Other". */
 	UFUNCTION(BlueprintCallable, Category="Motion Design|Modifiers|AutoSize")
-	AVALANCHEMODIFIERS_API void SetPadding(const FMargin& InPadding);
+	AVALANCHEMODIFIERS_API void SetPaddingVertical(double InPadding);
 
-	/** Gets the actor affecting the modifier. This is user selectable if the Reference Container is set to "Other". */
 	UFUNCTION(BlueprintPure, Category="Motion Design|Modifiers|AutoSize")
-	const FMargin& GetPadding() const
+	double GetPaddingVertical() const
 	{
-		return Padding;
+		return PaddingVertical;
 	}
 
 	UFUNCTION(BlueprintCallable, Category="Motion Design|Modifiers|AutoSize")
@@ -133,16 +130,13 @@ protected:
 	UPROPERTY(meta=(DeprecatedProperty, DeprecationMessage="Use ReferenceActor instead"))
 	bool bIgnoreHiddenActors_DEPRECATED = false;
 
-	/** Which axis should we follow, if none selected, it will not follow */
-	UPROPERTY(EditInstanceOnly, Setter="SetFollowedAxis", Getter="GetFollowedAxis", Category="AutoSize", meta=(Bitmask, BitmaskEnum="/Script/AvalancheModifiers.EAvaModifiersAxis", AllowPrivateAccess="true"))
-	int32 FollowedAxis = static_cast<int32>(
-		EAvaModifiersAxis::Y |
-		EAvaModifiersAxis::Z
-	);
+	/** Padding for top and bottom side */
+	UPROPERTY(EditInstanceOnly, Setter, Getter, Interp, Category="AutoSize", meta=(AllowPrivateAccess="true"))
+	double PaddingVertical = 0.f;
 
-	/* Padding around reference bounds */
-	UPROPERTY(EditInstanceOnly, Setter="SetPadding", Getter="GetPadding", Interp, Category="AutoSize", meta=(AllowPrivateAccess="true"))
-	FMargin Padding = FMargin(0.f);
+	/** Padding for left and right side */
+	UPROPERTY(EditInstanceOnly, Setter, Getter, Interp, Category="AutoSize", meta=(AllowPrivateAccess="true"))
+	double PaddingHorizontal = 0.f;
 
 	UPROPERTY(EditInstanceOnly, Setter="SetFitMode", Getter="GetFitMode", Category="AutoSize", meta=(AllowPrivateAccess="true"))
 	EAvaAutoSizeFitMode FitMode = EAvaAutoSizeFitMode::WidthAndHeight;
@@ -152,6 +146,10 @@ protected:
 	bool bIncludeChildren = true;
 
 private:
+	/** Padding added around reference actor bounds for geometry */
+	UPROPERTY()
+	FMargin Padding = FMargin(0.f);
+
 	UPROPERTY()
 	FVector2D PreModifierShapeDynMesh2DSize;
 
