MotionDesign : Modifiers
- Fixed boolean modifier mode not giving the expected geometry output result, batching tools mode operation (subtract, intersect, union) instead of applying them one by one on the target

diff --git a/Engine/Plugins/Experimental/Avalanche/Source/AvalancheModifiers/Private/Modifiers/AvaBooleanModifier.cpp b/Engine/Plugins/Experimental/Avalanche/Source/AvalancheModifiers/Private/Modifiers/AvaBooleanModifier.cpp
--- a/Engine/Plugins/Experimental/Avalanche/Source/AvalancheModifiers/Private/Modifiers/AvaBooleanModifier.cpp
+++ b/Engine/Plugins/Experimental/Avalanche/Source/AvalancheModifiers/Private/Modifiers/AvaBooleanModifier.cpp
@@ -8,6 +8,7 @@
 #include "DynamicMesh/MeshTransforms.h"
 #include "DynamicMeshes/AvaShapeDynMeshBase.h"
 #include "Extensions/AvaTransformUpdateModifierExtension.h"
+#include "GeometryScript/MeshBasicEditFunctions.h"
 #include "Materials/MaterialInstanceDynamic.h"
 #include "Modifiers/ActorModifierCoreStack.h"
 #include "Operations/MeshBoolean.h"
@@ -329,23 +330,83 @@ void UAvaBooleanModifier::ApplyInternal()
 
 	// find other colliding shapes
 	TSet<TWeakObjectPtr<UAvaBooleanModifier>> OutCollidingModifiersWeak = Shared->GetIntersectingModifiers(this, &ChannelInfo);
-	for (const TWeakObjectPtr<UAvaBooleanModifier>& CollidingModifierWeak : OutCollidingModifiersWeak)
-	{
-		UAvaBooleanModifier* CollidingModifier = CollidingModifierWeak.Get();
-		if (!CollidingModifier)
-		{
-			continue;
-		}
 
+	if (!OutCollidingModifiersWeak.IsEmpty())
+	{
 		if (bIsMasking)
 		{
-			// Mark other dirty to restore and reapply this mask
-			CollidingModifier->MarkModifierDirty();
+			for (const TWeakObjectPtr<UAvaBooleanModifier>& CollidingModifierWeak : OutCollidingModifiersWeak)
+			{
+				UAvaBooleanModifier* CollidingModifier = CollidingModifierWeak.Get();
+				if (!CollidingModifier)
+				{
+					continue;
+				}
+
+				// Mark other dirty to restore and reapply this mask
+				CollidingModifier->MarkModifierDirty();
+			}
 		}
 		else
 		{
-			// Compute colliding mask with this modifier
-			UAvaBooleanModifier::MaskActor(CollidingModifier, this);
+			using namespace UE::Geometry;
+
+			TMap<EAvaBooleanMode, FDynamicMesh3> ModeTools;
+
+			// Build the mode tools
+			for (const TWeakObjectPtr<UAvaBooleanModifier>& CollidingModifierWeak : OutCollidingModifiersWeak)
+			{
+				const UAvaBooleanModifier* CollidingModifier = CollidingModifierWeak.Get();
+				if (!CollidingModifier)
+				{
+					continue;
+				}
+
+				const EAvaBooleanMode ToolMode = CollidingModifier->GetMode();
+				if (ToolMode == EAvaBooleanMode::None)
+				{
+					continue;
+				}
+
+				const UDynamicMeshComponent* ToolMeshComponent = CollidingModifier->GetMeshComponent();
+				if (!ToolMeshComponent)
+				{
+					continue;
+				}
+
+				FDynamicMesh3& ModeTool = ModeTools.FindOrAdd(ToolMode);
+
+				ToolMeshComponent->ProcessMesh([ToolMeshComponent, &ModeTool](const FDynamicMesh3& InToolMesh)
+				{
+					FDynamicMesh3 ToolMesh = InToolMesh;
+					MeshTransforms::ApplyTransform(ToolMesh, ToolMeshComponent->GetComponentTransform());
+
+					FGeometryScriptAppendMeshOptions AppendOptions;
+					AppendOptions.CombineMode = EGeometryScriptCombineAttributesMode::EnableAllMatching;
+					FMeshIndexMappings TmpMappings;
+					AppendOptions.UpdateAttributesForCombineMode(ModeTool, ToolMesh);
+
+					FDynamicMeshEditor Editor(&ModeTool);
+					Editor.AppendMesh(&ToolMesh, TmpMappings);
+				});
+			}
+
+			FDynamicMesh3 OutputMesh;
+			GetMeshComponent()->ProcessMesh([&OutputMesh](const FDynamicMesh3& InEditMesh)
+			{
+				OutputMesh = InEditMesh;
+			});
+
+			// Apply mode tools
+			for (const TPair<EAvaBooleanMode, FDynamicMesh3>& ModeToolPair : ModeTools)
+			{
+				UAvaBooleanModifier::ApplyTool(ModeToolPair.Key, ModeToolPair.Value, this, OutputMesh);
+			}
+
+			GetMeshComponent()->EditMesh([&OutputMesh](FDynamicMesh3& InEditMesh)
+			{
+				InEditMesh = MoveTemp(OutputMesh);
+			});
 		}
 	}
 
@@ -524,65 +585,67 @@ void UAvaBooleanModifier::UpdateMaskVisibility()
 	}
 }
 
-void UAvaBooleanModifier::MaskActor(const UAvaBooleanModifier* InTool, const UAvaBooleanModifier* InTarget)
+void UAvaBooleanModifier::ApplyTool(EAvaBooleanMode InMode, const UE::Geometry::FDynamicMesh3& InModeTool, const UAvaBooleanModifier* InTarget, UE::Geometry::FDynamicMesh3& OutMesh)
 {
-	UDynamicMeshComponent* ToolDynMesh = InTool->GetMeshComponent();
-	UDynamicMeshComponent* TargetDynMesh = InTarget->GetMeshComponent();
-
-	if (!ToolDynMesh || !TargetDynMesh)
-	{
-		return;
-	}
-
-	EAvaBooleanMode ToolMode = InTool->Mode;
-
-	if (ToolMode == EAvaBooleanMode::None || InTarget->GetMode() != EAvaBooleanMode::None)
+	if (!IsValid(InTarget)
+		|| !InTarget->GetMeshComponent()
+		|| InTarget->GetMode() != EAvaBooleanMode::None)
 	{
 		return;
 	}
 
 	using namespace UE::Geometry;
 
-	TargetDynMesh->EditMesh([ToolDynMesh, TargetDynMesh, ToolMode](FDynamicMesh3& TargetMesh)
+	UDynamicMeshComponent* TargetDynMesh = InTarget->GetMeshComponent();
+
+	TargetDynMesh->ProcessMesh([TargetDynMesh, InMode, &InModeTool, &OutMesh](const FDynamicMesh3& InSourceMesh)
 	{
-		FTransformSRT3d TargetTransform(TargetDynMesh->GetComponentTransform());
-		FTransformSRT3d ToolTransform(ToolDynMesh->GetComponentTransform());
+		FTransformSRT3d SourceTransform(TargetDynMesh->GetComponentTransform());
 
-		ToolDynMesh->ProcessMesh([&TargetMesh, ToolMode, TargetTransform, ToolTransform](const FDynamicMesh3& ToolMesh)
+		FMeshBoolean::EBooleanOp Operation;
+		switch (InMode)
 		{
-			FMeshBoolean::EBooleanOp Operation;
-			switch (ToolMode)
+			case EAvaBooleanMode::Intersect:
 			{
-				case EAvaBooleanMode::Intersect:
-					Operation = FMeshBoolean::EBooleanOp::Intersect;
-				break;
-				case EAvaBooleanMode::Subtract:
-					Operation = FMeshBoolean::EBooleanOp::Difference;
-				break;
-				case EAvaBooleanMode::Union:
-					Operation = FMeshBoolean::EBooleanOp::Union;
-				break;
-				default:
-					return;
-				}
+				Operation = FMeshBoolean::EBooleanOp::Intersect;
+			}
+			break;
+
+			case EAvaBooleanMode::Subtract:
+			{
+				Operation = FMeshBoolean::EBooleanOp::Difference;
+			}
+			break;
+
+			case EAvaBooleanMode::Union:
+			{
+				Operation = FMeshBoolean::EBooleanOp::Union;
+			}
+			break;
+
+			default: return;
+		}
 
-				FMeshBoolean MeshBoolean(
-				&TargetMesh, TargetTransform,
-				&ToolMesh, ToolTransform,
-					&TargetMesh, Operation);
-				MeshBoolean.bPutResultInInputSpace = true;
-				MeshBoolean.bSimplifyAlongNewEdges = true;
-				MeshBoolean.bWeldSharedEdges = true;
-				MeshBoolean.bCollapseDegenerateEdgesOnCut = true;
-				MeshBoolean.bPreserveTriangleGroups = true;
-				MeshBoolean.bTrackAllNewEdges = false;
-				MeshBoolean.Compute();
-		});
-
-		if (TargetMesh.TriangleCount() > 0)
+		FMeshBoolean MeshBoolean(
+		&OutMesh
+			, SourceTransform
+			, &InModeTool
+			, FTransform::Identity
+			, &OutMesh
+			, Operation
+		);
+		MeshBoolean.bPutResultInInputSpace = true;
+		MeshBoolean.bSimplifyAlongNewEdges = true;
+		MeshBoolean.bWeldSharedEdges = true;
+		MeshBoolean.bCollapseDegenerateEdgesOnCut = true;
+		MeshBoolean.bPreserveTriangleGroups = true;
+		MeshBoolean.bTrackAllNewEdges = false;
+		MeshBoolean.Compute();
+
+		if (OutMesh.TriangleCount() > 0)
 		{
 			// Boolean result is in the space of TargetTransform, so invert that
-			MeshTransforms::ApplyTransformInverse(TargetMesh, TargetTransform, true);
+			MeshTransforms::ApplyTransformInverse(OutMesh, SourceTransform, true);
 		}
 	});
 }
diff --git a/Engine/Plugins/Experimental/Avalanche/Source/AvalancheModifiers/Private/Shared/AvaBooleanModifierShared.cpp b/Engine/Plugins/Experimental/Avalanche/Source/AvalancheModifiers/Private/Shared/AvaBooleanModifierShared.cpp
--- a/Engine/Plugins/Experimental/Avalanche/Source/AvalancheModifiers/Private/Shared/AvaBooleanModifierShared.cpp
+++ b/Engine/Plugins/Experimental/Avalanche/Source/AvalancheModifiers/Private/Shared/AvaBooleanModifierShared.cpp
@@ -40,7 +40,7 @@ TSet<TWeakObjectPtr<UAvaBooleanModifier>> UAvaBooleanModifierShared::GetIntersec
 		OutDesc->ChannelTargetCount = 0;
 		OutDesc->ChannelIntersectCount = 0;
 	}
-			
+
 	for (const TWeakObjectPtr<UAvaBooleanModifier>& OtherModifierWeak : Channel->ModifiersWeak)
 	{
 		UAvaBooleanModifier* OtherModifier = OtherModifierWeak.Get();
@@ -62,7 +62,7 @@ TSet<TWeakObjectPtr<UAvaBooleanModifier>> UAvaBooleanModifierShared::GetIntersec
 				OutDesc->ChannelToolCount++;
 			}
 		}
-		
+
 		const AActor* OtherActor = OtherModifier->GetModifiedActor();
 		if (!OtherActor || TargetActor == OtherActor)
 		{
@@ -76,7 +76,7 @@ TSet<TWeakObjectPtr<UAvaBooleanModifier>> UAvaBooleanModifierShared::GetIntersec
 		}
 
 		const FTransform OtherTransform = OtherModifier->GetMeshComponent()->GetComponentTransform();
-		
+
 		// Other is mask
 		bool bIsIntersecting = false;
 		if (bOtherBooleanMode)
@@ -98,11 +98,11 @@ TSet<TWeakObjectPtr<UAvaBooleanModifier>> UAvaBooleanModifierShared::GetIntersec
 				bIsIntersecting = TestIntersection(InTargetMesh, TargetTransform, OtherMesh, OtherTransform);
 			});
 		}
-		
+
 		if (bIsIntersecting)
 		{
 			IntersectingModifiers.Add(OtherModifier);
-			
+
 			if (OutDesc)
 			{
 				OutDesc->ChannelIntersectCount++;
@@ -114,7 +114,21 @@ TSet<TWeakObjectPtr<UAvaBooleanModifier>> UAvaBooleanModifierShared::GetIntersec
 	{
 		OutDesc->ChannelCount = GetChannelCount();
 	}
-	
+
+	// Sort modifier based on mode priority
+	static const TMap<EAvaBooleanMode, uint32> ModePriority
+	{
+		{EAvaBooleanMode::None, 0},
+		{EAvaBooleanMode::Intersect, 1},
+		{EAvaBooleanMode::Union, 2},
+		{EAvaBooleanMode::Subtract, 3}
+	};
+
+	IntersectingModifiers.StableSort([](const TWeakObjectPtr<UAvaBooleanModifier>& InModifierA, const TWeakObjectPtr<UAvaBooleanModifier>& InModifierB)
+	{
+		return ModePriority[InModifierA->GetMode()] < ModePriority[InModifierB->GetMode()];
+	});
+
 	return IntersectingModifiers;
 }
 
@@ -126,7 +140,7 @@ bool UAvaBooleanModifierShared::TestIntersection(const UE::Geometry::FDynamicMes
 	{
 		return false;
 	}
-	
+
 	FDynamicMeshAABBTree3 Spatials[2];
 	ParallelFor(2, [&Spatials, &InToolMesh, &InTargetMesh](int32 k)
 	{
@@ -138,12 +152,12 @@ bool UAvaBooleanModifierShared::TestIntersection(const UE::Geometry::FDynamicMes
 
 	const bool bIsIdentity2 = InTargetTransform.Equals(FTransform::Identity, 0);
 	FTransformSRT3d Transform2(InTargetTransform);
-	
+
 	bool bIsIntersecting = false;
 	if (bIsIdentity1 && bIsIdentity2)
 	{
 		bIsIntersecting = Spatials[0].TestIntersection(Spatials[1]);
-	} 
+	}
 	else if (bIsIdentity1 || bIsIdentity2)
 	{
 		const FIndex2i Indices = (bIsIdentity1) ? FIndex2i(0,1) : FIndex2i(1,0);
@@ -162,7 +176,7 @@ bool UAvaBooleanModifierShared::TestIntersection(const UE::Geometry::FDynamicMes
 	{
 		const FBox ToolBounds = static_cast<FBox>(InToolMesh.GetBounds(true));
 		const FVector ActorOffset = InToolTransform.GetLocation() - InTargetTransform.GetLocation();
-		
+
 		if (ToolBounds.IsInside(ActorOffset))
 		{
 			bIsIntersecting = true;
@@ -178,7 +192,7 @@ void UAvaBooleanModifierShared::TrackModifierChannel(UAvaBooleanModifier* InModi
 	{
 		return;
 	}
-	
+
 	FAvaBooleanModifierSharedChannel& Channel = Channels.FindOrAdd(InModifier->GetChannel());
 	Channel.ModifiersWeak.Add(InModifier);
 }
@@ -189,7 +203,7 @@ void UAvaBooleanModifierShared::UntrackModifierChannel(UAvaBooleanModifier* InMo
 	{
 		return;
 	}
-	
+
 	for (TMap<uint8, FAvaBooleanModifierSharedChannel>::TIterator It(Channels); It; ++It)
 	{
 		It->Value.ModifiersWeak.Remove(InModifier);
@@ -208,14 +222,14 @@ void UAvaBooleanModifierShared::UpdateModifierChannel(UAvaBooleanModifier* InMod
 	{
 		return;
 	}
-	
+
 	const uint8 Channel = InModifier->GetChannel();
 	for (TMap<uint8, FAvaBooleanModifierSharedChannel>::TIterator It(Channels); It; ++It)
 	{
 		if (Channel != It->Key)
 		{
 			It->Value.ModifiersWeak.Remove(InModifier);
-			
+
 			// Remove empty channel
 			if (It->Value.ModifiersWeak.IsEmpty())
 			{
diff --git a/Engine/Plugins/Experimental/Avalanche/Source/AvalancheModifiers/Public/Modifiers/AvaBooleanModifier.h b/Engine/Plugins/Experimental/Avalanche/Source/AvalancheModifiers/Public/Modifiers/AvaBooleanModifier.h
--- a/Engine/Plugins/Experimental/Avalanche/Source/AvalancheModifiers/Public/Modifiers/AvaBooleanModifier.h
+++ b/Engine/Plugins/Experimental/Avalanche/Source/AvalancheModifiers/Public/Modifiers/AvaBooleanModifier.h
@@ -98,8 +98,7 @@ protected:
 	void UpdateMaskDelegates();
 	void UpdateMaskVisibility();
 
-	/** When masking target with tool, tool must have a depth */
-	static void MaskActor(const UAvaBooleanModifier* InTool, const UAvaBooleanModifier* InTarget);
+	static void ApplyTool(EAvaBooleanMode InMode, const UE::Geometry::FDynamicMesh3& InModeTool, const UAvaBooleanModifier* InTarget, UE::Geometry::FDynamicMesh3& OutMesh);
 
 	/** Mode to use when shapes are colliding, none means you will be masked otherwise you are masking */
 	UPROPERTY(EditInstanceOnly, Setter="SetMode", Getter="GetMode", Category="Boolean", meta=(AllowPrivateAccess="true"))
