MotionDesign : AvaShape
- Fixed crash happening when async loading level and loading materials

diff --git a/Engine/Plugins/Experimental/Avalanche/Source/AvalancheShapes/Private/AvaShapeParametricMaterial.cpp b/Engine/Plugins/Experimental/Avalanche/Source/AvalancheShapes/Private/AvaShapeParametricMaterial.cpp
--- a/Engine/Plugins/Experimental/Avalanche/Source/AvalancheShapes/Private/AvaShapeParametricMaterial.cpp
+++ b/Engine/Plugins/Experimental/Avalanche/Source/AvalancheShapes/Private/AvaShapeParametricMaterial.cpp
@@ -22,8 +22,6 @@ FAvaShapeParametricMaterial::FAvaShapeParametricMaterial()
 	ColorB.A = 1.f;
 
 	GradientOffset = 0.5f;
-
-	LoadDefaultMaterials();
 }
 
 FAvaShapeParametricMaterial::FAvaShapeParametricMaterial(const FAvaShapeParametricMaterial& Other)
@@ -58,46 +56,40 @@ FAvaShapeParametricMaterial& FAvaShapeParametricMaterial::operator=(const FAvaSh
 
 UMaterialInterface* FAvaShapeParametricMaterial::GetDefaultMaterial() const
 {
+	LoadDefaultMaterials();
+
 	const int32 Index = GetActiveInstanceIndex();
 	check(DefaultMaterials.IsValidIndex(Index))
 	return DefaultMaterials[Index];
 }
 
-void FAvaShapeParametricMaterial::LoadDefaultMaterials()
+void FAvaShapeParametricMaterial::LoadDefaultMaterials() const
 {
-	static const TCHAR* LitOpaqueMaterialPath = TEXT("/Avalanche/ToolboxResources/M_Toolbox_Opaque_Lit.M_Toolbox_Opaque_Lit");
-	UMaterialInterface* DefaultLitMaterialOpaque = FindObject<UMaterialInterface>(nullptr, LitOpaqueMaterialPath);
-	if (!DefaultLitMaterialOpaque)
+	if (!DefaultMaterials.IsEmpty())
 	{
-		DefaultLitMaterialOpaque = LoadObject<UMaterialInterface>(nullptr, LitOpaqueMaterialPath);
+		return;
 	}
-	DefaultMaterials.Add(DefaultLitMaterialOpaque);
 
-	static const TCHAR* LitTranslucentMaterialPath = TEXT("/Avalanche/ToolboxResources/M_Toolbox_Translucent_Lit.M_Toolbox_Translucent_Lit");
-	UMaterialInterface* DefaultLitMaterialTranslucent = FindObject<UMaterialInterface>(nullptr, LitTranslucentMaterialPath);
-	if (!DefaultLitMaterialTranslucent)
+	if (FAvaShapeParametricMaterial* MutableThis = const_cast<FAvaShapeParametricMaterial*>(this))
 	{
-		DefaultLitMaterialTranslucent = LoadObject<UMaterialInterface>(nullptr, LitTranslucentMaterialPath);
-	}
-	DefaultMaterials.Add(DefaultLitMaterialTranslucent);
+		static const TCHAR* LitOpaqueMaterialPath = TEXT("/Avalanche/ToolboxResources/M_Toolbox_Opaque_Lit.M_Toolbox_Opaque_Lit");
+		UMaterialInterface* DefaultLitMaterialOpaque = LoadObject<UMaterialInterface>(nullptr, LitOpaqueMaterialPath);
+		MutableThis->DefaultMaterials.Add(DefaultLitMaterialOpaque);
 
-	static const TCHAR* UnlitOpaqueMaterialPath = TEXT("/Avalanche/ToolboxResources/M_Toolbox_Opaque_Unlit.M_Toolbox_Opaque_Unlit");
-	UMaterialInterface* DefaultUnlitMaterialOpaque = FindObject<UMaterialInterface>(nullptr, UnlitOpaqueMaterialPath);
-	if (!DefaultUnlitMaterialOpaque)
-	{
-		DefaultUnlitMaterialOpaque = LoadObject<UMaterialInterface>(nullptr, UnlitOpaqueMaterialPath);
-	}
-	DefaultMaterials.Add(DefaultUnlitMaterialOpaque);
+		static const TCHAR* LitTranslucentMaterialPath = TEXT("/Avalanche/ToolboxResources/M_Toolbox_Translucent_Lit.M_Toolbox_Translucent_Lit");
+		UMaterialInterface* DefaultLitMaterialTranslucent = LoadObject<UMaterialInterface>(nullptr, LitTranslucentMaterialPath);
+		MutableThis->DefaultMaterials.Add(DefaultLitMaterialTranslucent);
 
-	static const TCHAR* UnlitTranslucentMaterialPath = TEXT("/Avalanche/ToolboxResources/M_Toolbox_Translucent_Unlit.M_Toolbox_Translucent_Unlit");
-	UMaterialInterface* DefaultUnlitMaterialTranslucent = FindObject<UMaterialInterface>(nullptr, UnlitTranslucentMaterialPath);
-	if (!DefaultUnlitMaterialTranslucent)
-	{
-		DefaultUnlitMaterialTranslucent = LoadObject<UMaterialInterface>(nullptr, UnlitTranslucentMaterialPath);
-	}
-	DefaultMaterials.Add(DefaultUnlitMaterialTranslucent);
+		static const TCHAR* UnlitOpaqueMaterialPath = TEXT("/Avalanche/ToolboxResources/M_Toolbox_Opaque_Unlit.M_Toolbox_Opaque_Unlit");
+		UMaterialInterface* DefaultUnlitMaterialOpaque = LoadObject<UMaterialInterface>(nullptr, UnlitOpaqueMaterialPath);
+		MutableThis->DefaultMaterials.Add(DefaultUnlitMaterialOpaque);
 
-	InstanceMaterials.SetNum(DefaultMaterials.Num(), EAllowShrinking::No);
+		static const TCHAR* UnlitTranslucentMaterialPath = TEXT("/Avalanche/ToolboxResources/M_Toolbox_Translucent_Unlit.M_Toolbox_Translucent_Unlit");
+		UMaterialInterface* DefaultUnlitMaterialTranslucent = LoadObject<UMaterialInterface>(nullptr, UnlitTranslucentMaterialPath);
+		MutableThis->DefaultMaterials.Add(DefaultUnlitMaterialTranslucent);
+
+		MutableThis->InstanceMaterials.SetNum(DefaultMaterials.Num(), EAllowShrinking::No);
+	}
 }
 
 void FAvaShapeParametricMaterial::SetMaterialParameterValues(UMaterialInstanceDynamic* InMaterialInstance, bool bInNotifyUpdate) const
@@ -175,15 +167,15 @@ UMaterialInstanceDynamic* FAvaShapeParametricMaterial::CreateMaterialInstance(UO
 		return nullptr;
 	}
 
-	const UMaterialInterface* ParentMaterial = GetDefaultMaterial();
+	UMaterialInterface* ParentMaterial = GetDefaultMaterial();
 
 	if (!ParentMaterial)
 	{
-		LoadDefaultMaterials();
+		return nullptr;
 	}
 
 	UMaterialInstanceDynamic* MaterialInstance = UMaterialInstanceDynamic::Create(
-		GetDefaultMaterial(),
+		ParentMaterial,
 		InOuter
 	);
 
@@ -284,8 +276,9 @@ bool FAvaShapeParametricMaterial::IsParametricMaterial(UMaterialInterface* InMat
 
 UMaterialInstanceDynamic* FAvaShapeParametricMaterial::GetMaterial() const
 {
-	const int32 Index = GetActiveInstanceIndex();
+	LoadDefaultMaterials();
 
+	const int32 Index = GetActiveInstanceIndex();
 	check(InstanceMaterials.IsValidIndex(Index))
 	UMaterialInstanceDynamic* Material = InstanceMaterials[Index];
 
diff --git a/Engine/Plugins/Experimental/Avalanche/Source/AvalancheShapes/Public/AvaShapeParametricMaterial.h b/Engine/Plugins/Experimental/Avalanche/Source/AvalancheShapes/Public/AvaShapeParametricMaterial.h
--- a/Engine/Plugins/Experimental/Avalanche/Source/AvalancheShapes/Public/AvaShapeParametricMaterial.h
+++ b/Engine/Plugins/Experimental/Avalanche/Source/AvalancheShapes/Public/AvaShapeParametricMaterial.h
@@ -164,7 +164,7 @@ private:
 	int32 ActiveInstanceIndex = INDEX_NONE;
 
 	/** Load parents materials to create instance materials */
-	void LoadDefaultMaterials();
+	void LoadDefaultMaterials() const;
 
 	/** Create an instance material based on the current active parent */
 	UMaterialInstanceDynamic* CreateMaterialInstance(UObject* InOuter);
