MotionDesign : Modifiers
- Fixed Plane Cut Modifier preview plane is z-fighting with the actual mesh

diff --git a/Engine/Plugins/Experimental/Avalanche/Source/AvalancheModifiers/Private/Modifiers/AvaPlaneCutModifier.cpp b/Engine/Plugins/Experimental/Avalanche/Source/AvalancheModifiers/Private/Modifiers/AvaPlaneCutModifier.cpp
--- a/Engine/Plugins/Experimental/Avalanche/Source/AvalancheModifiers/Private/Modifiers/AvaPlaneCutModifier.cpp
+++ b/Engine/Plugins/Experimental/Avalanche/Source/AvalancheModifiers/Private/Modifiers/AvaPlaneCutModifier.cpp
@@ -97,7 +97,7 @@ void UAvaPlaneCutModifier::UpdatePreviewComponent()
 
 	PreviewPlane.Update(FTransform(
 		PlaneRotation,
-		PlaneLocation,
+		PlaneLocation + PlaneRotation.RotateVector(FVector::UpVector) * 1,
 		FVector::Max(PlaneMeshSize, FVector(0.1f, 0.1f, 0.1f)))
 	);
 	
