[Motion Design] Rundown: Fixing transitions between Combo Pages and Single Pages in reuse instancing mode.

This loosens the template matching. It no longer requires that both templates be either combo or single. One can be a combo and the other can be single. It will always try to find the requested sub-template.

diff --git a/Engine/Plugins/Experimental/Avalanche/Source/AvalancheMedia/Private/Rundown/AvaRundown.cpp b/Engine/Plugins/Experimental/Avalanche/Source/AvalancheMedia/Private/Rundown/AvaRundown.cpp
--- a/Engine/Plugins/Experimental/Avalanche/Source/AvalancheMedia/Private/Rundown/AvaRundown.cpp
+++ b/Engine/Plugins/Experimental/Avalanche/Source/AvalancheMedia/Private/Rundown/AvaRundown.cpp
@@ -1930,6 +1930,13 @@ namespace UE::AvaMedia::Rundown::Private
 		bool bInIsPreview,
 		const FName& InPreviewChannelName)
 	{
+		const FAvaRundownPage& SubTemplate = InTemplate.GetTemplate(InRundown, InSubPageIndex);
+
+		if (!SubTemplate.IsValidPage())
+		{
+			return nullptr;
+		}
+
 		for (const TObjectPtr<UAvaRundownPagePlayer>& PagePlayer : InRundown->GetPagePlayers())
 		{
 			// Early filter on preview/channel.
@@ -1942,17 +1949,37 @@ namespace UE::AvaMedia::Rundown::Private
 			
 			const FAvaRundownPage& PlayingPage = InRundown->GetPage(PagePlayer->PageId);
 
-			// Check if same template.
-			if (!PlayingPage.IsValidPage() || PlayingPage.GetTemplateId() != InTemplate.GetPageId())
+			if (!PlayingPage.IsValidPage())
+			{
+				continue;
+			}
+
+			const FAvaRundownPage& PlayingTemplate = PlayingPage.ResolveTemplate(InRundown);
+
+			if (!PlayingTemplate.IsValidPage())
 			{
 				continue;
 			}
 			
-			// Find Instance Player for the given sub-template.
-			const FAvaRundownPage& SubTemplate = InTemplate.GetTemplate(InRundown, InSubPageIndex);
-			if (SubTemplate.IsValidPage())
+			// Check if we have a corresponding template.
+			if (PlayingTemplate.IsComboTemplate())
+			{
+				if(!PlayingTemplate.GetCombinedTemplateIds().Contains(SubTemplate.GetPageId()))
+				{
+					continue;
+				}
+			}
+			else if (PlayingTemplate.GetPageId() != SubTemplate.GetPageId())
+			{
+				continue;
+			}
+			
+			// Find Instance Player for the given sub-template.			
+			// Remark: if not found, keep looking. With "reuse" instancing mode, instance players can bounce from combo to single and back to combo.
+			UAvaRundownPlaybackInstancePlayer* InstancePlayer = PagePlayer->FindInstancePlayerByAssetPath(SubTemplate.GetAssetPath(InRundown));
+			if (InstancePlayer && InstancePlayer->PlaybackInstance)
 			{
-				return PagePlayer->FindInstancePlayerByAssetPath(SubTemplate.GetAssetPath(InRundown));
+				return InstancePlayer;
 			}
 		}
 
