[Motion Design] Rundown Server Fixes:
- Fixing Transient Rundowns delete
- Fixing Create Template returns the Id of the created template.

diff --git a/Engine/Plugins/Experimental/Avalanche/Source/AvalancheMedia/Private/Rundown/AvaRundownServer.cpp b/Engine/Plugins/Experimental/Avalanche/Source/AvalancheMedia/Private/Rundown/AvaRundownServer.cpp
--- a/Engine/Plugins/Experimental/Avalanche/Source/AvalancheMedia/Private/Rundown/AvaRundownServer.cpp
+++ b/Engine/Plugins/Experimental/Avalanche/Source/AvalancheMedia/Private/Rundown/AvaRundownServer.cpp
@@ -830,7 +830,8 @@ void FAvaRundownServer::HandleCreateRundown(const FAvaRundownCreateRundown& InMe
 	}
 
 	constexpr EObjectFlags AssetFlags = RF_Public|RF_Standalone|RF_Transactional;
-	UAvaRundown* Rundown = NewObject<UAvaRundown>(RundownPackage, FName(*InMessage.AssetName), AssetFlags);
+	constexpr EObjectFlags TransientFlags = RF_Public|RF_Transactional;
+	UAvaRundown* Rundown = NewObject<UAvaRundown>(RundownPackage, FName(*InMessage.AssetName), bTransient ? TransientFlags : AssetFlags);
 
 	if (!bTransient)
 	{
@@ -877,6 +878,7 @@ void FAvaRundownServer::HandleDeleteRundown(const FAvaRundownDeleteRundown& InMe
 		// Also, flush command contexts if associated to this rundown.
 		EditCommandContext.ConditionalFlush(SharedThis(this), RundownPath);
 		PlaybackCommandContext.ConditionalFlush(SharedThis(this), RundownPath);
+		CollectGarbage( GARBAGE_COLLECTION_KEEPFLAGS );
 		
 		SendMessage(InContext->GetSender(), InMessage.RequestId, ELogVerbosity::Log,
 			TEXT("\"DeleteRundown\": Rundown \"%s\" removed."), *InMessage.Rundown);
@@ -1188,11 +1190,17 @@ void FAvaRundownServer::HandleCreateTemplate(const FAvaRundownCreateTemplate& In
 		return; // Response sent by GetOrLoadRundownForEdit.
 	}
 	
-	Rundown->AddTemplate();
-	SendMessage(InContext->GetSender(), InMessage.RequestId, ELogVerbosity::Log, TEXT("Template Created"));
+	const int32 TemplateId = Rundown->AddTemplate();
+	if (TemplateId != FAvaRundownPage::InvalidPageId)
+	{
+		SendMessage(InContext->GetSender(), InMessage.RequestId, ELogVerbosity::Log, TEXT("Template %d Created"), TemplateId);
+	}
+	else
+	{
+		SendMessage(InContext->GetSender(), InMessage.RequestId, ELogVerbosity::Error, TEXT("Failed to create a new template"));
+	}
 }
 
-
 void FAvaRundownServer::HandleChangeTemplateBP(const FAvaRundownChangeTemplateBP& InMessage,
 	const TSharedRef<IMessageContext, ESPMode::ThreadSafe>& InContext)
 {
