Motion Design / Material Designer
- Mask (Apply) now works correctly with Material Designer materials in the editor.

diff --git a/Engine/Plugins/Experimental/Avalanche/Source/AvalancheMask/Private/Handling/AvaMaskDesignedMaterialHandle.cpp b/Engine/Plugins/Experimental/Avalanche/Source/AvalancheMask/Private/Handling/AvaMaskDesignedMaterialHandle.cpp
--- a/Engine/Plugins/Experimental/Avalanche/Source/AvalancheMask/Private/Handling/AvaMaskDesignedMaterialHandle.cpp
+++ b/Engine/Plugins/Experimental/Avalanche/Source/AvalancheMask/Private/Handling/AvaMaskDesignedMaterialHandle.cpp
@@ -104,6 +104,19 @@ const EBlendMode FAvaMaskDesignedMaterialHandle::GetBlendMode()
 {
 	if (const UMaterialInterface* MaterialInstance = GetMaterialInstance())
 	{
+#if WITH_EDITOR
+		if (const UDynamicMaterialInstance* DynamicMaterial = Cast<UDynamicMaterialInstance>(MaterialInstance))
+		{
+			if (UDynamicMaterialModel* MaterialModel = DynamicMaterial->GetMaterialModel())
+			{
+				if (UDynamicMaterialModelEditorOnlyData* ModelData = UDynamicMaterialModelEditorOnlyData::Get(MaterialModel))
+				{
+					return ModelData->GetBlendMode();
+				}
+			}
+		}
+#endif
+
 		return MaterialInstance->GetBlendMode();
 	}
 	
@@ -117,9 +130,9 @@ void FAvaMaskDesignedMaterialHandle::SetBlendMode(const EBlendMode InBlendMode)
 #if WITH_EDITOR
 		if (UDynamicMaterialInstance* MaterialInstance = WeakDesignedMaterial.Get())
 		{
-			if (UDynamicMaterialModel* Model = MaterialInstance->GetMaterialModel())
+			if (UDynamicMaterialModel* MaterialModel = MaterialInstance->GetMaterialModel())
 			{
-				if (UDynamicMaterialModelEditorOnlyData* ModelData = UDynamicMaterialModelEditorOnlyData::Get(Model))
+				if (UDynamicMaterialModelEditorOnlyData* ModelData = UDynamicMaterialModelEditorOnlyData::Get(MaterialModel))
 				{
 					ModelData->SetBlendMode(TargetBlendMode);
 				}
@@ -130,9 +143,6 @@ void FAvaMaskDesignedMaterialHandle::SetBlendMode(const EBlendMode InBlendMode)
 		{
 			MaterialInstance->BlendMode = TargetBlendMode;
 			MaterialInstance->BasePropertyOverrides.BlendMode = TargetBlendMode;
-#if WITH_EDITOR
-			MaterialInstance->UpdateCachedData();
-#endif
 		}
 #endif
 	}
@@ -154,9 +164,14 @@ bool FAvaMaskDesignedMaterialHandle::SaveOriginalState(const FStructView& InHand
 				, UE::AvaMask::Internal::PaddingParameterInfo
 				, UE::AvaMask::Internal::FeatherParameterInfo);
 
-			MaterialHandleData->OriginalOutputProcessor = GetOutputProcessor();
 			MaterialHandleData->OriginalBlendMode = GetBlendMode();
 
+			MaterialHandleData->OriginalOutputProcessor = GetOutputProcessor(
+				MaterialHandleData->OriginalBlendMode == BLEND_Translucent
+					? BLEND_Translucent
+					: BLEND_Masked
+			);
+
 			return true;
 		}
 	}
@@ -182,7 +197,7 @@ bool FAvaMaskDesignedMaterialHandle::ApplyOriginalState(const FStructView& InHan
 				, MaterialHandleData->OriginalMaskMaterialParameters.InnerFeatherRadius);
 
 #if WITH_EDITOR
-			SetOutputProcessor(MaterialHandleData->OriginalOutputProcessor.Get());
+			SetOutputProcessor(GetBlendMode(), MaterialHandleData->OriginalOutputProcessor.Get());
 #endif
 
 			SetBlendMode(MaterialHandleData->OriginalBlendMode);
@@ -213,13 +228,13 @@ bool FAvaMaskDesignedMaterialHandle::ApplyModifiedState(
 				, InModifiedParameters.MaterialParameters.OuterFeatherRadius
 				, InModifiedParameters.MaterialParameters.InnerFeatherRadius);
 
+			SetBlendMode(InModifiedParameters.MaterialParameters.BlendMode);
+
 #if WITH_EDITOR
 			UMaterialFunctionInterface* MaterialFunctionToUse = GetMutableDefault<UAvaMaskSettings>()->GetMaterialFunction();
-			SetOutputProcessor(MaterialFunctionToUse);
+			SetOutputProcessor(InModifiedParameters.MaterialParameters.BlendMode, MaterialFunctionToUse);
 #endif
 
-			SetBlendMode(InModifiedParameters.MaterialParameters.BlendMode);
-
 			LastAppliedParameters = InModifiedParameters;
 
 			return true;
@@ -283,7 +298,7 @@ void FAvaMaskDesignedMaterialHandle::OnMaterialBuilt(UDynamicMaterialModelBase*
 	SetBlendMode(LastAppliedParameters.MaterialParameters.BlendMode);
 }
 
-UMaterialFunctionInterface* FAvaMaskDesignedMaterialHandle::GetOutputProcessor()
+UMaterialFunctionInterface* FAvaMaskDesignedMaterialHandle::GetOutputProcessor(EBlendMode InBlendMode)
 {
 	if (UDynamicMaterialInstance* DesignedMaterial = WeakDesignedMaterial.Get())
 	{
@@ -302,8 +317,7 @@ UMaterialFunctionInterface* FAvaMaskDesignedMaterialHandle::GetOutputProcessor()
 				return nullptr;
 			};
 
-			const EBlendMode BlendMode = GetBlendMode();
-			switch (BlendMode)
+			switch (InBlendMode)
 			{
 			case EBlendMode::BLEND_Masked:
 				return GetOutputProcessorForPropertyType(Model, EDMMaterialPropertyType::OpacityMask);
@@ -312,7 +326,7 @@ UMaterialFunctionInterface* FAvaMaskDesignedMaterialHandle::GetOutputProcessor()
 				return GetOutputProcessorForPropertyType(Model, EDMMaterialPropertyType::Opacity);
 
 			default:
-				UE_LOG(LogAvaMask, Error, TEXT("BlendMode not supported: %s"), *UE::AvaMask::Internal::GetBlendModeString(BlendMode));
+				UE_LOG(LogAvaMask, Error, TEXT("BlendMode not supported: %s"), *UE::AvaMask::Internal::GetBlendModeString(InBlendMode));
 				return nullptr;
 			}
 		}
@@ -321,7 +335,7 @@ UMaterialFunctionInterface* FAvaMaskDesignedMaterialHandle::GetOutputProcessor()
 	return nullptr;
 }
 
-void FAvaMaskDesignedMaterialHandle::SetOutputProcessor(UMaterialFunctionInterface* InMaterialFunction)
+void FAvaMaskDesignedMaterialHandle::SetOutputProcessor(EBlendMode InBlendMode, UMaterialFunctionInterface* InMaterialFunction)
 {
 	if (UDynamicMaterialInstance* DesignedMaterial = WeakDesignedMaterial.Get())
 	{
@@ -350,8 +364,7 @@ void FAvaMaskDesignedMaterialHandle::SetOutputProcessor(UMaterialFunctionInterfa
 				}
 			};
 		
-			const EBlendMode BlendMode = UE::AvaMask::Internal::GetTargetBlendMode(DesignedMaterial->GetBlendMode(), GetBlendMode());
-			switch (BlendMode)
+			switch (InBlendMode)
 			{
 			case EBlendMode::BLEND_Masked:
 				SetOutputProcessorForPropertyType(Model, EDMMaterialPropertyType::OpacityMask);
@@ -362,7 +375,8 @@ void FAvaMaskDesignedMaterialHandle::SetOutputProcessor(UMaterialFunctionInterfa
 				break;
 
 			default:
-				UE_LOG(LogAvaMask, Error, TEXT("BlendMode not supported: %s"), *UE::AvaMask::Internal::GetBlendModeString(BlendMode));
+				UE_LOG(LogAvaMask, Error, TEXT("BlendMode not supported: %s"), *UE::AvaMask::Internal::GetBlendModeString(InBlendMode));
+				break;
 			}
 		}
 	}
diff --git a/Engine/Plugins/Experimental/Avalanche/Source/AvalancheMask/Private/Handling/AvaMaskDesignedMaterialHandle.h b/Engine/Plugins/Experimental/Avalanche/Source/AvalancheMask/Private/Handling/AvaMaskDesignedMaterialHandle.h
--- a/Engine/Plugins/Experimental/Avalanche/Source/AvalancheMask/Private/Handling/AvaMaskDesignedMaterialHandle.h
+++ b/Engine/Plugins/Experimental/Avalanche/Source/AvalancheMask/Private/Handling/AvaMaskDesignedMaterialHandle.h
@@ -72,8 +72,8 @@ protected:
 #if WITH_EDITOR
 	void OnMaterialBuilt(UDynamicMaterialModelBase* InMaterialModel);
 	
-	UMaterialFunctionInterface* GetOutputProcessor();
-	void SetOutputProcessor(UMaterialFunctionInterface* InMaterialFunction);
+	UMaterialFunctionInterface* GetOutputProcessor(EBlendMode InBlendMode);
+	void SetOutputProcessor(EBlendMode InBlendMode, UMaterialFunctionInterface* InMaterialFunction);
 #else
 	// Dummy, returns nullptr
 	UMaterialFunctionInterface* GetOutputProcessor();
diff --git a/Engine/Plugins/Experimental/DynamicMaterial/Source/DynamicMaterial/Private/Material/DynamicMaterialInstance.cpp b/Engine/Plugins/Experimental/DynamicMaterial/Source/DynamicMaterial/Private/Material/DynamicMaterialInstance.cpp
--- a/Engine/Plugins/Experimental/DynamicMaterial/Source/DynamicMaterial/Private/Material/DynamicMaterialInstance.cpp
+++ b/Engine/Plugins/Experimental/DynamicMaterial/Source/DynamicMaterial/Private/Material/DynamicMaterialInstance.cpp
@@ -49,12 +49,12 @@ UDynamicMaterialInstance::UDynamicMaterialInstance()
 	bOutputTranslucentVelocity = true;
 }
 
-UDynamicMaterialModelBase* UDynamicMaterialInstance::GetMaterialModelBase()
+UDynamicMaterialModelBase* UDynamicMaterialInstance::GetMaterialModelBase() const
 {
 	return MaterialModelBase;
 }
 
-UDynamicMaterialModel* UDynamicMaterialInstance::GetMaterialModel()
+UDynamicMaterialModel* UDynamicMaterialInstance::GetMaterialModel() const
 {
 	if (IsValid(MaterialModelBase))
 	{
diff --git a/Engine/Plugins/Experimental/DynamicMaterial/Source/DynamicMaterial/Public/Material/DynamicMaterialInstance.h b/Engine/Plugins/Experimental/DynamicMaterial/Source/DynamicMaterial/Public/Material/DynamicMaterialInstance.h
--- a/Engine/Plugins/Experimental/DynamicMaterial/Source/DynamicMaterial/Public/Material/DynamicMaterialInstance.h
+++ b/Engine/Plugins/Experimental/DynamicMaterial/Source/DynamicMaterial/Public/Material/DynamicMaterialInstance.h
@@ -31,11 +31,11 @@ public:
 
 	/** Returns the Material Model associated with this Material Designer Material. */
 	UFUNCTION(BlueprintPure, Category = "Material Designer")
-	DYNAMICMATERIAL_API UDynamicMaterialModelBase* GetMaterialModelBase();
+	DYNAMICMATERIAL_API UDynamicMaterialModelBase* GetMaterialModelBase() const;
 
 	/** Resolves the base Material Model used with this Instance and returns it. */
 	UFUNCTION(BlueprintPure, Category = "Material Designer")
-	DYNAMICMATERIAL_API UDynamicMaterialModel* GetMaterialModel();
+	DYNAMICMATERIAL_API UDynamicMaterialModel* GetMaterialModel() const;
 
 	//~ Begin UObject
 	DYNAMICMATERIAL_API virtual void GetAssetRegistryTags(FAssetRegistryTagsContext Context) const override;
diff --git a/Engine/Plugins/Experimental/DynamicMaterial/Source/DynamicMaterialEditor/Private/UI/Widgets/SDMMaterialEditor.cpp b/Engine/Plugins/Experimental/DynamicMaterial/Source/DynamicMaterialEditor/Private/UI/Widgets/SDMMaterialEditor.cpp
--- a/Engine/Plugins/Experimental/DynamicMaterial/Source/DynamicMaterialEditor/Private/UI/Widgets/SDMMaterialEditor.cpp
+++ b/Engine/Plugins/Experimental/DynamicMaterial/Source/DynamicMaterialEditor/Private/UI/Widgets/SDMMaterialEditor.cpp
@@ -597,7 +597,7 @@ void SDMMaterialEditor::Validate()
 
 		if (!UDMMaterialModelFunctionLibrary::IsModelValid(MaterialModelBaseFromProperty))
 		{
-			MaterialModelBase = nullptr;
+			MaterialModelBaseFromProperty = nullptr;
 		}
 
 		if (MaterialModelBase != MaterialModelBaseFromProperty)
