Motion Design:
- Horde issue 592886

diff --git a/Engine/Plugins/Experimental/Avalanche/Source/AvalancheMask/Private/Handling/AvaMaskDesignedMaterialHandle.cpp b/Engine/Plugins/Experimental/Avalanche/Source/AvalancheMask/Private/Handling/AvaMaskDesignedMaterialHandle.cpp
--- a/Engine/Plugins/Experimental/Avalanche/Source/AvalancheMask/Private/Handling/AvaMaskDesignedMaterialHandle.cpp
+++ b/Engine/Plugins/Experimental/Avalanche/Source/AvalancheMask/Private/Handling/AvaMaskDesignedMaterialHandle.cpp
@@ -382,7 +382,7 @@ void FAvaMaskDesignedMaterialHandle::SetOutputProcessor(EBlendMode InBlendMode,
 	}
 }
 #else
-UMaterialFunctionInterface* FAvaMaskDesignedMaterialHandle::GetOutputProcessor()
+UMaterialFunctionInterface* FAvaMaskDesignedMaterialHandle::GetOutputProcessor(EBlendMode InBlendMode)
 {
 	return nullptr;
 }
diff --git a/Engine/Plugins/Experimental/Avalanche/Source/AvalancheMask/Private/Handling/AvaMaskDesignedMaterialHandle.h b/Engine/Plugins/Experimental/Avalanche/Source/AvalancheMask/Private/Handling/AvaMaskDesignedMaterialHandle.h
--- a/Engine/Plugins/Experimental/Avalanche/Source/AvalancheMask/Private/Handling/AvaMaskDesignedMaterialHandle.h
+++ b/Engine/Plugins/Experimental/Avalanche/Source/AvalancheMask/Private/Handling/AvaMaskDesignedMaterialHandle.h
@@ -76,7 +76,7 @@ protected:
 	void SetOutputProcessor(EBlendMode InBlendMode, UMaterialFunctionInterface* InMaterialFunction);
 #else
 	// Dummy, returns nullptr
-	UMaterialFunctionInterface* GetOutputProcessor();
+	UMaterialFunctionInterface* GetOutputProcessor(EBlendMode InBlendMode);
 #endif
 
 private:
