MotionDesign : Modifiers
- Fixed visibility modifier reveals all children when parent visibility is changed

diff --git a/Engine/Plugins/Experimental/Avalanche/Source/AvalancheModifiers/Private/Extensions/AvaRenderStateUpdateModifierExtension.cpp b/Engine/Plugins/Experimental/Avalanche/Source/AvalancheModifiers/Private/Extensions/AvaRenderStateUpdateModifierExtension.cpp
--- a/Engine/Plugins/Experimental/Avalanche/Source/AvalancheModifiers/Private/Extensions/AvaRenderStateUpdateModifierExtension.cpp
+++ b/Engine/Plugins/Experimental/Avalanche/Source/AvalancheModifiers/Private/Extensions/AvaRenderStateUpdateModifierExtension.cpp
@@ -95,11 +95,6 @@ void FAvaRenderStateUpdateModifierExtension::OnExtensionDisabled(EActorModifierC
 
 void FAvaRenderStateUpdateModifierExtension::OnRenderStateDirty(UActorComponent& InComponent)
 {
-	if (!InComponent.IsA<UPrimitiveComponent>())
-	{
-		return;
-	}
-
 	const AActor* ModifierActor = GetModifierActor();
 	AActor* ActorDirty = InComponent.GetOwner();
 
diff --git a/Engine/Plugins/Experimental/Avalanche/Source/AvalancheModifiers/Private/Modifiers/AvaVisibilityModifier.cpp b/Engine/Plugins/Experimental/Avalanche/Source/AvalancheModifiers/Private/Modifiers/AvaVisibilityModifier.cpp
--- a/Engine/Plugins/Experimental/Avalanche/Source/AvalancheModifiers/Private/Modifiers/AvaVisibilityModifier.cpp
+++ b/Engine/Plugins/Experimental/Avalanche/Source/AvalancheModifiers/Private/Modifiers/AvaVisibilityModifier.cpp
@@ -178,10 +178,26 @@ void UAvaVisibilityModifier::OnRenderStateUpdated(AActor* InActor, UActorCompone
 {
 	Super::OnRenderStateUpdated(InActor, InComponent);
 
-	const AActor* ActorModified = GetModifiedActor();
+	AActor* ActorModified = GetModifiedActor();
+
+	// Only handle what is linked to us
+	if (!IsValid(ActorModified))
+	{
+		return;
+	}
+
+	const bool bThisActorUpdated = InActor == ActorModified;
+	const bool bActorAttachedToThisUpdated = InActor->IsAttachedTo(ActorModified);
+
+	if (!bThisActorUpdated && !bActorAttachedToThisUpdated)
+	{
+		return;
+	}
+
+	// If no modifier is found above us, then we handle this case
+	const UAvaVisibilityModifier* Modifier = GetFirstModifierAbove(ActorModified);
 
-	if (!IsValid(ActorModified)
-		|| !InActor->IsAttachedTo(ActorModified))
+	if (bThisActorUpdated && Modifier)
 	{
 		return;
 	}
