MotionDesign : SVG
- Fixed SVG editor commands not registered even after module has loaded and causing crash on linux

diff --git a/Engine/Plugins/Experimental/Avalanche/Source/AvalancheSVGEditor/Private/AvaSVGEditorModule.cpp b/Engine/Plugins/Experimental/Avalanche/Source/AvalancheSVGEditor/Private/AvaSVGEditorModule.cpp
--- a/Engine/Plugins/Experimental/Avalanche/Source/AvalancheSVGEditor/Private/AvaSVGEditorModule.cpp
+++ b/Engine/Plugins/Experimental/Avalanche/Source/AvalancheSVGEditor/Private/AvaSVGEditorModule.cpp
@@ -10,6 +10,7 @@
 #include "Modules/ModuleObserver.h"
 #include "ProceduralMeshes/SVGDynamicMeshComponent.h"
 #include "SVGEngineSubsystem.h"
+#include "SVGImporterEditorCommands.h"
 #include "SVGShapesParentActor.h"
 #include "Subsystems/ActorModifierCoreSubsystem.h"
 #include "Tool/AvaSVGActorTool.h"
@@ -34,23 +35,32 @@ void FAvaSVGEditorModule::ShutdownModule()
 
 void FAvaSVGEditorModule::RegisterTools(IAvalancheInteractiveToolsModule* InModule)
 {
-	// Use module observer to make sure SVGImporterEditor is loaded
-	FModuleObserver ModuleObserver(
-		"SVGImporterEditor",
-		FSimpleDelegate::CreateLambda([InModule]()
-		{
-			if (!InModule)
-			{
-				return;
-			}
+	if (!InModule)
+	{
+		return;
+	}
 
-			InModule->RegisterTool(
-				IAvalancheInteractiveToolsModule::Get().CategoryNameActor,
-				GetDefault<UAvaSVGActorTool>()->GetToolParameters()
-			);
-		}),
-		FSimpleDelegate()
-	);
+	if (FSVGImporterEditorCommands::IsRegistered())
+	{
+		InModule->RegisterTool(
+			IAvalancheInteractiveToolsModule::Get().CategoryNameActor,
+			GetDefault<UAvaSVGActorTool>()->GetToolParameters()
+		);
+	}
+	else
+	{
+		// Register SVG when commands are registered instead of when module is loaded, may happen at different time on various platform
+		FSVGImporterEditorCommands::CommandsChanged.AddRaw(this, &FAvaSVGEditorModule::OnSVGCommandsRegistered, InModule);
+	}
+}
+
+void FAvaSVGEditorModule::OnSVGCommandsRegistered(const FBindingContext& InContext, IAvalancheInteractiveToolsModule* InModule)
+{
+	if (InContext.GetContextName() == TEXT("SVGImporterEditor") && FSVGImporterEditorCommands::IsRegistered())
+	{
+		RegisterTools(InModule);
+		FSVGImporterEditorCommands::CommandsChanged.RemoveAll(this);
+	}
 }
 
 void FAvaSVGEditorModule::OnSVGActorSplit(ASVGShapesParentActor* InSVGShapesParent)
diff --git a/Engine/Plugins/Experimental/Avalanche/Source/AvalancheSVGEditor/Private/AvaSVGEditorModule.h b/Engine/Plugins/Experimental/Avalanche/Source/AvalancheSVGEditor/Private/AvaSVGEditorModule.h
--- a/Engine/Plugins/Experimental/Avalanche/Source/AvalancheSVGEditor/Private/AvaSVGEditorModule.h
+++ b/Engine/Plugins/Experimental/Avalanche/Source/AvalancheSVGEditor/Private/AvaSVGEditorModule.h
@@ -7,6 +7,7 @@
 
 class AActor;
 class ASVGShapesParentActor;
+class FBindingContext;
 class IAvalancheInteractiveToolsModule;
 class UEdMode;
 struct FAvaInteractiveToolsToolParameters;
@@ -19,6 +20,7 @@ public:
 
 private:
 	void RegisterTools(IAvalancheInteractiveToolsModule* InModule);
+	void OnSVGCommandsRegistered(const FBindingContext& InContext, IAvalancheInteractiveToolsModule* InModule);
 
 	void OnSVGActorSplit(ASVGShapesParentActor* InSVGShapesParent);
 	void OnSVGShapesUpdated(AActor* InActor) const;
