Motion Design: Removed unnecessary check and instead handle it more gracefully

diff --git a/Engine/Plugins/Experimental/Avalanche/Source/Avalanche/Private/AvaOutlinerUtils.cpp b/Engine/Plugins/Experimental/Avalanche/Source/Avalanche/Private/AvaOutlinerUtils.cpp
--- a/Engine/Plugins/Experimental/Avalanche/Source/Avalanche/Private/AvaOutlinerUtils.cpp
+++ b/Engine/Plugins/Experimental/Avalanche/Source/Avalanche/Private/AvaOutlinerUtils.cpp
@@ -12,7 +12,10 @@
 
 TSharedPtr<IAvaOutliner> FAvaOutlinerUtils::EditorGetOutliner(const UWorld* const InWorld)
 {
-	check(IsValid(InWorld));
+	if (!IsValid(InWorld))
+	{
+		return nullptr;
+	}
 
 	UAvaOutlinerSubsystem* const OutlinerSubsystem = InWorld->GetSubsystem<UAvaOutlinerSubsystem>();
 	if (!IsValid(OutlinerSubsystem))
