Motion Design:
- Added 4 new "simple" materials which have 1 side.
- Added the ability to pick one or two sided simple materials (default 2).

#virtualized

diff --git a/Engine/Plugins/Experimental/Avalanche/Source/AvalancheShapes/Private/AvaShapeParametricMaterial.cpp b/Engine/Plugins/Experimental/Avalanche/Source/AvalancheShapes/Private/AvaShapeParametricMaterial.cpp
--- a/Engine/Plugins/Experimental/Avalanche/Source/AvalancheShapes/Private/AvaShapeParametricMaterial.cpp
+++ b/Engine/Plugins/Experimental/Avalanche/Source/AvalancheShapes/Private/AvaShapeParametricMaterial.cpp
@@ -10,8 +10,9 @@ FAvaShapeParametricMaterial::FOnMaterialChanged FAvaShapeParametricMaterial::OnM
 
 FAvaShapeParametricMaterial::FAvaShapeParametricMaterial()
 {
-	bUseTranslucentMaterial         = false;
-	bUseUnlitMaterial               = false;
+	bUseTranslucentMaterial = false;
+	bUseUnlitMaterial       = false;
+	bUseTwoSidedMaterial    = true;
 
 	Style    = EAvaShapeParametricMaterialStyle::Solid;
 
@@ -28,6 +29,7 @@ FAvaShapeParametricMaterial::FAvaShapeParametricMaterial(const FAvaShapeParametr
 {
 	bUseUnlitMaterial        = Other.bUseUnlitMaterial;
 	bUseTranslucentMaterial  = Other.bUseTranslucentMaterial;
+	bUseTwoSidedMaterial     = Other.bUseTwoSidedMaterial;
 	ColorA                   = Other.ColorA;
 	ColorB                   = Other.ColorB;
 	GradientOffset           = Other.GradientOffset;
@@ -42,6 +44,7 @@ FAvaShapeParametricMaterial& FAvaShapeParametricMaterial::operator=(const FAvaSh
 {
 	bUseUnlitMaterial        = Other.bUseUnlitMaterial;
 	bUseTranslucentMaterial  = Other.bUseTranslucentMaterial;
+	bUseTwoSidedMaterial     = Other.bUseTwoSidedMaterial;
 	ColorA                   = Other.ColorA;
 	ColorB                   = Other.ColorB;
 	GradientOffset           = Other.GradientOffset;
@@ -65,31 +68,65 @@ UMaterialInterface* FAvaShapeParametricMaterial::GetDefaultMaterial() const
 
 void FAvaShapeParametricMaterial::LoadDefaultMaterials() const
 {
-	if (!DefaultMaterials.IsEmpty())
+	if (DefaultMaterials.Num() == MaterialTypeCount)
 	{
 		return;
 	}
 
-	if (FAvaShapeParametricMaterial* MutableThis = const_cast<FAvaShapeParametricMaterial*>(this))
+	FAvaShapeParametricMaterial* MutableThis = const_cast<FAvaShapeParametricMaterial*>(this);
+
+	if (!MutableThis)
 	{
-		static const TCHAR* LitOpaqueMaterialPath = TEXT("/Avalanche/ToolboxResources/M_Toolbox_Opaque_Lit.M_Toolbox_Opaque_Lit");
-		UMaterialInterface* DefaultLitMaterialOpaque = LoadObject<UMaterialInterface>(nullptr, LitOpaqueMaterialPath);
-		MutableThis->DefaultMaterials.Add(DefaultLitMaterialOpaque);
+		return;
+	}
+
+	MutableThis->DefaultMaterials.Empty(MaterialTypeCount);
+
+	static const TCHAR* LitOpaqueMaterialPath = TEXT("/Avalanche/ToolboxResources/M_Toolbox_Opaque_Lit.M_Toolbox_Opaque_Lit");
+	UMaterialInterface* DefaultLitMaterialOpaque = LoadObject<UMaterialInterface>(nullptr, LitOpaqueMaterialPath);
+	MutableThis->DefaultMaterials.Add(DefaultLitMaterialOpaque);
 
-		static const TCHAR* LitTranslucentMaterialPath = TEXT("/Avalanche/ToolboxResources/M_Toolbox_Translucent_Lit.M_Toolbox_Translucent_Lit");
-		UMaterialInterface* DefaultLitMaterialTranslucent = LoadObject<UMaterialInterface>(nullptr, LitTranslucentMaterialPath);
-		MutableThis->DefaultMaterials.Add(DefaultLitMaterialTranslucent);
+	static const TCHAR* LitTranslucentMaterialPath = TEXT("/Avalanche/ToolboxResources/M_Toolbox_Translucent_Lit.M_Toolbox_Translucent_Lit");
+	UMaterialInterface* DefaultLitMaterialTranslucent = LoadObject<UMaterialInterface>(nullptr, LitTranslucentMaterialPath);
+	MutableThis->DefaultMaterials.Add(DefaultLitMaterialTranslucent);
 
-		static const TCHAR* UnlitOpaqueMaterialPath = TEXT("/Avalanche/ToolboxResources/M_Toolbox_Opaque_Unlit.M_Toolbox_Opaque_Unlit");
-		UMaterialInterface* DefaultUnlitMaterialOpaque = LoadObject<UMaterialInterface>(nullptr, UnlitOpaqueMaterialPath);
-		MutableThis->DefaultMaterials.Add(DefaultUnlitMaterialOpaque);
+	static const TCHAR* UnlitOpaqueMaterialPath = TEXT("/Avalanche/ToolboxResources/M_Toolbox_Opaque_Unlit.M_Toolbox_Opaque_Unlit");
+	UMaterialInterface* DefaultUnlitMaterialOpaque = LoadObject<UMaterialInterface>(nullptr, UnlitOpaqueMaterialPath);
+	MutableThis->DefaultMaterials.Add(DefaultUnlitMaterialOpaque);
 
-		static const TCHAR* UnlitTranslucentMaterialPath = TEXT("/Avalanche/ToolboxResources/M_Toolbox_Translucent_Unlit.M_Toolbox_Translucent_Unlit");
-		UMaterialInterface* DefaultUnlitMaterialTranslucent = LoadObject<UMaterialInterface>(nullptr, UnlitTranslucentMaterialPath);
-		MutableThis->DefaultMaterials.Add(DefaultUnlitMaterialTranslucent);
+	static const TCHAR* UnlitTranslucentMaterialPath = TEXT("/Avalanche/ToolboxResources/M_Toolbox_Translucent_Unlit.M_Toolbox_Translucent_Unlit");
+	UMaterialInterface* DefaultUnlitMaterialTranslucent = LoadObject<UMaterialInterface>(nullptr, UnlitTranslucentMaterialPath);
+	MutableThis->DefaultMaterials.Add(DefaultUnlitMaterialTranslucent);
 
-		MutableThis->InstanceMaterials.SetNum(DefaultMaterials.Num(), EAllowShrinking::No);
+	// As above, but one sided materials.
+	static const TCHAR* LitOpaqueOneSidedMaterialPath = TEXT("/Avalanche/ToolboxResources/M_Toolbox_Opaque_Lit_Onesided.M_Toolbox_Opaque_Lit_Onesided");
+	UMaterialInterface* DefaultLitOneSidedMaterialOpaque = LoadObject<UMaterialInterface>(nullptr, LitOpaqueOneSidedMaterialPath);
+	MutableThis->DefaultMaterials.Add(DefaultLitOneSidedMaterialOpaque);
+
+	static const TCHAR* LitTranslucentOneSidedMaterialPath = TEXT("/Avalanche/ToolboxResources/M_Toolbox_Translucent_Lit_Onesided.M_Toolbox_Translucent_Lit_Onesided");
+	UMaterialInterface* DefaultLitOneSidedMaterialTranslucent = LoadObject<UMaterialInterface>(nullptr, LitTranslucentOneSidedMaterialPath);
+	MutableThis->DefaultMaterials.Add(DefaultLitOneSidedMaterialTranslucent);
+
+	static const TCHAR* UnlitOpaqueOneSidedMaterialPath = TEXT("/Avalanche/ToolboxResources/M_Toolbox_Opaque_Unlit_Onesided.M_Toolbox_Opaque_Unlit_Onesided");
+	UMaterialInterface* DefaultUnlitOneSidedMaterialOpaque = LoadObject<UMaterialInterface>(nullptr, UnlitOpaqueOneSidedMaterialPath);
+	MutableThis->DefaultMaterials.Add(DefaultUnlitOneSidedMaterialOpaque);
+
+	static const TCHAR* UnlitTranslucentOneSidedMaterialPath = TEXT("/Avalanche/ToolboxResources/M_Toolbox_Translucent_Unlit_Onesided.M_Toolbox_Translucent_Unlit_Onesided");
+	UMaterialInterface* DefaultUnlitOneSidedMaterialTranslucent = LoadObject<UMaterialInterface>(nullptr, UnlitTranslucentOneSidedMaterialPath);
+	MutableThis->DefaultMaterials.Add(DefaultUnlitOneSidedMaterialTranslucent);
+
+	MutableThis->InstanceMaterials.SetNum(DefaultMaterials.Num(), EAllowShrinking::No);
+}
+
+void FAvaShapeParametricMaterial::SetUseTwoSidedMaterial(bool bInUse)
+{
+	if (bInUse == bUseTwoSidedMaterial)
+	{
+		return;
 	}
+
+	bUseTwoSidedMaterial = bInUse;
+	OnMaterialParameterUpdated();
 }
 
 void FAvaShapeParametricMaterial::SetMaterialParameterValues(UMaterialInstanceDynamic* InMaterialInstance, bool bInNotifyUpdate) const
@@ -148,14 +185,9 @@ int32 FAvaShapeParametricMaterial::GetActiveInstanceIndex() const
 {
 	FAvaShapeParametricMaterial* MutableThis = const_cast<FAvaShapeParametricMaterial*>(this);
 
-	if (bUseTranslucentMaterial)
-	{
-		MutableThis->ActiveInstanceIndex = bUseUnlitMaterial ? TRANSLUCENT_UNLIT : TRANSLUCENT_LIT;
-	}
-	else
-	{
-		MutableThis->ActiveInstanceIndex = bUseUnlitMaterial ? OPAQUE_UNLIT : OPAQUE_LIT;
-	}
+	MutableThis->ActiveInstanceIndex = (bUseTranslucentMaterial ? Translucent : Opaque)
+		+ (bUseUnlitMaterial ? Unlit : Lit)
+		+ (bUseTwoSidedMaterial ? TwoSided : OneSided);
 
 	return MutableThis->ActiveInstanceIndex;
 }
diff --git a/Engine/Plugins/Experimental/Avalanche/Source/AvalancheShapes/Public/AvaShapeParametricMaterial.h b/Engine/Plugins/Experimental/Avalanche/Source/AvalancheShapes/Public/AvaShapeParametricMaterial.h
--- a/Engine/Plugins/Experimental/Avalanche/Source/AvalancheShapes/Public/AvaShapeParametricMaterial.h
+++ b/Engine/Plugins/Experimental/Avalanche/Source/AvalancheShapes/Public/AvaShapeParametricMaterial.h
@@ -20,10 +20,19 @@ struct FAvaShapeParametricMaterial
 {
 	GENERATED_BODY()
 
-	static inline constexpr int32 OPAQUE_LIT = 0;
-	static inline constexpr int32 TRANSLUCENT_LIT = 1;
-	static inline constexpr int32 OPAQUE_UNLIT = 2;
-	static inline constexpr int32 TRANSLUCENT_UNLIT = 3;
+	enum EMaterialType : int32
+	{
+		Opaque = 0,
+		Translucent = 1 << 0,
+
+		Lit = 0,
+		Unlit = 1 << 1,
+
+		TwoSided = 0,
+		OneSided = 1 << 2,
+
+		MaterialTypeCount = Translucent + Unlit + OneSided + 1
+	};
 
 	DECLARE_MULTICAST_DELEGATE_OneParam(FOnMaterialChanged, FAvaShapeParametricMaterial&)
 	DECLARE_MULTICAST_DELEGATE(FOnMaterialParameterChanged)
@@ -105,6 +114,13 @@ struct FAvaShapeParametricMaterial
 
 	AVALANCHESHAPES_API void SetUseUnlitMaterial(bool bInUse);
 
+	bool GetUseTwoSidedMaterial() const
+	{
+		return bUseTwoSidedMaterial;
+	}
+
+	AVALANCHESHAPES_API void SetUseTwoSidedMaterial(bool bInUse);
+
 	/** Set parameter values on a material instance */
 	AVALANCHESHAPES_API void SetMaterialParameterValues(UMaterialInstanceDynamic* InMaterialInstance, bool bInNotifyUpdate = false) const;
 
@@ -156,6 +172,10 @@ protected:
 	UPROPERTY(EditInstanceOnly, BlueprintReadWrite, Category="Material", meta=(DisplayName="Use Unlit Material", AllowPrivateAccess = "true"))
 	bool bUseUnlitMaterial;
 
+	// whether the material is one sided or two sided
+	UPROPERTY(EditInstanceOnly, BlueprintReadWrite, Category = "Material", meta = (DisplayName = "Use Two Sided Material", AllowPrivateAccess = "true"))
+	bool bUseTwoSidedMaterial;
+
 private:
 	static FOnMaterialChanged OnMaterialChangedDelegate;
 
