MotionDesign : Modifiers
- Updating radial arrange modifier test to reflect adjusted default options

diff --git a/Engine/Plugins/Experimental/Avalanche/Source/AvalancheModifiers/Private/Tests/AvaModifiersRadialArrange.cpp b/Engine/Plugins/Experimental/Avalanche/Source/AvalancheModifiers/Private/Tests/AvaModifiersRadialArrange.cpp
--- a/Engine/Plugins/Experimental/Avalanche/Source/AvalancheModifiers/Private/Tests/AvaModifiersRadialArrange.cpp
+++ b/Engine/Plugins/Experimental/Avalanche/Source/AvalancheModifiers/Private/Tests/AvaModifiersRadialArrange.cpp
@@ -57,10 +57,10 @@ void AvalancheModifiersRadialArrange::Define()
 		It("Should arrange child actors in a radial layout around parent actor", [this]
 		{
 			TArray<FVector> ExpectedChildrenLocations = {
-				FVector(0, -InnerRadius, 0),
-				FVector(0, 0, InnerRadius),
+				FVector(-InnerRadius, 0, 0),
 				FVector(0, InnerRadius, 0),
-				FVector(0, 0, -InnerRadius)
+				FVector(InnerRadius, 0, 0),
+				FVector(0, -InnerRadius, 0)
 			};
 
 			for (int32 i = 0; i < ChildActors.Num(); i++)
