MotionDesign : InteractiveTool
- Motion Design Tool Actors will spawn at 0,0,0 when double clicked (default action) from the palette

diff --git a/Engine/Plugins/Experimental/Avalanche/Source/AvalancheEffectorsEditor/Private/Cloner/AvaClonerActorTool.cpp b/Engine/Plugins/Experimental/Avalanche/Source/AvalancheEffectorsEditor/Private/Cloner/AvaClonerActorTool.cpp
--- a/Engine/Plugins/Experimental/Avalanche/Source/AvalancheEffectorsEditor/Private/Cloner/AvaClonerActorTool.cpp
+++ b/Engine/Plugins/Experimental/Avalanche/Source/AvalancheEffectorsEditor/Private/Cloner/AvaClonerActorTool.cpp
@@ -12,11 +12,6 @@ UAvaClonerActorTool::UAvaClonerActorTool()
 	ActorClass = ACEClonerActor::StaticClass();
 }
 
-bool UAvaClonerActorTool::UseIdentityRotation() const
-{
-	return ConditionalIdentityRotation();
-}
-
 FName UAvaClonerActorTool::GetCategoryName()
 {
 	return IAvalancheInteractiveToolsModule::CategoryNameActor;
diff --git a/Engine/Plugins/Experimental/Avalanche/Source/AvalancheEffectorsEditor/Private/Cloner/AvaClonerActorTool.h b/Engine/Plugins/Experimental/Avalanche/Source/AvalancheEffectorsEditor/Private/Cloner/AvaClonerActorTool.h
--- a/Engine/Plugins/Experimental/Avalanche/Source/AvalancheEffectorsEditor/Private/Cloner/AvaClonerActorTool.h
+++ b/Engine/Plugins/Experimental/Avalanche/Source/AvalancheEffectorsEditor/Private/Cloner/AvaClonerActorTool.h
@@ -14,8 +14,6 @@ public:
 	UAvaClonerActorTool();
 
 	//~ Begin UAvaInteractiveToolsToolBase
-	virtual bool UseIdentityLocation() const override { return false; }
-	virtual bool UseIdentityRotation() const override;
 	virtual FName GetCategoryName() override;
 	virtual FAvaInteractiveToolsToolParameters GetToolParameters() const override;
 	//~ End UAvaInteractiveToolsToolBase
diff --git a/Engine/Plugins/Experimental/Avalanche/Source/AvalancheEffectorsEditor/Private/Effector/AvaEffectorActorTool.cpp b/Engine/Plugins/Experimental/Avalanche/Source/AvalancheEffectorsEditor/Private/Effector/AvaEffectorActorTool.cpp
--- a/Engine/Plugins/Experimental/Avalanche/Source/AvalancheEffectorsEditor/Private/Effector/AvaEffectorActorTool.cpp
+++ b/Engine/Plugins/Experimental/Avalanche/Source/AvalancheEffectorsEditor/Private/Effector/AvaEffectorActorTool.cpp
@@ -12,11 +12,6 @@ UAvaEffectorActorTool::UAvaEffectorActorTool()
 	ActorClass = ACEEffectorActor::StaticClass();
 }
 
-bool UAvaEffectorActorTool::UseIdentityRotation() const
-{
-	return ConditionalIdentityRotation();
-}
-
 FName UAvaEffectorActorTool::GetCategoryName()
 {
 	return IAvalancheInteractiveToolsModule::CategoryNameActor;
diff --git a/Engine/Plugins/Experimental/Avalanche/Source/AvalancheEffectorsEditor/Private/Effector/AvaEffectorActorTool.h b/Engine/Plugins/Experimental/Avalanche/Source/AvalancheEffectorsEditor/Private/Effector/AvaEffectorActorTool.h
--- a/Engine/Plugins/Experimental/Avalanche/Source/AvalancheEffectorsEditor/Private/Effector/AvaEffectorActorTool.h
+++ b/Engine/Plugins/Experimental/Avalanche/Source/AvalancheEffectorsEditor/Private/Effector/AvaEffectorActorTool.h
@@ -14,8 +14,6 @@ public:
 	UAvaEffectorActorTool();
 
 	//~ Begin UAvaInteractiveToolsToolBase
-	virtual bool UseIdentityLocation() const override { return false; }
-	virtual bool UseIdentityRotation() const override;
 	virtual FName GetCategoryName() override;
 	virtual FAvaInteractiveToolsToolParameters GetToolParameters() const override;
 	//~ End UAvaInteractiveToolsToolBase
diff --git a/Engine/Plugins/Experimental/Avalanche/Source/AvalancheInteractiveTools/Private/Tools/AvaInteractiveToolsActorToolBase.cpp b/Engine/Plugins/Experimental/Avalanche/Source/AvalancheInteractiveTools/Private/Tools/AvaInteractiveToolsActorToolBase.cpp
--- a/Engine/Plugins/Experimental/Avalanche/Source/AvalancheInteractiveTools/Private/Tools/AvaInteractiveToolsActorToolBase.cpp
+++ b/Engine/Plugins/Experimental/Avalanche/Source/AvalancheInteractiveTools/Private/Tools/AvaInteractiveToolsActorToolBase.cpp
@@ -24,3 +24,8 @@ void UAvaInteractiveToolsActorToolBase::DefaultAction()
 
 	Super::DefaultAction();
 }
+
+bool UAvaInteractiveToolsActorToolBase::UseIdentityRotation() const
+{
+	return ConditionalIdentityRotation();
+}
diff --git a/Engine/Plugins/Experimental/Avalanche/Source/AvalancheInteractiveTools/Private/Tools/AvaInteractiveToolsActorToolNull.h b/Engine/Plugins/Experimental/Avalanche/Source/AvalancheInteractiveTools/Private/Tools/AvaInteractiveToolsActorToolNull.h
--- a/Engine/Plugins/Experimental/Avalanche/Source/AvalancheInteractiveTools/Private/Tools/AvaInteractiveToolsActorToolNull.h
+++ b/Engine/Plugins/Experimental/Avalanche/Source/AvalancheInteractiveTools/Private/Tools/AvaInteractiveToolsActorToolNull.h
@@ -14,7 +14,6 @@ public:
 	UAvaInteractiveToolsActorToolNull();
 
 	//~ Begin UAvaInteractiveToolsToolBase
-	virtual bool UseIdentityLocation() const override { return false; }
 	virtual bool UseIdentityRotation() const override { return true; }
 	virtual FName GetCategoryName() override;
 	virtual FAvaInteractiveToolsToolParameters GetToolParameters() const override;
diff --git a/Engine/Plugins/Experimental/Avalanche/Source/AvalancheInteractiveTools/Private/Tools/AvaInteractiveToolsActorToolSpline.h b/Engine/Plugins/Experimental/Avalanche/Source/AvalancheInteractiveTools/Private/Tools/AvaInteractiveToolsActorToolSpline.h
--- a/Engine/Plugins/Experimental/Avalanche/Source/AvalancheInteractiveTools/Private/Tools/AvaInteractiveToolsActorToolSpline.h
+++ b/Engine/Plugins/Experimental/Avalanche/Source/AvalancheInteractiveTools/Private/Tools/AvaInteractiveToolsActorToolSpline.h
@@ -14,7 +14,6 @@ public:
 	UAvaInteractiveToolsActorToolSpline();
 
 	//~ Begin UAvaInteractiveToolsToolBase
-	virtual bool UseIdentityLocation() const override { return false; }
 	virtual bool UseIdentityRotation() const override { return false; }
 	virtual FName GetCategoryName() override;
 	virtual FAvaInteractiveToolsToolParameters GetToolParameters() const override;
diff --git a/Engine/Plugins/Experimental/Avalanche/Source/AvalancheInteractiveTools/Public/Tools/AvaInteractiveToolsActorToolBase.h b/Engine/Plugins/Experimental/Avalanche/Source/AvalancheInteractiveTools/Public/Tools/AvaInteractiveToolsActorToolBase.h
--- a/Engine/Plugins/Experimental/Avalanche/Source/AvalancheInteractiveTools/Public/Tools/AvaInteractiveToolsActorToolBase.h
+++ b/Engine/Plugins/Experimental/Avalanche/Source/AvalancheInteractiveTools/Public/Tools/AvaInteractiveToolsActorToolBase.h
@@ -20,5 +20,7 @@ protected:
 	//~ Begin UAvaInteractiveToolsToolBase
 	virtual bool OnBegin() override;
 	virtual void DefaultAction() override;
+	virtual bool UseIdentityLocation() const override { return bPerformingDefaultAction; }
+	virtual bool UseIdentityRotation() const override;
 	//~ End UAvaInteractiveToolsToolBase
 };
diff --git a/Engine/Plugins/Experimental/Avalanche/Source/AvalancheSVGEditor/Private/Tool/AvaSVGActorTool.cpp b/Engine/Plugins/Experimental/Avalanche/Source/AvalancheSVGEditor/Private/Tool/AvaSVGActorTool.cpp
--- a/Engine/Plugins/Experimental/Avalanche/Source/AvalancheSVGEditor/Private/Tool/AvaSVGActorTool.cpp
+++ b/Engine/Plugins/Experimental/Avalanche/Source/AvalancheSVGEditor/Private/Tool/AvaSVGActorTool.cpp
@@ -12,11 +12,6 @@ UAvaSVGActorTool::UAvaSVGActorTool()
 	ActorClass = ASVGActor::StaticClass();
 }
 
-bool UAvaSVGActorTool::UseIdentityRotation() const
-{
-	return ConditionalIdentityRotation();
-}
-
 FName UAvaSVGActorTool::GetCategoryName()
 {
 	return IAvalancheInteractiveToolsModule::CategoryNameActor;
diff --git a/Engine/Plugins/Experimental/Avalanche/Source/AvalancheSVGEditor/Private/Tool/AvaSVGActorTool.h b/Engine/Plugins/Experimental/Avalanche/Source/AvalancheSVGEditor/Private/Tool/AvaSVGActorTool.h
--- a/Engine/Plugins/Experimental/Avalanche/Source/AvalancheSVGEditor/Private/Tool/AvaSVGActorTool.h
+++ b/Engine/Plugins/Experimental/Avalanche/Source/AvalancheSVGEditor/Private/Tool/AvaSVGActorTool.h
@@ -14,8 +14,6 @@ public:
 	UAvaSVGActorTool();
 
 	//~ Begin UAvaInteractiveToolsToolBase
-	virtual bool UseIdentityLocation() const override { return false; }
-	virtual bool UseIdentityRotation() const override;
 	virtual FName GetCategoryName() override;
 	virtual FAvaInteractiveToolsToolParameters GetToolParameters() const override;
 	//~ End UAvaInteractiveToolsToolBase
diff --git a/Engine/Plugins/Experimental/Avalanche/Source/AvalancheShapesEditor/Private/Tools/AvaShapesEditorShapeToolBase.cpp b/Engine/Plugins/Experimental/Avalanche/Source/AvalancheShapesEditor/Private/Tools/AvaShapesEditorShapeToolBase.cpp
--- a/Engine/Plugins/Experimental/Avalanche/Source/AvalancheShapesEditor/Private/Tools/AvaShapesEditorShapeToolBase.cpp
+++ b/Engine/Plugins/Experimental/Avalanche/Source/AvalancheShapesEditor/Private/Tools/AvaShapesEditorShapeToolBase.cpp
@@ -99,15 +99,5 @@ bool UAvaShapesEditorShapeToolBase::UseIdentityLocation() const
 		return false;
 	}
 
-	if (!ShapeClass || ShapeClass->IsChildOf<UAvaShape2DDynMeshBase>())
-	{
-		return false;
-	}
-
 	return !IsMotionDesignViewport();
 }
-
-bool UAvaShapesEditorShapeToolBase::UseIdentityRotation() const
-{
-	return ConditionalIdentityRotation();
-}
diff --git a/Engine/Plugins/Experimental/Avalanche/Source/AvalancheShapesEditor/Private/Tools/AvaShapesEditorShapeToolBase.h b/Engine/Plugins/Experimental/Avalanche/Source/AvalancheShapesEditor/Private/Tools/AvaShapesEditorShapeToolBase.h
--- a/Engine/Plugins/Experimental/Avalanche/Source/AvalancheShapesEditor/Private/Tools/AvaShapesEditorShapeToolBase.h
+++ b/Engine/Plugins/Experimental/Avalanche/Source/AvalancheShapesEditor/Private/Tools/AvaShapesEditorShapeToolBase.h
@@ -48,7 +48,6 @@ protected:
 
 	//~ Begin UAvaInteractiveToolsToolBase
 	virtual bool UseIdentityLocation() const override;
-	virtual bool UseIdentityRotation() const override;
 
 	virtual AActor* SpawnActor(TSubclassOf<AActor> InActorClass, EAvaViewportStatus InViewportStatus, 
 		const FVector2f& InViewportPosition, bool bInPreview, FString* InActorLabelOverride = nullptr) const override;
diff --git a/Engine/Plugins/Experimental/Avalanche/Source/AvalancheTextEditor/Private/Tool/AvaTextActorTool.cpp b/Engine/Plugins/Experimental/Avalanche/Source/AvalancheTextEditor/Private/Tool/AvaTextActorTool.cpp
--- a/Engine/Plugins/Experimental/Avalanche/Source/AvalancheTextEditor/Private/Tool/AvaTextActorTool.cpp
+++ b/Engine/Plugins/Experimental/Avalanche/Source/AvalancheTextEditor/Private/Tool/AvaTextActorTool.cpp
@@ -13,11 +13,6 @@ UAvaTextActorTool::UAvaTextActorTool()
 	ActorClass = AAvaTextActor::StaticClass();
 }
 
-bool UAvaTextActorTool::UseIdentityRotation() const
-{
-	return ConditionalIdentityRotation();
-}
-
 FName UAvaTextActorTool::GetCategoryName()
 {
 	return IAvalancheInteractiveToolsModule::CategoryNameActor;
diff --git a/Engine/Plugins/Experimental/Avalanche/Source/AvalancheTextEditor/Private/Tool/AvaTextActorTool.h b/Engine/Plugins/Experimental/Avalanche/Source/AvalancheTextEditor/Private/Tool/AvaTextActorTool.h
--- a/Engine/Plugins/Experimental/Avalanche/Source/AvalancheTextEditor/Private/Tool/AvaTextActorTool.h
+++ b/Engine/Plugins/Experimental/Avalanche/Source/AvalancheTextEditor/Private/Tool/AvaTextActorTool.h
@@ -14,8 +14,6 @@ public:
 	UAvaTextActorTool();
 
 	//~ Begin UAvaInteractiveToolsToolBase
-	virtual bool UseIdentityLocation() const override { return false; }
-	virtual bool UseIdentityRotation() const override;
 	virtual FName GetCategoryName() override;
 	virtual FAvaInteractiveToolsToolParameters GetToolParameters() const override;
 	//~ End UAvaInteractiveToolsToolBase
