MotionDesign : SVGImporter
- Fix commands not registered and causing crash on linux when retrieving them

diff --git a/Engine/Plugins/Experimental/Avalanche/Source/AvalancheSVGEditor/Private/AvaSVGEditorModule.cpp b/Engine/Plugins/Experimental/Avalanche/Source/AvalancheSVGEditor/Private/AvaSVGEditorModule.cpp
--- a/Engine/Plugins/Experimental/Avalanche/Source/AvalancheSVGEditor/Private/AvaSVGEditorModule.cpp
+++ b/Engine/Plugins/Experimental/Avalanche/Source/AvalancheSVGEditor/Private/AvaSVGEditorModule.cpp
@@ -7,10 +7,9 @@
 #include "Modifiers/ActorModifierCoreStack.h"
 #include "Modifiers/AvaBevelModifier.h"
 #include "Modifiers/AvaExtrudeModifier.h"
-#include "Modules/ModuleObserver.h"
 #include "ProceduralMeshes/SVGDynamicMeshComponent.h"
 #include "SVGEngineSubsystem.h"
-#include "SVGImporterEditorCommands.h"
+#include "SVGImporter.h"
 #include "SVGShapesParentActor.h"
 #include "Subsystems/ActorModifierCoreSubsystem.h"
 #include "Tool/AvaSVGActorTool.h"
@@ -40,27 +39,13 @@ void FAvaSVGEditorModule::RegisterTools(IAvalancheInteractiveToolsModule* InModu
 		return;
 	}
 
-	if (FSVGImporterEditorCommands::IsRegistered())
-	{
-		InModule->RegisterTool(
-			IAvalancheInteractiveToolsModule::Get().CategoryNameActor,
-			GetDefault<UAvaSVGActorTool>()->GetToolParameters()
-		);
-	}
-	else
-	{
-		// Register SVG when commands are registered instead of when module is loaded, may happen at different time on various platform
-		FSVGImporterEditorCommands::CommandsChanged.AddRaw(this, &FAvaSVGEditorModule::OnSVGCommandsRegistered, InModule);
-	}
-}
+	// Load svg module
+	FSVGImporterModule::Get();
 
-void FAvaSVGEditorModule::OnSVGCommandsRegistered(const FBindingContext& InContext, IAvalancheInteractiveToolsModule* InModule)
-{
-	if (InContext.GetContextName() == TEXT("SVGImporterEditor") && FSVGImporterEditorCommands::IsRegistered())
-	{
-		RegisterTools(InModule);
-		FSVGImporterEditorCommands::CommandsChanged.RemoveAll(this);
-	}
+	InModule->RegisterTool(
+		IAvalancheInteractiveToolsModule::Get().CategoryNameActor,
+		GetDefault<UAvaSVGActorTool>()->GetToolParameters()
+	);
 }
 
 void FAvaSVGEditorModule::OnSVGActorSplit(ASVGShapesParentActor* InSVGShapesParent)
diff --git a/Engine/Plugins/Experimental/Avalanche/Source/AvalancheSVGEditor/Private/AvaSVGEditorModule.h b/Engine/Plugins/Experimental/Avalanche/Source/AvalancheSVGEditor/Private/AvaSVGEditorModule.h
--- a/Engine/Plugins/Experimental/Avalanche/Source/AvalancheSVGEditor/Private/AvaSVGEditorModule.h
+++ b/Engine/Plugins/Experimental/Avalanche/Source/AvalancheSVGEditor/Private/AvaSVGEditorModule.h
@@ -20,7 +20,6 @@ public:
 
 private:
 	void RegisterTools(IAvalancheInteractiveToolsModule* InModule);
-	void OnSVGCommandsRegistered(const FBindingContext& InContext, IAvalancheInteractiveToolsModule* InModule);
 
 	void OnSVGActorSplit(ASVGShapesParentActor* InSVGShapesParent);
 	void OnSVGShapesUpdated(AActor* InActor) const;
diff --git a/Engine/Plugins/Experimental/Avalanche/Source/AvalancheSVGEditor/Private/Tool/AvaSVGActorTool.cpp b/Engine/Plugins/Experimental/Avalanche/Source/AvalancheSVGEditor/Private/Tool/AvaSVGActorTool.cpp
--- a/Engine/Plugins/Experimental/Avalanche/Source/AvalancheSVGEditor/Private/Tool/AvaSVGActorTool.cpp
+++ b/Engine/Plugins/Experimental/Avalanche/Source/AvalancheSVGEditor/Private/Tool/AvaSVGActorTool.cpp
@@ -20,7 +20,7 @@ FName UAvaSVGActorTool::GetCategoryName()
 FAvaInteractiveToolsToolParameters UAvaSVGActorTool::GetToolParameters() const
 {
 	return {
-		FSVGImporterEditorCommands::Get().SpawnSVGActor,
+		FSVGImporterEditorCommands::GetExternal().SpawnSVGActor,
 		TEXT("SVG Actor Tool"),
 		6000,
 		FAvalancheInteractiveToolsCreateBuilder::CreateLambda(
diff --git a/Engine/Plugins/Experimental/SVGImporter/Source/SVGImporterEditor/Private/SVGImporterEditorCommands.cpp b/Engine/Plugins/Experimental/SVGImporter/Source/SVGImporterEditor/Private/SVGImporterEditorCommands.cpp
--- a/Engine/Plugins/Experimental/SVGImporter/Source/SVGImporterEditor/Private/SVGImporterEditorCommands.cpp
+++ b/Engine/Plugins/Experimental/SVGImporter/Source/SVGImporterEditor/Private/SVGImporterEditorCommands.cpp
@@ -5,6 +5,21 @@
 
 #define LOCTEXT_NAMESPACE "AvaSVGEditorCommands"
 
+const FSVGImporterEditorCommands& FSVGImporterEditorCommands::GetExternal()
+{
+	if (!IsRegistered())
+	{
+		Register();
+	}
+
+	return Get();
+}
+
+const FSVGImporterEditorCommands& FSVGImporterEditorCommands::GetInternal()
+{
+	return Get();
+}
+
 FSVGImporterEditorCommands::FSVGImporterEditorCommands()
 	: TCommands<FSVGImporterEditorCommands>(
 		TEXT("SVGImporterEditor")
diff --git a/Engine/Plugins/Experimental/SVGImporter/Source/SVGImporterEditor/Public/SVGImporterEditorCommands.h b/Engine/Plugins/Experimental/SVGImporter/Source/SVGImporterEditor/Public/SVGImporterEditorCommands.h
--- a/Engine/Plugins/Experimental/SVGImporter/Source/SVGImporterEditor/Public/SVGImporterEditorCommands.h
+++ b/Engine/Plugins/Experimental/SVGImporter/Source/SVGImporterEditor/Public/SVGImporterEditorCommands.h
@@ -4,9 +4,13 @@
 
 #include "Framework/Commands/Commands.h"
 
-class SVGIMPORTEREDITOR_API FSVGImporterEditorCommands : public TCommands<FSVGImporterEditorCommands>
+class FSVGImporterEditorCommands : public TCommands<FSVGImporterEditorCommands>
 {
 public:
+	SVGIMPORTEREDITOR_API static const FSVGImporterEditorCommands& GetExternal();
+
+	static const FSVGImporterEditorCommands& GetInternal();
+
 	FSVGImporterEditorCommands();
 
 	//~ Begin TCommands
@@ -14,4 +18,8 @@ public:
 	//~ End TCommands
 
 	TSharedPtr<FUICommandInfo> SpawnSVGActor;
+
+private:
+	// Make this unavailable to the public
+	using TCommands<FSVGImporterEditorCommands>::Get;
 };
