Motion Design: Fixed issue with State Task not working as expected in combo pages

#virtualized

diff --git a/Engine/Plugins/Experimental/Avalanche/Source/AvalancheSequence/Private/AvaSequenceLibrary.cpp b/Engine/Plugins/Experimental/Avalanche/Source/AvalancheSequence/Private/AvaSequenceLibrary.cpp
--- a/Engine/Plugins/Experimental/Avalanche/Source/AvalancheSequence/Private/AvaSequenceLibrary.cpp
+++ b/Engine/Plugins/Experimental/Avalanche/Source/AvalancheSequence/Private/AvaSequenceLibrary.cpp
@@ -44,3 +44,16 @@ TScriptInterface<IAvaSequencePlaybackObject> UAvaSequenceLibrary::GetPlaybackObj
 
 	return TScriptInterface<IAvaSequencePlaybackObject>(PlaybackObject->ToUObject());	
 }
+
+FAvaSequencePlayParams UAvaSequenceLibrary::MakeSingleFramePlaySettings(const FAvaSequenceTime& InTargetTime, EAvaSequencePlayMode InPlayMode)
+{
+	FAvaSequencePlayParams PlaySettings;
+
+	// Set start to the largest double, so that it gets clamped down to be End Time (i.e. time that the Sequence should evaluate)
+	// Start is not set to End as the passed in time could be unset (no time constraint) -- which then could resolve to different values for Start/End.
+	PlaySettings.Start = FAvaSequenceTime(TNumericLimits<double>::Max());
+	PlaySettings.End = InTargetTime;
+	PlaySettings.PlayMode = InPlayMode;
+
+	return PlaySettings;
+}
diff --git a/Engine/Plugins/Experimental/Avalanche/Source/AvalancheSequence/Private/Transition/AvaTransitionInitializeSequence.cpp b/Engine/Plugins/Experimental/Avalanche/Source/AvalancheSequence/Private/Transition/AvaTransitionInitializeSequence.cpp
--- a/Engine/Plugins/Experimental/Avalanche/Source/AvalancheSequence/Private/Transition/AvaTransitionInitializeSequence.cpp
+++ b/Engine/Plugins/Experimental/Avalanche/Source/AvalancheSequence/Private/Transition/AvaTransitionInitializeSequence.cpp
@@ -1,6 +1,7 @@
 // Copyright Epic Games, Inc. All Rights Reserved.
 
 #include "Transition/AvaTransitionInitializeSequence.h"
+#include "AvaSequenceLibrary.h"
 #include "AvaSequencePlaybackObject.h"
 #include "AvaTransitionUtils.h"
 #include "Math/NumericLimits.h"
@@ -36,14 +37,9 @@ TArray<UAvaSequencePlayer*> FAvaTransitionInitializeSequence::ExecuteSequenceTas
 		return TArray<UAvaSequencePlayer*>();
 	}
 
-	FAvaSequencePlayParams PlaySettings;
-
 	const FInstanceDataType& InstanceData = InContext.GetInstanceData(*this);
 
-	// Set start to the largest double, so that it gets clamped down to be at the End Time (i.e. time that the Sequence should evaluate)
-	PlaySettings.Start = FAvaSequenceTime(TNumericLimits<double>::Max());
-	PlaySettings.End = InstanceData.InitializeTime;
-	PlaySettings.PlayMode = InstanceData.PlayMode;
+	const FAvaSequencePlayParams PlaySettings = UAvaSequenceLibrary::MakeSingleFramePlaySettings(InstanceData.InitializeTime, InstanceData.PlayMode);
 
 	switch (InstanceData.QueryType)
 	{
diff --git a/Engine/Plugins/Experimental/Avalanche/Source/AvalancheSequence/Public/AvaSequenceLibrary.h b/Engine/Plugins/Experimental/Avalanche/Source/AvalancheSequence/Public/AvaSequenceLibrary.h
--- a/Engine/Plugins/Experimental/Avalanche/Source/AvalancheSequence/Public/AvaSequenceLibrary.h
+++ b/Engine/Plugins/Experimental/Avalanche/Source/AvalancheSequence/Public/AvaSequenceLibrary.h
@@ -6,6 +6,9 @@
 #include "AvaSequenceLibrary.generated.h"
 
 class IAvaSequencePlaybackObject;
+enum class EAvaSequencePlayMode : uint8;
+struct FAvaSequencePlayParams;
+struct FAvaSequenceTime;
 template<typename InInterfaceType> class TScriptInterface;
 
 UCLASS(MinimalAPI, DisplayName = "Motion Design Sequence Library", meta=(ScriptName = "MotionDesignSequenceLibrary"))
@@ -16,4 +19,8 @@ class UAvaSequenceLibrary : public UBlueprintFunctionLibrary
 public:
 	UFUNCTION(BlueprintCallable, Category = "Motion Design|Sequence", meta=(WorldContext = "InWorldContextObject"))
 	static AVALANCHESEQUENCE_API TScriptInterface<IAvaSequencePlaybackObject> GetPlaybackObject(const UObject* InWorldContextObject);
+
+	/** Helper function to build Play Settings for Single Frame Playback */
+	UFUNCTION(BlueprintPure, Category = "Motion Design|Sequence")
+	static AVALANCHESEQUENCE_API FAvaSequencePlayParams MakeSingleFramePlaySettings(const FAvaSequenceTime& InTargetTime, EAvaSequencePlayMode InPlayMode);
 };
diff --git a/Engine/Plugins/Experimental/Avalanche/Source/AvalancheTransition/Private/AvaTransitionLibrary.cpp b/Engine/Plugins/Experimental/Avalanche/Source/AvalancheTransition/Private/AvaTransitionLibrary.cpp
--- a/Engine/Plugins/Experimental/Avalanche/Source/AvalancheTransition/Private/AvaTransitionLibrary.cpp
+++ b/Engine/Plugins/Experimental/Avalanche/Source/AvalancheTransition/Private/AvaTransitionLibrary.cpp
@@ -174,3 +174,12 @@ bool UAvaTransitionLibrary::IsOtherSceneTransitioning(UObject* InTransitionNode,
 
 	return bFoundInstance;
 }
+
+const UAvaTransitionTree* UAvaTransitionLibrary::GetTransitionTree(UObject* InTransitionNode)
+{
+	if (const IAvaTransitionNodeInterface* NodeInterface = Cast<IAvaTransitionNodeInterface>(InTransitionNode))
+	{
+		return NodeInterface->GetBehaviorInstanceCache().GetTransitionTree();
+	}
+	return nullptr;
+}
diff --git a/Engine/Plugins/Experimental/Avalanche/Source/AvalancheTransition/Private/AvaTransitionLibrary.h b/Engine/Plugins/Experimental/Avalanche/Source/AvalancheTransition/Private/AvaTransitionLibrary.h
--- a/Engine/Plugins/Experimental/Avalanche/Source/AvalancheTransition/Private/AvaTransitionLibrary.h
+++ b/Engine/Plugins/Experimental/Avalanche/Source/AvalancheTransition/Private/AvaTransitionLibrary.h
@@ -5,6 +5,7 @@
 #include "Kismet/BlueprintFunctionLibrary.h"
 #include "AvaTransitionLibrary.generated.h"
 
+class UAvaTransitionTree;
 enum class EAvaTransitionComparisonResult : uint8;
 enum class EAvaTransitionLayerCompareType : uint8;
 enum class EAvaTransitionType : uint8;
@@ -27,4 +28,7 @@ public:
 
 	UFUNCTION(BlueprintCallable, Category="Transition Logic", meta=(DefaultToSelf="InTransitionNode"))
 	static bool IsOtherSceneTransitioning(UObject* InTransitionNode, const TSoftObjectPtr<UWorld>& InSceneAsset);
+
+	UFUNCTION(BlueprintCallable, Category="Transition Logic", meta=(DefaultToSelf="InTransitionNode"))
+	static const UAvaTransitionTree* GetTransitionTree(UObject* InTransitionNode);
 };
diff --git a/Engine/Plugins/Experimental/Avalanche/Source/AvalancheTransition/Private/Behavior/AvaTransitionInstanceCache.cpp b/Engine/Plugins/Experimental/Avalanche/Source/AvalancheTransition/Private/Behavior/AvaTransitionInstanceCache.cpp
--- a/Engine/Plugins/Experimental/Avalanche/Source/AvalancheTransition/Private/Behavior/AvaTransitionInstanceCache.cpp
+++ b/Engine/Plugins/Experimental/Avalanche/Source/AvalancheTransition/Private/Behavior/AvaTransitionInstanceCache.cpp
@@ -13,6 +13,15 @@ const FAvaTransitionContext* FAvaTransitionBehaviorInstanceCache::GetTransitionC
 	return nullptr;
 }
 
+const UAvaTransitionTree* FAvaTransitionBehaviorInstanceCache::GetTransitionTree() const
+{
+	if (CachedBehaviorInstance && CachedSceneOwner.IsValid())
+	{
+		return CachedBehaviorInstance->GetTransitionTree();
+	}
+	return nullptr;
+}
+
 void FAvaTransitionBehaviorInstanceCache::SetCachedInstanceDataFromContext(const FStateTreeExecutionContext& InExecutionContext)
 {
 	if (const FAvaTransitionBehaviorInstance* BehaviorInstance = UE::AvaTransition::GetBehaviorInstance(InExecutionContext))
diff --git a/Engine/Plugins/Experimental/Avalanche/Source/AvalancheTransition/Public/AvaTransitionTree.h b/Engine/Plugins/Experimental/Avalanche/Source/AvalancheTransition/Public/AvaTransitionTree.h
--- a/Engine/Plugins/Experimental/Avalanche/Source/AvalancheTransition/Public/AvaTransitionTree.h
+++ b/Engine/Plugins/Experimental/Avalanche/Source/AvalancheTransition/Public/AvaTransitionTree.h
@@ -42,6 +42,7 @@ public:
 		InstancingMode = InInstancingMode;
 	}
 
+	UFUNCTION(BlueprintPure, Category = "Transition Logic")
 	EAvaTransitionInstancingMode GetInstancingMode() const
 	{
 		return InstancingMode;
diff --git a/Engine/Plugins/Experimental/Avalanche/Source/AvalancheTransition/Public/Behavior/AvaTransitionBehaviorInstanceCache.h b/Engine/Plugins/Experimental/Avalanche/Source/AvalancheTransition/Public/Behavior/AvaTransitionBehaviorInstanceCache.h
--- a/Engine/Plugins/Experimental/Avalanche/Source/AvalancheTransition/Public/Behavior/AvaTransitionBehaviorInstanceCache.h
+++ b/Engine/Plugins/Experimental/Avalanche/Source/AvalancheTransition/Public/Behavior/AvaTransitionBehaviorInstanceCache.h
@@ -4,6 +4,7 @@
 
 #include "AvaTransitionScene.h"
 
+class UAvaTransitionTree;
 struct FAvaTransitionBehaviorInstance;
 struct FAvaTransitionContext;
 struct FStateTreeExecutionContext;
@@ -12,6 +13,8 @@ struct FAvaTransitionBehaviorInstanceCache
 {
 	AVALANCHETRANSITION_API const FAvaTransitionContext* GetTransitionContext() const;
 
+	const UAvaTransitionTree* GetTransitionTree() const;
+
 	void SetCachedInstanceDataFromContext(const FStateTreeExecutionContext& InExecutionContext);
 
 	void ClearCachedInstanceData();
