[Motion Design] Display Media Output: Disabling fullscreen on Mac platform.

diff --git a/Engine/Plugins/Experimental/Avalanche/Source/AvalancheMedia/Private/Broadcast/OutputDevices/AvaBroadcastDisplayMediaCapture.cpp b/Engine/Plugins/Experimental/Avalanche/Source/AvalancheMedia/Private/Broadcast/OutputDevices/AvaBroadcastDisplayMediaCapture.cpp
--- a/Engine/Plugins/Experimental/Avalanche/Source/AvalancheMedia/Private/Broadcast/OutputDevices/AvaBroadcastDisplayMediaCapture.cpp
+++ b/Engine/Plugins/Experimental/Avalanche/Source/AvalancheMedia/Private/Broadcast/OutputDevices/AvaBroadcastDisplayMediaCapture.cpp
@@ -82,7 +82,10 @@ public:
 		// This will make the window and show it.
 		FSlateApplication::Get().AddWindow(NewWindow, true);
 
+		// Fullscreen support for Mac causes a lock up. Disabling until it gets fixed.
+#if !PLATFORM_MAC
 		NewWindow->SetWindowMode(EWindowMode::Type::WindowedFullscreen);
+#endif
 
 		Window = NewWindow;
 
