Motion Design: Changed BP State Task to deal with Layers instead of Levels directly

#virtualized

diff --git a/Engine/Plugins/Experimental/Avalanche/Source/AvalancheTransition/Private/AvaTransitionLibrary.cpp b/Engine/Plugins/Experimental/Avalanche/Source/AvalancheTransition/Private/AvaTransitionLibrary.cpp
--- a/Engine/Plugins/Experimental/Avalanche/Source/AvalancheTransition/Private/AvaTransitionLibrary.cpp
+++ b/Engine/Plugins/Experimental/Avalanche/Source/AvalancheTransition/Private/AvaTransitionLibrary.cpp
@@ -89,7 +89,7 @@ EAvaTransitionType UAvaTransitionLibrary::GetTransitionType(UObject* InTransitio
 	return TransitionContext->GetTransitionType();
 }
 
-bool UAvaTransitionLibrary::IsOtherSceneTransitioning(UObject* InTransitionNode, const TSoftObjectPtr<UWorld>& InSceneAsset)
+bool UAvaTransitionLibrary::AreScenesTransitioning(UObject* InTransitionNode, const FAvaTagHandleContainer& InLayers, const TArray<TSoftObjectPtr<UWorld>>& InScenesToIgnore)
 {
 	const IAvaTransitionNodeInterface* NodeInterface = Cast<IAvaTransitionNodeInterface>(InTransitionNode);
 	if (!NodeInterface)
@@ -131,48 +131,62 @@ bool UAvaTransitionLibrary::IsOtherSceneTransitioning(UObject* InTransitionNode,
 		return false;
 	}
 
-	bool bFoundInstance = BehaviorInstances.ContainsByPredicate(
-		[TransitionScene, &InSceneAsset](const FAvaTransitionBehaviorInstance* InInstance)
+	TSet<FString> ScenesToIgnore;
+	ScenesToIgnore.Reserve(InScenesToIgnore.Num());
+	for (const TSoftObjectPtr<UWorld>& SceneToIgnore : InScenesToIgnore)
+	{
+		ScenesToIgnore.Add(SceneToIgnore.GetLongPackageName());
+	}
+
+	for (const FAvaTransitionBehaviorInstance* BehaviorInstance : BehaviorInstances)
+	{
+		if (!InLayers.ContainsTag(BehaviorInstance->GetTransitionLayer()))
 		{
-			const FAvaTransitionScene* const OtherTransitionScene = InInstance->GetTransitionContext().GetTransitionScene();
-
-			// ignore check on scenes marked as needing discard
-			if (!OtherTransitionScene || OtherTransitionScene->HasAnyFlags(EAvaTransitionSceneFlags::NeedsDiscard))
-			{
-				return false;
-			}
-
-			const ULevel* const OtherLevel = OtherTransitionScene->GetLevel();
-			if (!OtherLevel)
-			{
-				return false;
-			}
-
-			const UPackage* const OtherPackage = OtherLevel->GetPackage();
-			if (!OtherPackage)
-			{
-				return false;
-			}
-
-			// Remove the /Temp from the Package Name
-			FString PackageName = OtherPackage->GetName();
-			if (PackageName.StartsWith(TEXT("/Temp")))
-			{
-				PackageName.RightChopInline(-1 + sizeof(TEXT("/Temp")) / sizeof(TCHAR));
-			}
-
-			// Remove the _LevelInstance_[Num] from the Package Name
-			const int32 LevelInstancePosition = PackageName.Find(TEXT("_LevelInstance_"), ESearchCase::IgnoreCase, ESearchDir::FromEnd);
-			if (LevelInstancePosition != INDEX_NONE)
-			{
-				PackageName.LeftChopInline(PackageName.Len() - LevelInstancePosition);
-			}
-
-			const FString SceneAssetName = InSceneAsset.GetLongPackageName();
-			return PackageName == SceneAssetName;
-		});
+			continue;
+		}
+
+		const FAvaTransitionScene* const OtherTransitionScene = BehaviorInstance->GetTransitionContext().GetTransitionScene();
+
+		// skip scenes marked as needing discard
+		if (!OtherTransitionScene || OtherTransitionScene->HasAnyFlags(EAvaTransitionSceneFlags::NeedsDiscard))
+		{
+			continue;
+		}
+
+		const ULevel* const OtherLevel = OtherTransitionScene->GetLevel();
+		if (!OtherLevel)
+		{
+			continue;
+		}
+
+		const UPackage* const OtherPackage = OtherLevel->GetPackage();
+		if (!OtherPackage)
+		{
+			continue;
+		}
+
+		// Remove the /Temp from the Package Name
+		FString PackageName = OtherPackage->GetName();
+		if (PackageName.StartsWith(TEXT("/Temp")))
+		{
+			PackageName.RightChopInline(-1 + sizeof(TEXT("/Temp")) / sizeof(TCHAR));
+		}
+
+		// Remove the _LevelInstance_[Num] from the Package Name
+		const int32 LevelInstancePosition = PackageName.Find(TEXT("_LevelInstance_"), ESearchCase::IgnoreCase, ESearchDir::FromEnd);
+		if (LevelInstancePosition != INDEX_NONE)
+		{
+			PackageName.LeftChopInline(PackageName.Len() - LevelInstancePosition);
+		}
+
+		// If this scene isn't part of those to ignore, then it's a valid transitioning scene
+		if (!ScenesToIgnore.Contains(PackageName))
+		{
+			return true;
+		}
+	}
 
-	return bFoundInstance;
+	return false;
 }
 
 const UAvaTransitionTree* UAvaTransitionLibrary::GetTransitionTree(UObject* InTransitionNode)
diff --git a/Engine/Plugins/Experimental/Avalanche/Source/AvalancheTransition/Private/AvaTransitionLibrary.h b/Engine/Plugins/Experimental/Avalanche/Source/AvalancheTransition/Private/AvaTransitionLibrary.h
--- a/Engine/Plugins/Experimental/Avalanche/Source/AvalancheTransition/Private/AvaTransitionLibrary.h
+++ b/Engine/Plugins/Experimental/Avalanche/Source/AvalancheTransition/Private/AvaTransitionLibrary.h
@@ -27,7 +27,7 @@ public:
 	static EAvaTransitionType GetTransitionType(UObject* InTransitionNode);
 
 	UFUNCTION(BlueprintCallable, Category="Transition Logic", meta=(DefaultToSelf="InTransitionNode"))
-	static bool IsOtherSceneTransitioning(UObject* InTransitionNode, const TSoftObjectPtr<UWorld>& InSceneAsset);
+	static bool AreScenesTransitioning(UObject* InTransitionNode, const FAvaTagHandleContainer& InLayers, const TArray<TSoftObjectPtr<UWorld>>& InScenesToIgnore);
 
 	UFUNCTION(BlueprintCallable, Category="Transition Logic", meta=(DefaultToSelf="InTransitionNode"))
 	static const UAvaTransitionTree* GetTransitionTree(UObject* InTransitionNode);
