Motion Design: Fix Motion Design Mode setting current level editing viewport to null after it has been updated, causing issues with Scene Outliner (not Motion Design) as well as other systems relying on that Global to be correctly set.
The main issue was that the SLevelViewport shared reference was being held throughout the entirety of a function. This would cause LevelViewport to destroy at the end of that function instead of prior to the viewport set.

diff --git a/Engine/Plugins/Experimental/Avalanche/Source/AvalancheEditor/Private/Viewport/AvaLevelViewportExtension.cpp b/Engine/Plugins/Experimental/Avalanche/Source/AvalancheEditor/Private/Viewport/AvaLevelViewportExtension.cpp
--- a/Engine/Plugins/Experimental/Avalanche/Source/AvalancheEditor/Private/Viewport/AvaLevelViewportExtension.cpp
+++ b/Engine/Plugins/Experimental/Avalanche/Source/AvalancheEditor/Private/Viewport/AvaLevelViewportExtension.cpp
@@ -312,15 +312,12 @@ void FAvaLevelViewportExtension::SetViewportType(const TSharedRef<FUICommandInfo
 		return;
 	}
 
-	TSharedPtr<SLevelViewport> ActiveLevelViewport = LevelEditor->GetActiveViewportInterface();
-
-	if (!ActiveLevelViewport.IsValid())
+	TSharedPtr<FUICommandList> CommandList;
+	if (TSharedPtr<SLevelViewport> ActiveLevelViewport = LevelEditor->GetActiveViewportInterface())
 	{
-		return;
+		CommandList = ActiveLevelViewport->GetCommandList();
 	}
 
-	TSharedPtr<FUICommandList> CommandList = ActiveLevelViewport->GetCommandList();
-
 	if (!CommandList.IsValid())
 	{
 		return;
