MotionDesign : AvaLevelViewport
- Adding enable and disable animator buttons in viewport status bar

diff --git a/Engine/Plugins/Experimental/Avalanche/Source/AvalancheEditor/Private/AvaEditorCommands.cpp b/Engine/Plugins/Experimental/Avalanche/Source/AvalancheEditor/Private/AvaEditorCommands.cpp
--- a/Engine/Plugins/Experimental/Avalanche/Source/AvalancheEditor/Private/AvaEditorCommands.cpp
+++ b/Engine/Plugins/Experimental/Avalanche/Source/AvalancheEditor/Private/AvaEditorCommands.cpp
@@ -23,7 +23,6 @@ void FAvaEditorCommands::RegisterCommands()
 	RegisterPivotCommands();
 	RegisterToolsCommands();
 	RegisterAdvancedRenamerCommands();
-	RegisterAnimatorCommands();
 	RegisterClonerEffectorCommands();
 }
 
@@ -415,21 +414,6 @@ void FAvaEditorCommands::RegisterAdvancedRenamerCommands()
 		, FInputChord())
 }
 
-void FAvaEditorCommands::RegisterAnimatorCommands()
-{
-	UI_COMMAND(DisableAnimators
-		, "Disable Animators"
-		, "Disable animators of the selected actors in the level"
-		, EUserInterfaceActionType::Button
-		, FInputChord(EKeys::M, EModifierKey::Control))
-
-	UI_COMMAND(EnableAnimators
-		, "Enable Animators"
-		, "Enable animators of the selected actors in the level"
-		, EUserInterfaceActionType::Button
-		, FInputChord(EKeys::U, EModifierKey::Control))
-}
-
 void FAvaEditorCommands::RegisterClonerEffectorCommands()
 {
 	UI_COMMAND(DisableEffectors
diff --git a/Engine/Plugins/Experimental/Avalanche/Source/AvalancheEditor/Private/AvaEditorCommands.h b/Engine/Plugins/Experimental/Avalanche/Source/AvalancheEditor/Private/AvaEditorCommands.h
--- a/Engine/Plugins/Experimental/Avalanche/Source/AvalancheEditor/Private/AvaEditorCommands.h
+++ b/Engine/Plugins/Experimental/Avalanche/Source/AvalancheEditor/Private/AvaEditorCommands.h
@@ -15,7 +15,6 @@ public:
 	void RegisterPivotCommands();
 	void RegisterToolsCommands();
 	void RegisterAdvancedRenamerCommands();
-	void RegisterAnimatorCommands();
 	void RegisterClonerEffectorCommands();
 
 	/** Switches between Camera "2D" viewport and 3D viewport. */
@@ -96,10 +95,6 @@ public:
 	TSharedPtr<FUICommandInfo> OpenAdvancedRenamerTool_SelectedActors;
 	TSharedPtr<FUICommandInfo> OpenAdvancedRenamerTool_SharedClassActors;
 
-	// Animator
-	TSharedPtr<FUICommandInfo> DisableAnimators;
-	TSharedPtr<FUICommandInfo> EnableAnimators;
-
 	// ClonerEffector
 	TSharedPtr<FUICommandInfo> DisableCloners;
 	TSharedPtr<FUICommandInfo> EnableCloners;
diff --git a/Engine/Plugins/Experimental/Avalanche/Source/AvalancheEditor/Private/OperatorStack/AvaOperatorStackExtension.cpp b/Engine/Plugins/Experimental/Avalanche/Source/AvalancheEditor/Private/OperatorStack/AvaOperatorStackExtension.cpp
--- a/Engine/Plugins/Experimental/Avalanche/Source/AvalancheEditor/Private/OperatorStack/AvaOperatorStackExtension.cpp
+++ b/Engine/Plugins/Experimental/Avalanche/Source/AvalancheEditor/Private/OperatorStack/AvaOperatorStackExtension.cpp
@@ -2,6 +2,7 @@
 
 #include "AvaOperatorStackExtension.h"
 #include "AvaEditorCommands.h"
+#include "AvaLevelViewportCommands.h"
 #include "AvaOperatorStackTabSpawner.h"
 #include "Components/PropertyAnimatorCoreComponent.h"
 #include "EditorModeManager.h"
@@ -89,12 +90,12 @@ void FAvaOperatorStackExtension::BindCommands(const TSharedRef<FUICommandList>&
 {
 	InCommandList->Append(AnimatorCommands);
 
-	const FAvaEditorCommands& EditorCommands = FAvaEditorCommands::Get();
+	const FAvaLevelViewportCommands& ViewportCommands = FAvaLevelViewportCommands::GetExternal();
 
-	AnimatorCommands->MapAction(EditorCommands.DisableAnimators
+	AnimatorCommands->MapAction(ViewportCommands.DisableAnimators
 		, FExecuteAction::CreateSP(this, &FAvaOperatorStackExtension::EnableAnimators, false));
 
-	AnimatorCommands->MapAction(EditorCommands.EnableAnimators
+	AnimatorCommands->MapAction(ViewportCommands.EnableAnimators
 		, FExecuteAction::CreateSP(this, &FAvaOperatorStackExtension::EnableAnimators, true));
 }
 
diff --git a/Engine/Plugins/Experimental/Avalanche/Source/AvalancheLevelViewport/AvalancheLevelViewport.Build.cs b/Engine/Plugins/Experimental/Avalanche/Source/AvalancheLevelViewport/AvalancheLevelViewport.Build.cs
--- a/Engine/Plugins/Experimental/Avalanche/Source/AvalancheLevelViewport/AvalancheLevelViewport.Build.cs
+++ b/Engine/Plugins/Experimental/Avalanche/Source/AvalancheLevelViewport/AvalancheLevelViewport.Build.cs
@@ -43,6 +43,7 @@ public class AvalancheLevelViewport : ModuleRules
 				"InputCore",
 				"LevelEditor",
 				"Projects",
+				"PropertyAnimatorCore",
 				"PropertyEditor",
 				"ToolMenus",
 				"TypedElementRuntime"
diff --git a/Engine/Plugins/Experimental/Avalanche/Source/AvalancheLevelViewport/Private/AvaLevelViewportCommands.cpp b/Engine/Plugins/Experimental/Avalanche/Source/AvalancheLevelViewport/Private/AvaLevelViewportCommands.cpp
--- a/Engine/Plugins/Experimental/Avalanche/Source/AvalancheLevelViewport/Private/AvaLevelViewportCommands.cpp
+++ b/Engine/Plugins/Experimental/Avalanche/Source/AvalancheLevelViewport/Private/AvaLevelViewportCommands.cpp
@@ -39,6 +39,7 @@ void FAvaLevelViewportCommands::RegisterCommands()
 	RegisterVirtualSizeCommands();
 	RegisterGuideCommands();
 	RegisterTransformCommands();
+	RegisterAnimatorCommands();
 }
 
 void FAvaLevelViewportCommands::RegisterViewportCommands()
@@ -351,4 +352,19 @@ void FAvaLevelViewportCommands::RegisterTransformCommands()
 		, FInputChord(EKeys::T, EModifierKey::Alt))
 }
 
+void FAvaLevelViewportCommands::RegisterAnimatorCommands()
+{
+	UI_COMMAND(DisableAnimators
+	, "Disable Animators"
+	, "Disable animators of the selected actors in the level"
+	, EUserInterfaceActionType::Button
+	, FInputChord(EKeys::M, EModifierKey::Control))
+
+	UI_COMMAND(EnableAnimators
+	, "Enable Animators"
+	, "Enable animators of the selected actors in the level"
+	, EUserInterfaceActionType::Button
+	, FInputChord(EKeys::U, EModifierKey::Control))
+}
+
 #undef LOCTEXT_NAMESPACE
diff --git a/Engine/Plugins/Experimental/Avalanche/Source/AvalancheLevelViewport/Private/Widgets/SAvaLevelViewportStatusBarButtons.cpp b/Engine/Plugins/Experimental/Avalanche/Source/AvalancheLevelViewport/Private/Widgets/SAvaLevelViewportStatusBarButtons.cpp
--- a/Engine/Plugins/Experimental/Avalanche/Source/AvalancheLevelViewport/Private/Widgets/SAvaLevelViewportStatusBarButtons.cpp
+++ b/Engine/Plugins/Experimental/Avalanche/Source/AvalancheLevelViewport/Private/Widgets/SAvaLevelViewportStatusBarButtons.cpp
@@ -18,6 +18,8 @@
 #include "SAvaViewportInfo.h"
 #include "Selection.h"
 #include "Styling/AppStyle.h"
+#include "Styling/SlateIconFinder.h"
+#include "Subsystems/PropertyAnimatorCoreSubsystem.h"
 #include "ThumbnailRendering/ThumbnailManager.h"
 #include "ToolMenu.h"
 #include "ToolMenus.h"
@@ -286,6 +288,36 @@ void SAvaLevelViewportStatusBarButtons::PopulateActorButtons(TSharedPtr<SHorizon
 		[
 			AlignmentButton
 		];
+
+	const FSlateBrush* AnimatorBrush = FSlateIconFinder::FindIconBrushForClass(UPropertyAnimatorCoreBase::StaticClass());
+
+	InContainer->AddSlot()
+		.AutoWidth()
+		.Padding(ViewportStatusBarButton::Padding)
+		[
+			ViewportStatusBarButton::MakeButton(
+				this,
+				CommandsRef.DisableAnimators,
+				AnimatorBrush,
+				&SAvaLevelViewportStatusBarButtons::DisableAnimators,
+				&SAvaLevelViewportStatusBarButtons::GetAnimatorButtonEnabled,
+				&SAvaLevelViewportStatusBarButtons::GetAnimatorButtonMuteColor
+			)
+		];
+
+	InContainer->AddSlot()
+		.AutoWidth()
+		.Padding(ViewportStatusBarButton::Padding)
+		[
+			ViewportStatusBarButton::MakeButton(
+				this,
+				CommandsRef.EnableAnimators,
+				AnimatorBrush,
+				&SAvaLevelViewportStatusBarButtons::EnableAnimators,
+				&SAvaLevelViewportStatusBarButtons::GetAnimatorButtonEnabled,
+				&SAvaLevelViewportStatusBarButtons::GetAnimatorButtonUnmuteColor
+			)
+		];
 }
 
 void SAvaLevelViewportStatusBarButtons::PopulateViewportButtons(TSharedPtr<SHorizontalBox> InContainer)
@@ -818,6 +850,113 @@ TSharedRef<SWidget> SAvaLevelViewportStatusBarButtons::GetActorColorMenuContent(
 	return SAvaLevelViewportActorColorMenu::CreateMenu(LevelEditor.ToSharedRef());
 }
 
+FSlateColor SAvaLevelViewportStatusBarButtons::GetAnimatorButtonMuteColor() const
+{
+	return UE::AvaLevelViewport::Private::ViewportStatusBarButton::EnabledColor;
+}
+
+FSlateColor SAvaLevelViewportStatusBarButtons::GetAnimatorButtonUnmuteColor() const
+{
+	return UE::AvaLevelViewport::Private::ViewportStatusBarButton::ActiveColor;
+}
+
+bool SAvaLevelViewportStatusBarButtons::GetAnimatorButtonEnabled() const
+{
+	return true;
+}
+
+FReply SAvaLevelViewportStatusBarButtons::EnableAnimators()
+{
+	const TSharedPtr<SAvaLevelViewportFrame> ViewportFrame = ViewportFrameWeak.Pin();
+
+	if (!ViewportFrame.IsValid())
+	{
+		return FReply::Handled();
+	}
+
+	const TSharedPtr<FAvaLevelViewportClient> ViewportClient = ViewportFrame->GetViewportClient();
+
+	if (!ViewportClient.IsValid())
+	{
+		return FReply::Handled();
+	}
+
+	const FEditorModeTools* ModeTools = ViewportClient->GetModeTools();
+
+	if (!ModeTools)
+	{
+		return FReply::Handled();
+	}
+
+	UPropertyAnimatorCoreSubsystem* AnimatorSubsystem = UPropertyAnimatorCoreSubsystem::Get();
+	const UWorld* World = ModeTools->GetWorld();
+	const UTypedElementSelectionSet* SelectionSet = ModeTools->GetEditorSelectionSet();
+
+	if (!World || !SelectionSet || !AnimatorSubsystem)
+	{
+		return FReply::Handled();
+	}
+
+	const TSet<AActor*> SelectedActors(SelectionSet->GetSelectedObjects<AActor>());
+
+	if (SelectedActors.IsEmpty())
+	{
+		AnimatorSubsystem->SetLevelAnimatorsEnabled(World, /** Enabled */true, /** Transact */true);
+	}
+	else
+	{
+		AnimatorSubsystem->SetActorAnimatorsEnabled(SelectedActors, /** Enabled */true, /** Transact */true);
+	}
+
+	return FReply::Handled();
+}
+
+FReply SAvaLevelViewportStatusBarButtons::DisableAnimators()
+{
+	const TSharedPtr<SAvaLevelViewportFrame> ViewportFrame = ViewportFrameWeak.Pin();
+
+	if (!ViewportFrame.IsValid())
+	{
+		return FReply::Handled();
+	}
+
+	const TSharedPtr<FAvaLevelViewportClient> ViewportClient = ViewportFrame->GetViewportClient();
+
+	if (!ViewportClient.IsValid())
+	{
+		return FReply::Handled();
+	}
+
+	const FEditorModeTools* ModeTools = ViewportClient->GetModeTools();
+
+	if (!ModeTools)
+	{
+		return FReply::Handled();
+	}
+
+	UPropertyAnimatorCoreSubsystem* AnimatorSubsystem = UPropertyAnimatorCoreSubsystem::Get();
+	const UWorld* World = ModeTools->GetWorld();
+	const UTypedElementSelectionSet* SelectionSet = ModeTools->GetEditorSelectionSet();
+
+	if (!World || !SelectionSet || !AnimatorSubsystem)
+	{
+		return FReply::Handled();
+	}
+
+	const TSet<AActor*> SelectedActors(SelectionSet->GetSelectedObjects<AActor>());
+
+	if (SelectedActors.IsEmpty())
+	{
+		AnimatorSubsystem->SetLevelAnimatorsEnabled(World, /** Enabled */false, /** Transact */true);
+	}
+	else
+	{
+		AnimatorSubsystem->SetActorAnimatorsEnabled(SelectedActors, /** Enabled */false, /** Transact */true);
+	}
+
+	return FReply::Handled();
+}
+
 FSlateColor SAvaLevelViewportStatusBarButtons::GetToggleSnapColor() const
 {
 	using namespace UE::AvaLevelViewport::Private;
diff --git a/Engine/Plugins/Experimental/Avalanche/Source/AvalancheLevelViewport/Private/Widgets/SAvaLevelViewportStatusBarButtons.h b/Engine/Plugins/Experimental/Avalanche/Source/AvalancheLevelViewport/Private/Widgets/SAvaLevelViewportStatusBarButtons.h
--- a/Engine/Plugins/Experimental/Avalanche/Source/AvalancheLevelViewport/Private/Widgets/SAvaLevelViewportStatusBarButtons.h
+++ b/Engine/Plugins/Experimental/Avalanche/Source/AvalancheLevelViewport/Private/Widgets/SAvaLevelViewportStatusBarButtons.h
@@ -45,9 +45,14 @@ protected:
 	FSlateColor GetActorAlignmentColor() const;
 	bool GetActorAlignmentEnabled() const;
 	TSharedRef<SWidget> GetActorAlignmentMenuContent() const;
-
 	TSharedRef<SWidget> GetActorColorMenuContent() const;
 
+	FSlateColor GetAnimatorButtonMuteColor() const;
+	FSlateColor GetAnimatorButtonUnmuteColor() const;
+	bool GetAnimatorButtonEnabled() const;
+	FReply EnableAnimators();
+	FReply DisableAnimators();
+
 	// Post Process
 	FSlateColor GetPostProcessColor() const;
 	bool GetPostProcessEnabled() const;
@@ -114,4 +119,4 @@ protected:
 	bool GetTextureOverlayEnabled() const;
 	TSharedRef<SWidget> GetTextureOverlayMenuContent();
 	FReply ToggleTextureOverlay();
-};
\ No newline at end of file
+};
diff --git a/Engine/Plugins/Experimental/Avalanche/Source/AvalancheLevelViewport/Public/AvaLevelViewportCommands.h b/Engine/Plugins/Experimental/Avalanche/Source/AvalancheLevelViewport/Public/AvaLevelViewportCommands.h
--- a/Engine/Plugins/Experimental/Avalanche/Source/AvalancheLevelViewport/Public/AvaLevelViewportCommands.h
+++ b/Engine/Plugins/Experimental/Avalanche/Source/AvalancheLevelViewport/Public/AvaLevelViewportCommands.h
@@ -79,6 +79,10 @@ public:
 	TSharedPtr<FUICommandInfo> ResetScale;
 	TSharedPtr<FUICommandInfo> ResetTransform;
 
+	// Animator
+	TSharedPtr<FUICommandInfo> DisableAnimators;
+	TSharedPtr<FUICommandInfo> EnableAnimators;
+
 private:
 	// Make this unavailable to the public
 	using TCommands<FAvaLevelViewportCommands>::Get;
@@ -90,4 +94,5 @@ private:
 	void RegisterVirtualSizeCommands();
 	void RegisterGuideCommands();
 	void RegisterTransformCommands();
+	void RegisterAnimatorCommands();
 };
