MotionDesign : ActorModifier
- Fix CachedSortedComponentStates array has changed during range for iteration in translucent sort priority modifier

diff --git a/Engine/Plugins/Experimental/Avalanche/Source/AvalancheModifiers/Private/Modifiers/AvaTranslucentPriorityModifier.cpp b/Engine/Plugins/Experimental/Avalanche/Source/AvalancheModifiers/Private/Modifiers/AvaTranslucentPriorityModifier.cpp
--- a/Engine/Plugins/Experimental/Avalanche/Source/AvalancheModifiers/Private/Modifiers/AvaTranslucentPriorityModifier.cpp
+++ b/Engine/Plugins/Experimental/Avalanche/Source/AvalancheModifiers/Private/Modifiers/AvaTranslucentPriorityModifier.cpp
@@ -219,7 +219,11 @@ void UAvaTranslucentPriorityModifier::Apply()
 			else if (Component->TranslucencySortPriority != TranslucentSortPriority)
 			{
 				// Cache to avoid doing the same query for the same result
-				ComponentModifier->CachedSortedComponentStates = CachedSortedComponentStates;
+				if (ComponentModifier->CachedSortedComponentStates.IsEmpty())
+				{
+					ComponentModifier->CachedSortedComponentStates = CachedSortedComponentStates;
+				}
+
 				ComponentModifier->MarkModifierDirty();
 			}
 
