[Motion Design] Documenting the Rundown Server's messages to be able to generate a more complete websocket api.

#rnx

diff --git a/Engine/Plugins/Experimental/Avalanche/Source/AvalancheMedia/Internal/Rundown/AvaRundownMessages.h b/Engine/Plugins/Experimental/Avalanche/Source/AvalancheMedia/Internal/Rundown/AvaRundownMessages.h
--- a/Engine/Plugins/Experimental/Avalanche/Source/AvalancheMedia/Internal/Rundown/AvaRundownMessages.h
+++ b/Engine/Plugins/Experimental/Avalanche/Source/AvalancheMedia/Internal/Rundown/AvaRundownMessages.h
@@ -69,35 +69,54 @@ enum class EAvaRundownServerEngineMode : uint8
 	Other
 };
 
-
+/** Base class for all rundown server messages. */
 USTRUCT()
 struct FAvaRundownMsgBase
 {
 	GENERATED_BODY()
 public:
+	/**
+	 * Request Identifier (client assigned) for matching server responses with their corresponding requests. 
+	 */
 	UPROPERTY()
 	int32 RequestId = INDEX_NONE;
 };
 
+/**
+ * This message is the default response message for all requests, unless a specific response message type
+ * is specified for the request.
+ * On success, the message will have a Verbosity of "Log" and the text may contain response payload related data.
+ * On failure, a message with Verbosity "Error" will be sent.
+ * This message's RequestId mirrors that of the corresponding request from the client.
+ */
 USTRUCT()
 struct FAvaRundownServerMsg : public FAvaRundownMsgBase
 {
 	GENERATED_BODY()
 public:
+	/**
+	 * Debug, Log, Warning, Error, etc.
+	 */
 	UPROPERTY()
 	FString Verbosity;
 
+	/**
+	 * Message Text.
+	 */
 	UPROPERTY()
 	FString Text;
 };
 
-/** Request published by client to discover servers. */
+/**
+ * Request published by client to discover servers on the message bus.
+ * The available servers will respond with a FAvaRundownPong.
+ */
 USTRUCT()
 struct FAvaRundownPing : public FAvaRundownMsgBase
 {
 	GENERATED_BODY()
 public:
-	/** True if the request originates from an automatic timer. False if requests originates from user interaction. */
+	/** True if the request originates from an automatic timer. False if request originates from user interaction. */
 	UPROPERTY()
 	bool bAuto = true;
 
@@ -109,7 +128,10 @@ public:
 	int32 RequestedApiVersion = EAvaRundownApiVersion::Unspecified;
 };
 
-/** Response sent by server to client to be discovered. */
+/**
+ * The server will send this message to the client in response to FAvaRundownPing.
+ * This is used to discover the server's entry point on the message bus.
+ */
 USTRUCT()
 struct FAvaRundownPong : public FAvaRundownMsgBase
 {
@@ -137,12 +159,14 @@ public:
 	UPROPERTY()
 	int32 LatestApiVersion = EAvaRundownApiVersion::Unspecified;
 
+	/** Server Host Name */
 	UPROPERTY()
 	FString HostName;
 };
 
 /**
- * Request the extended server information.
+ * Requests the extended server information.
+ * Response is FAvaRundownServerInfo.
  */
 USTRUCT()
 struct FAvaRundownGetServerInfo : public FAvaRundownMsgBase
@@ -151,7 +175,7 @@ struct FAvaRundownGetServerInfo : public FAvaRundownMsgBase
 };
 
 /**
- * Extended server information 
+ * Extended server information.
  */
 USTRUCT()
 struct FAvaRundownServerInfo : public FAvaRundownMsgBase
@@ -170,6 +194,7 @@ struct FAvaRundownServerInfo : public FAvaRundownMsgBase
 	UPROPERTY()
 	int32 LatestApiVersion = EAvaRundownApiVersion::Unspecified;
 
+	/** Server Host Name */
 	UPROPERTY()
 	FString HostName;
 
@@ -177,7 +202,7 @@ struct FAvaRundownServerInfo : public FAvaRundownMsgBase
 	UPROPERTY()
 	uint32 EngineVersion = 0;
 	
-	/** Holds the instance identifier. */
+	/** Application Instance Identifier. */
 	UPROPERTY()
 	FGuid InstanceId;
 
@@ -209,7 +234,8 @@ struct FAvaRundownServerInfo : public FAvaRundownMsgBase
 };
 
 /**
- *	Request list of rundown that can be opened on the current server.
+ *	Requests the list of rundowns that can be opened on the current server.
+ *	Response is FAvaRundownRundowns.
  */
 USTRUCT()
 struct FAvaRundownGetRundowns : public FAvaRundownMsgBase
@@ -227,12 +253,15 @@ struct FAvaRundownRundowns : public FAvaRundownMsgBase
 	GENERATED_BODY()
 
 public:
+	/**
+	 * List of Rundown asset paths in format: [PackagePath]/[AssetName].[AssetName]
+	 */
 	UPROPERTY()
 	TArray<FString> Rundowns;
 };
 
 /**
- *	Request that the given rundown be loaded for playback.
+ *	Loads the given rundown for playback operations.
  *	This will also open an associated playback context.
  *	Only one rundown can be opened for playback at a time by the rundown server.
  *	If another rundown is opened, the previous one will be closed and all currently playing pages stopped,
@@ -255,11 +284,13 @@ public:
 };
 
 /**
- * Request to create a new rundown asset.
+ * Creates a new rundown asset.
  *
  * The full package name is going to be: [PackagePath]/[AssetName] 
  * The full asset path is going to be: [PackagePath]/[AssetName].[AssetName]
  * For all other requests, the rundown reference is the full asset path.
+ *
+ * Response is FAvaRundownServerMsg.
  */
 USTRUCT()
 struct FAvaRundownCreateRundown : public FAvaRundownMsgBase
@@ -284,7 +315,9 @@ public:
 };
 
 /**
- * Request a previously created rundown to be deleted or at least no longer managed (if transient only). 
+ * Deletes an existing rundown.
+ *
+ * Response is FAvaRundownServerMsg.
  */
 USTRUCT()
 struct FAvaRundownDeleteRundown : public FAvaRundownMsgBase
@@ -298,7 +331,7 @@ public:
 };
 
 /**
- * Import rundown from json data or file.
+ * Imports rundown from json data or file.
  */
 USTRUCT()
 struct FAvaRundownImportRundown : public FAvaRundownMsgBase
@@ -324,7 +357,7 @@ public:
 };
 
 /**
- * Export a rundown to json data or file.
+ * Exports a rundown to json data or file.
  * This command is supported in game build.
  */
 USTRUCT()
@@ -361,7 +394,7 @@ public:
 };
 
 /**
- * Request that the given rundown be saved to disk.
+ * Requests that the given rundown be saved to disk.
  * The rundown asset must have been loaded, either by an edit command
  * or playback, prior to this command.
  * Unloaded assets will not be loaded by this command.
@@ -373,9 +406,11 @@ struct FAvaRundownSaveRundown : public FAvaRundownMsgBase
 	GENERATED_BODY()
 
 public:
+	/** Rundown asset path: [PackagePath]/[AssetName].[AssetName] */
 	UPROPERTY()
 	FString Rundown;
 
+	/** The save command will be executed only if the asset package is dirty. */
 	UPROPERTY()
 	bool bOnlyIfIsDirty = false;
 };
@@ -403,118 +438,173 @@ public:
 };
 
 /**
- * Request the list of pages from the given rundown.
+ * Requests the list of pages from the given rundown.
  */
 USTRUCT()
 struct FAvaRundownGetPages : public FAvaRundownMsgBase
 {
 	GENERATED_BODY()
 public:
+	/** Rundown asset path: [PackagePath]/[AssetName].[AssetName] */
 	UPROPERTY()
 	FString Rundown;
 };
 
+/**
+ * Requests a new page be created from the specified template in the given rundown.
+ */
 USTRUCT()
 struct FAvaRundownCreatePage : public FAvaRundownMsgBase
 {
 	GENERATED_BODY()
 
+	/** Rundown asset path: [PackagePath]/[AssetName].[AssetName] */
 	UPROPERTY()
 	FString Rundown;
 
+	/** Specifies the template for the newly created page. */
 	UPROPERTY()
 	int32 TemplateId = FAvaRundownPage::InvalidPageId;
 };
 
+/**
+ * Requests the page be deleted from the given rundown.
+ */
 USTRUCT()
 struct FAvaRundownDeletePage : public FAvaRundownMsgBase
 {
 	GENERATED_BODY()
 
+	/** Rundown asset path: [PackagePath]/[AssetName].[AssetName] */
 	UPROPERTY()
 	FString Rundown;
 
+	/** Id of the page to be deleted. */
 	UPROPERTY()
 	int32 PageId = FAvaRundownPage::InvalidPageId;
 };
 
+/**
+ * Requests the creation of a new template.
+ * If successful, the response is FAvaRundownServerMsg with a "Template [Id] Created" text.
+ * The id of the created template can be parsed from that message's text.
+ * Also a secondary FAvaRundownPageListChanged event with added template id will be sent.
+ */
 USTRUCT()
 struct FAvaRundownCreateTemplate : public FAvaRundownMsgBase
 {
 	GENERATED_BODY()
 
+	/** Rundown asset path: [PackagePath]/[AssetName].[AssetName] */
 	UPROPERTY()
 	FString Rundown;
 };
 
+/** Requests deletion of the given template. */
 USTRUCT()
 struct FAvaRundownDeleteTemplate : public FAvaRundownMsgBase
 {
 	GENERATED_BODY()
 
+	/** Rundown asset path: [PackagePath]/[AssetName].[AssetName] */
 	UPROPERTY()
 	FString Rundown;
 
+	/** Specifies the *template* id to delete. */
 	UPROPERTY()
 	int32 PageId = FAvaRundownPage::InvalidPageId;
 };
 
+/**
+ * Sets the Page's template asset. This applies to template pages only.
+ */
 USTRUCT()
 struct FAvaRundownChangeTemplateBP : public FAvaRundownMsgBase
 {
 	GENERATED_BODY()
 public:
+	/** Rundown asset path: [PackagePath]/[AssetName].[AssetName] */
 	UPROPERTY()
 	FString Rundown;
 
+	/** Specifies the template id to modify. */
 	UPROPERTY()
 	int32 TemplateId = FAvaRundownPage::InvalidPageId;
 
+	/** Specifies the asset path to assign. */
 	UPROPERTY()
 	FString AssetPath;
 };
 
+/** Page Information */
 USTRUCT()
 struct FAvaRundownPageInfo
 {
 	GENERATED_BODY()
 public:
+	/** Unique identifier for the page within the rundown. */
 	UPROPERTY()
 	int32 PageId = FAvaRundownPage::InvalidPageId;
 
+	/**
+	 * Short page name, usually the asset name for templates.
+	 * It is displayed as the page description if there is no page summary or user friendly name specified.
+	 */
 	UPROPERTY()
 	FString PageName;
 
+	/**
+	 * Summary is generated from the remote control values for this page.
+	 * It is displayed as the page description if there is no user friendly name specified.
+	 */
 	UPROPERTY()
 	FString PageSummary;
 
+	/** User editable page description. If not empty, this should be used as the page description. */
 	UPROPERTY()
 	FString FriendlyName;
 
+	/** Indicates if the page is a template (true) or an instance (false). */
 	UPROPERTY()
 	bool IsTemplate = false;
 
+	/** Page Instance property: Template Id for this page. */
 	UPROPERTY()
 	int32 TemplateId = FAvaRundownPage::InvalidPageId;
 
+	/** Template property: For combination template, lists the templates that are combined. */ 
 	UPROPERTY()
 	TArray<int32> CombinedTemplateIds;
 
+	/** Template property: playable asset path for this template. */
 	UPROPERTY()
 	FSoftObjectPath AssetPath;
 
+	/**
+	 * List of page channel statuses.
+	 * There will be an entry for each channel the page is playing/previewing in.
+	 */
 	UPROPERTY()
 	TArray<FAvaRundownChannelPageStatus> Statuses;
 
+	/** Transition Layer Name (indicates the page has transition logic). */
 	UPROPERTY()
 	FString TransitionLayerName;
 
+	/** Indicate if the template asset has transition logic. */
+	UPROPERTY()
+	bool bTransitionLogicEnabled = false;
+
 	UPROPERTY()
 	FString OutputChannel;
 
+	/** Specifies if the page is enabled (i.e. can be played). */
 	UPROPERTY()
 	bool bIsEnabled = false;
 
+	/**
+	 * Indicates if the page is currently playing in it's program channel. 
+	 */
 	UPROPERTY()
 	bool bIsPlaying = false;
 };
@@ -527,12 +617,14 @@ struct FAvaRundownPages : public FAvaRundownMsgBase
 {
 	GENERATED_BODY()
 public:
+	/** List of page descriptors */
 	UPROPERTY()
 	TArray<FAvaRundownPageInfo> Pages;
 };
 
 /**
- * Request the page details from the given rundown.
+ * Requests the page details from the given rundown.
+ * Response is FAvaRundownPageDetails.
  */
 USTRUCT()
 struct FAvaRundownGetPageDetails : public FAvaRundownMsgBase
@@ -540,9 +632,11 @@ struct FAvaRundownGetPageDetails : public FAvaRundownMsgBase
 	GENERATED_BODY()
 
 public:
+	/** Rundown asset path: [PackagePath]/[AssetName].[AssetName] */
 	UPROPERTY()
 	FString Rundown;
 
+	/** Specified the requested page id. */
 	UPROPERTY()
 	int32 PageId = FAvaRundownPage::InvalidPageId;
 
@@ -561,12 +655,15 @@ struct FAvaRundownPageDetails : public FAvaRundownMsgBase
 	GENERATED_BODY()
 
 public:
+	/** Rundown asset path: [PackagePath]/[AssetName].[AssetName] */
 	UPROPERTY()
 	FString Rundown;
 
+	/** Page Information. */
 	UPROPERTY()
 	FAvaRundownPageInfo PageInfo;
 
+	/** Remote Control Values for this page. */
 	UPROPERTY()
 	FAvaPlayableRemoteControlValues RemoteControlValues;
 
@@ -574,103 +671,152 @@ public:
 	UPROPERTY()
 	FString RemoteControlPresetName;
 
+	/** Uuid of the remote control preset to resolve through WebRC API. */
 	UPROPERTY()
 	FString RemoteControlPresetId;
 };
 
+/**
+ * Event sent when a page status changes.
+ */
 USTRUCT()
 struct FAvaRundownPagesStatuses : public FAvaRundownMsgBase
 {
 	GENERATED_BODY()
 
 public:
+	/** Rundown asset path: [PackagePath]/[AssetName].[AssetName] */
 	UPROPERTY()
 	FString Rundown;
 
+	/** Page Information. */
 	UPROPERTY()
 	FAvaRundownPageInfo PageInfo;
 };
 
+/**
+ * Event sent when a page list (can be templates, pages or page views) has been modified.
+ */
 USTRUCT()
 struct FAvaRundownPageListChanged : public FAvaRundownMsgBase
 {
 	GENERATED_BODY()
 
+	/** Rundown asset path: [PackagePath]/[AssetName].[AssetName] */
 	UPROPERTY()
 	FString Rundown;
 
+	/** Specifies which page list has been modified. */
 	UPROPERTY()
 	EAvaRundownPageListType ListType = EAvaRundownPageListType::Instance;
 
+	/** Specifies the uuid of the page view, in case the event concerns a page view. */
 	UPROPERTY()
 	FGuid SubListId;
 	
-	/** See EAvaPageListChange flags. */
+	/**
+	 * Bitfield value indicating what has changed:
+	 * - bit 0: Added Pages
+	 * - bit 1: Remove Pages
+	 * - bit 2: Page Id Renumbered
+	 * - bit 3: Sublist added or removed
+	 * - bit 4: Sublist renamed
+	 * - bit 5: Page View reordered
+	 * 
+	 * See EAvaPageListChange flags.
+	 */
 	UPROPERTY()
 	uint8 ChangeType = 0;
 
+	/** List of page Ids affected by this event. */
 	UPROPERTY();
 	TArray<int32> AffectedPages;
 };
 
+/**
+ * Event sent when a page's asset is modified.
+ * Note: this applies to templates only.
+ */
 USTRUCT()
 struct FAvaRundownPageBlueprintChanged : public FAvaRundownMsgBase
 {
 	GENERATED_BODY()
 
+	/** Rundown asset path: [PackagePath]/[AssetName].[AssetName] */
 	UPROPERTY()
 	FString Rundown;
 
+	/** Specified the modified page id. */
 	UPROPERTY()
 	int32 PageId = FAvaRundownPage::InvalidPageId;
 
+	/** Asset the page is currently assigned to (post modification). */
 	UPROPERTY()
 	FString BlueprintPath;
 };
 
+/**
+ * Event sent when a page's channel is modified.
+ */
 USTRUCT()
 struct FAvaRundownPageChannelChanged : public FAvaRundownMsgBase
 {
 	GENERATED_BODY()
 
+	/** Rundown asset path: [PackagePath]/[AssetName].[AssetName] */
 	UPROPERTY()
 	FString Rundown;
 
+	/** Specified the modified page id. */
 	UPROPERTY()
 	int32 PageId = FAvaRundownPage::InvalidPageId;
 
+	/** Channel the page is currently assigned to (post modification). */
 	UPROPERTY()
 	FString ChannelName;
 };
 
+/**
+ * Event sent when a page's animation settings is modified.
+ */
 USTRUCT()
 struct FAvaRundownPageAnimSettingsChanged : public FAvaRundownMsgBase
 {
 	GENERATED_BODY()
 
+	/** Rundown asset path: [PackagePath]/[AssetName].[AssetName] */
 	UPROPERTY()
 	FString Rundown;
 
+	/** Specified the modified page id. */
 	UPROPERTY()
 	int32 PageId = FAvaRundownPage::InvalidPageId;
 };
 
+/**
+ * Sets the channel of the given page.
+ * The page must be valid (and not a template) and the channel must exist in the current profile.
+ * Along with the corresponding response, this will also trigger a FAvaRundownPageChannelChanged event.
+ */
 USTRUCT()
 struct FAvaRundownPageChangeChannel : public FAvaRundownMsgBase
 {
 	GENERATED_BODY()
 
+	/** Rundown asset path: [PackagePath]/[AssetName].[AssetName] */
 	UPROPERTY()
 	FString Rundown;
 
+	/** Specifies the page that will be modified. */
 	UPROPERTY()
 	int32 PageId = FAvaRundownPage::InvalidPageId;
 
+	/** Specifies a valid channel to set for the specified page. */
 	UPROPERTY()
 	FString ChannelName;
 };
 
-/** This is a request to save the managed RCP back to the corresponding page. */
+/** This is a request to save the managed Remote Control Preset (RCP) back to the corresponding page values. */
 USTRUCT()
 struct FAvaRundownUpdatePageFromRCP : public FAvaRundownMsgBase
 {
@@ -686,7 +832,7 @@ public:
 UENUM()
 enum class EAvaRundownPageActions
 {
-	None,
+	None UMETA(Hidden),
 	Load,
 	Unload,
 	Play,
@@ -698,30 +844,38 @@ enum class EAvaRundownPageActions
 	TakeToProgram
 };
 
+/**
+ * Request for a program page command on the current playback rundown.
+ */
 USTRUCT()
 struct FAvaRundownPageAction : public FAvaRundownMsgBase
 {
 	GENERATED_BODY()
 public:
+	/** Specifies the Page Id that is the target of this action command. */
 	UPROPERTY()
 	int32 PageId = FAvaRundownPage::InvalidPageId;
 
+	/** Specifies the page action to execute. */
 	UPROPERTY()
 	EAvaRundownPageActions Action = EAvaRundownPageActions::None;
 };
 
+/**
+ * Request for a preview page command on the current playback rundown.
+ */
 USTRUCT()
 struct FAvaRundownPagePreviewAction : public FAvaRundownPageAction
 {
 	GENERATED_BODY()
 public:
-	/** Specify which preview channel to use. If left empty, the rundown's default preview channel is used. */
+	/** Specifies which preview channel to use. If left empty, the rundown's default preview channel is used. */
 	UPROPERTY()
 	FString PreviewChannelName;
 };
 
 /**
- * Command to execute an action on multiple pages at the same time.
+ * Command to execute a program action on multiple pages at the same time.
  * This is necessary for pages to be part of the same transition.
  */
 USTRUCT()
@@ -729,48 +883,58 @@ struct FAvaRundownPageActions : public FAvaRundownMsgBase
 {
 	GENERATED_BODY()
 public:
+	/** Specifies a list of page Ids that are the target of this action command. */
 	UPROPERTY()
 	TArray<int32> PageIds;
 
+	/** Specifies the page action to execute. */
 	UPROPERTY()
 	EAvaRundownPageActions Action = EAvaRundownPageActions::None;
 };
 
+/**
+ * Command to execute a preview action on multiple pages at the same time.
+ * This is necessary for pages to be part of the same transition.
+ */
 USTRUCT()
 struct FAvaRundownPagePreviewActions : public FAvaRundownPageActions
 {
 	GENERATED_BODY()
 public:
-	/** Specify which preview channel to use. If left empty, the rundown's default preview channel is used. */
+	/** Specifies which preview channel to use. If left empty, the rundown's default preview channel is used. */
 	UPROPERTY()
 	FString PreviewChannelName;
 };
 
 /**
- * Page Sequence Event Message
+ * This message is sent by the server when a page sequence event occurs.
  */
 USTRUCT()
 struct FAvaRundownPageSequenceEvent : public FAvaRundownMsgBase
 {
 	GENERATED_BODY()
 
+	/** Specifies the broadcast channel the event occurred in. */
 	UPROPERTY()
 	FString Channel;
-	
+
+	/** Page Id associated with this event. */
 	UPROPERTY()
 	int32 PageId = FAvaRundownPage::InvalidPageId;
 
-	/** Playable Instance Id. */
+	/** Playable Instance uuid. */
 	UPROPERTY()
 	FGuid InstanceId;
 
 	/** Full asset path: /PackagePath/PackageName.AssetName */
 	UPROPERTY()
 	FString AssetPath;
-	
+
+	/** Specifies the label used to identify the sequence. */
 	UPROPERTY()
 	FString SequenceLabel;
 
+	/** Started, Paused, Finished */
 	UPROPERTY()
 	EAvaPlayableSequenceEventType Event = EAvaPlayableSequenceEventType::None;
 };
@@ -778,44 +942,58 @@ struct FAvaRundownPageSequenceEvent : public FAvaRundownMsgBase
 UENUM()
 enum class EAvaRundownPageTransitionEvents
 {
-	None,
+	None UMETA(Hidden),
 	Started,
 	Finished
 };
 
 /**
- * Page Transition Event Message
+ * This message is sent by the server when a page transition event occurs.
  */
 USTRUCT()
 struct FAvaRundownPageTransitionEvent : public FAvaRundownMsgBase
 {
 	GENERATED_BODY()
 
+	/** Specifies the broadcast channel the event occurred in. */
 	UPROPERTY()
 	FString Channel;
-	
+
+	/** UUID of the transition. */
 	UPROPERTY()
 	FGuid TransitionId;
 
+	/** Pages that are entering the scene during this transition. */
 	UPROPERTY()
 	TArray<int32> EnteringPageIds;
 
+	/** Pages that are already in the scene. May get kicked out or change during this transition. */
 	UPROPERTY()
 	TArray<int32> PlayingPageIds;
 
+	/** Pages that are requested to exit the scene during this transition. Typically part of a "Take Out" transition.  */
 	UPROPERTY()
 	TArray<int32> ExitingPageIds;
 
+	/** Started, Finished */
 	UPROPERTY()
 	EAvaRundownPageTransitionEvents Event = EAvaRundownPageTransitionEvents::None;
 };
 
+/**
+ * Requests a list of all profiles loaded for the current broadcast configuration.
+ * Response is FAvaRundownProfiles.
+ */
 USTRUCT()
 struct FAvaRundownGetProfiles : public FAvaRundownMsgBase
 {
 	GENERATED_BODY()
 };
 
+/**
+ * Response to FAvaRundownGetProfiles.
+ * Contains the list of all profiles in the broadcast configuration.
+ */
 USTRUCT()
 struct FAvaRundownProfiles : public FAvaRundownMsgBase
 {
@@ -839,11 +1017,13 @@ struct FAvaRundownCreateProfile : public FAvaRundownMsgBase
 {
 	GENERATED_BODY()
 public:
+	/** Name to be given to the newly created profile. */
 	UPROPERTY()
 	FString ProfileName;
 
 	/**
-	 * If true the created profile is make "current".
+	 * If true the created profile is also made "current".
+	 * Equivalent to FAvaRundownSetCurrentProfile.
 	 */
 	UPROPERTY()
 	bool bMakeCurrent = true;
@@ -859,33 +1039,41 @@ struct FAvaRundownDuplicateProfile : public FAvaRundownMsgBase
 {
 	GENERATED_BODY()
 public:
+	/** Specifies the existing profile to be duplicated. */
 	UPROPERTY()
 	FString SourceProfileName;
 
+	/** Specifies the name of the new profile to be created. */
 	UPROPERTY()
 	FString NewProfileName;
 
 	/**
-	 * If true the created profile is make "current".
+	 * If true the created profile is also made "current".
+	 * Equivalent to FAvaRundownSetCurrentProfile.
 	 */
 	UPROPERTY()
 	bool bMakeCurrent = true;
 };
 
+/**
+ * Renames an existing profile.
+ */
 USTRUCT()
 struct FAvaRundownRenameProfile : public FAvaRundownMsgBase
 {
 	GENERATED_BODY()
 public:
+	/** Specifies the name of the existing profile to be renamed. */
 	UPROPERTY()
 	FString OldProfileName;
 
+	/** Specifies the new name. */
 	UPROPERTY()
 	FString NewProfileName;
 };
 
 /**
- * Delete the specified profile.
+ * Deletes the specified profile.
  * Fails if profile to be deleted is the current profile.
  */
 USTRUCT()
@@ -893,6 +1081,7 @@ struct FAvaRundownDeleteProfile : public FAvaRundownMsgBase
 {
 	GENERATED_BODY()
 public:
+	/** Specifies the target profile. */
 	UPROPERTY()
 	FString ProfileName;
 };
@@ -907,37 +1096,53 @@ struct FAvaRundownSetCurrentProfile : public FAvaRundownMsgBase
 {
 	GENERATED_BODY()
 public:
+	/** Specifies the requested profile. */
 	UPROPERTY()
 	FString ProfileName;
 };
 
+/**
+ * Output Device information
+ */
 USTRUCT()
 struct FAvaRundownOutputDeviceItem
 {
 	GENERATED_BODY()
 public:
+	/**
+	 * Specifies the device name.
+	 * This is used as "MediaOutputName" in FAvaRundownAddChannelDevice and FAvaRundownEditChannelDevice.
+	 */
 	UPROPERTY()
 	FString Name;
 
+	/** Extra information about the device. */
 	UPROPERTY()
 	FAvaBroadcastMediaOutputInfo OutputInfo;
-	
+
+	/** Specifies the status of the output device. */
 	UPROPERTY()
 	EAvaBroadcastOutputState OutputState = EAvaBroadcastOutputState::Invalid;
-	
+
+	/** In case the device is live, this extra status indicates if the device is operating normally. */
 	UPROPERTY()
 	EAvaBroadcastIssueSeverity IssueSeverity = EAvaBroadcastIssueSeverity::None;
 
+	/** List of errors or warnings. */ 
 	UPROPERTY()
 	TArray<FString> IssueMessages;
 
 	/**
 	 * Raw Json string representing a serialized UMediaOutput.
+	 * This data can be edited, then used in FAvaRundownEditChannelDevice.
 	 */
 	UPROPERTY()
 	FString Data;
 };
 
+/**
+ * Output Device Class Information
+ */
 USTRUCT()
 struct FAvaRundownOutputClassItem
 {
@@ -954,39 +1159,62 @@ public:
 	UPROPERTY()
 	FString Server;
 
+	/**
+	 * Enumeration of the available devices of this class on the given host.
+	 * Note that not all classes can be enumerated.
+	 */
 	UPROPERTY()
 	TArray<FAvaRundownOutputDeviceItem> Devices;
 };
 
+/**
+ * Response to FAvaRundownGetDevices.
+ */
 USTRUCT()
 struct FAvaRundownDevicesList : public FAvaRundownMsgBase
 {
 	GENERATED_BODY()
 public:
+	/**
+	 * List of Output Device Classes
+	 */
 	UPROPERTY()
 	TArray<FAvaRundownOutputClassItem> DeviceClasses;
 };
 
+/**
+ * Requests information (devices, status, etc) on a specified channel.
+ * 
+ * Response is FAvaRundownChannelResponse.
+ */
 USTRUCT()
 struct FAvaRundownGetChannel : public FAvaRundownMsgBase
 {
 	GENERATED_BODY()
 public:
+	/** Specifies the requested channel. */
 	UPROPERTY()
 	FString ChannelName;
 };
 
+/**
+ * Requests information (devices, status, etc) on all channels of the current profile.
+ * 
+ * Response is FAvaRundownChannels.
+ */
 USTRUCT()
 struct FAvaRundownGetChannels : public FAvaRundownMsgBase
 {
 	GENERATED_BODY()
 };
 
+/** Channel Information */
 USTRUCT()
 struct FAvaRundownChannel
 {
 	GENERATED_BODY()
 public:
+	/** Specifies the Channel Name. */
 	UPROPERTY()
 	FString Name;
 
@@ -999,33 +1227,46 @@ public:
 	UPROPERTY()
 	EAvaBroadcastIssueSeverity IssueSeverity = EAvaBroadcastIssueSeverity::None;
 
+	/** List of devices. */
 	UPROPERTY()
 	TArray<FAvaRundownOutputDeviceItem> Devices;
 };
 
+/**
+ * This message is sent by the server if the list of channels is modified
+ * in the current profile. Channel added, removed, pinned or type (preview vs program) changed.
+ */
 USTRUCT()
 struct FAvaRundownChannelListChanged : public FAvaRundownMsgBase
 {
 	GENERATED_BODY()
 
+	/** List of channel information. */
 	UPROPERTY()
 	TArray<FAvaRundownChannel> Channels;
 };
 
+/**
+ * This message is sent by the server in response to FAvaRundownGetChannel or
+ * as an event if a channel's states, render target, devices or settings is changed.
+ */
 USTRUCT()
 struct FAvaRundownChannelResponse : public FAvaRundownMsgBase
 {
 	GENERATED_BODY()
 public:
+	/** Channel Information */
 	UPROPERTY()
 	FAvaRundownChannel Channel;
 };
 
+/** Response to FAvaRundownGetChannels */
 USTRUCT()
 struct FAvaRundownChannels : public FAvaRundownMsgBase
 {
 	GENERATED_BODY()
 public:
+	/** List of channel information. */
 	UPROPERTY()
 	TArray<FAvaRundownChannel> Channels;
 };
@@ -1036,7 +1277,7 @@ public:
 UENUM()
 enum class EAvaRundownAssetEvent : uint8
 {
-	Unknown = 0,
+	Unknown = 0 UMETA(Hidden),
 	Added,
 	Removed,
 	//Saved, // todo
@@ -1067,29 +1308,40 @@ struct FAvaRundownAssetsChanged : public FAvaRundownMsgBase
 	UPROPERTY()
 	bool bIsPlayable = false;
 
+	/** Specifies the event type, i.e. Added, Remove, etc. */
 	UPROPERTY()
 	EAvaRundownAssetEvent EventType = EAvaRundownAssetEvent::Unknown;
 };
 
 /**
- * Channel actions
+ * Channel broadcast actions
  */
 UENUM()
 enum class EAvaRundownChannelActions
 {
-	None,
+	None UMETA(Hidden),
+	/** Start broadcast of the specified channel(s). */
 	Start,
+	/** Stops broadcast of the specified channel(s). */
 	Stop
 };
 
+/**
+ * Requests a broadcast action on the specified channel(s).
+ */
 USTRUCT()
 struct FAvaRundownChannelAction : public FAvaRundownMsgBase
 {
 	GENERATED_BODY()
 public:
+	/**
+	 * Specifies the target channel for the action.
+	 * If left empty, the action will apply to all channels of the current profile.
+	 */
 	UPROPERTY()
 	FString ChannelName;
 
+	/** Specifies the broadcast action to perform on the target channel(s). */
 	UPROPERTY()
 	EAvaRundownChannelActions Action = EAvaRundownChannelActions::None;
 };
@@ -1097,37 +1349,45 @@ public:
 UENUM()
 enum class EAvaRundownChannelEditActions
 {
-	None,
+	None UMETA(Hidden),
+	/** Add new channel with given name. */
 	Add,
+	/** Removes channel with given name. */
 	Remove
 };
 
+/** Requests an edit action on the specified channel. */
 USTRUCT()
 struct FAvaRundownChannelEditAction : public FAvaRundownMsgBase
 {
 	GENERATED_BODY()
 public:
+	/** Specifies the target channel for the action. */
 	UPROPERTY()
 	FString ChannelName;
 
+	/** Specifies the edit action to perform on the target channel. */
 	UPROPERTY()
 	EAvaRundownChannelEditActions Action = EAvaRundownChannelEditActions::None;
 };
 
+/** Requests a channel to be renamed. */
 USTRUCT()
 struct FAvaRundownRenameChannel : public FAvaRundownMsgBase
 {
 	GENERATED_BODY()
 public:
+	/** Existing channel to be renamed. */
 	UPROPERTY()
 	FString OldChannelName;
 
+	/** Specifies the new channel name. */
 	UPROPERTY()
 	FString NewChannelName;
 };
 
 /**
- * Request a list of devices from the rundown server.
+ * Requests a list of devices from the rundown server.
  * The server will reply with FAvaRundownDevicesList containing
  * the devices that can be enumerated from the local host and all connected hosts
  * through the motion design playback service.
@@ -1153,6 +1413,7 @@ struct FAvaRundownAddChannelDevice : public FAvaRundownMsgBase
 {
 	GENERATED_BODY()
 public:
+	/** Specifies the target channel. */
 	UPROPERTY()
 	FString ChannelName;
 
@@ -1163,6 +1424,7 @@ public:
 	UPROPERTY()
 	FString MediaOutputName;
 
+	/** Save broadcast configuration after this operation (true by default). */
 	UPROPERTY()
 	bool bSaveBroadcast = true;
 };
@@ -1176,6 +1438,7 @@ struct FAvaRundownEditChannelDevice : public FAvaRundownMsgBase
 {
 	GENERATED_BODY()
 public:
+	/** Specifies the target channel. */
 	UPROPERTY()
 	FString ChannelName;
 
@@ -1188,9 +1451,14 @@ public:
 	UPROPERTY()
 	FString MediaOutputName;
 
+	/**
+	 * (Modified) Device Data in the same format as FAvaRundownOutputDeviceItem::Data.
+	 * See: FAvaRundownChannel, FAvaRundownDevicesList
+	 */
 	UPROPERTY()
 	FString Data;
 
+	/** Save broadcast configuration after this operation (true by default). */
 	UPROPERTY()
 	bool bSaveBroadcast = true;
 };
@@ -1204,6 +1472,7 @@ struct FAvaRundownRemoveChannelDevice : public FAvaRundownMsgBase
 {
 	GENERATED_BODY()
 public:
+	/** Specifies the target channel. */
 	UPROPERTY()
 	FString ChannelName;
 
@@ -1216,37 +1485,53 @@ public:
 	UPROPERTY()
 	FString MediaOutputName;
 
+	/** Save broadcast configuration after this operation (true by default). */
 	UPROPERTY()
 	bool bSaveBroadcast = true;
 };
 
+/**
+ * Captures an image from the specified channel.
+ * The captured image is 25% of the channel's resolution.
+ * Intended for preview.
+ * Response is FAvaRundownChannelImage.
+ */
 USTRUCT()
 struct FAvaRundownGetChannelImage : public FAvaRundownMsgBase
 {
 	GENERATED_BODY()
 public:
+	/** Specifies the target channel. */
 	UPROPERTY()
 	FString ChannelName;
 };
 
+/**
+ * Response to FAvaRundownGetChannelImage.
+ */
 USTRUCT()
 struct FAvaRundownChannelImage : public FAvaRundownMsgBase
 {
 	GENERATED_BODY()
 public:
+	/**
+	 * Byte array containing the image data.
+	 * Expected format is compressed jpeg.
+	 */ 
 	UPROPERTY()
 	TArray<uint8> ImageData;
 };
 
 /**
  * Queries the given channel's quality settings.
- * Response message is FAvaRundownChannelQualitySettings.
+ * Response is FAvaRundownChannelQualitySettings.
  */
 USTRUCT()
 struct FAvaRundownGetChannelQualitySettings : public FAvaRundownMsgBase
 {
 	GENERATED_BODY()
 public:
+	/** Specifies the target channel. */
 	UPROPERTY()
 	FString ChannelName;
 };
@@ -1257,6 +1542,7 @@ struct FAvaRundownChannelQualitySettings : public FAvaRundownMsgBase
 {
 	GENERATED_BODY()
 public:
+	/** Specifies the target channel. */
 	UPROPERTY()
 	FString ChannelName;
 
@@ -1271,6 +1557,7 @@ struct FAvaRundownSetChannelQualitySettings : public FAvaRundownMsgBase
 {
 	GENERATED_BODY()
 public:
+	/** Specifies the target channel. */
 	UPROPERTY()
 	FString ChannelName;
 
@@ -1279,7 +1566,9 @@ public:
 	TArray<FAvaViewportQualitySettingsFeature> Features;
 };
 
-/** Save current broadcast configuration to file on the server. */
+/**
+ * Save current broadcast configuration to a json file in the Config folder on the server. 
+ */
 USTRUCT()
 struct FAvaRundownSaveBroadcast : public FAvaRundownMsgBase
 {
diff --git a/Engine/Plugins/Experimental/Avalanche/Source/AvalancheMedia/Private/Rundown/AvaRundownServer.cpp b/Engine/Plugins/Experimental/Avalanche/Source/AvalancheMedia/Private/Rundown/AvaRundownServer.cpp
--- a/Engine/Plugins/Experimental/Avalanche/Source/AvalancheMedia/Private/Rundown/AvaRundownServer.cpp
+++ b/Engine/Plugins/Experimental/Avalanche/Source/AvalancheMedia/Private/Rundown/AvaRundownServer.cpp
@@ -169,6 +169,7 @@ namespace UE::AvaRundownServer::Private
 		PageInfo.AssetPath = InPage.GetAssetPath(InRundown);	// Todo: combo templates
 		PageInfo.Statuses = InPage.GetPageStatuses(InRundown);
 		PageInfo.TransitionLayerName = InPage.GetTransitionLayer(InRundown).ToString(); 	// Todo: combo templates
+		PageInfo.bTransitionLogicEnabled = InPage.HasTransitionLogic(InRundown);
 		PageInfo.OutputChannel = InPage.GetChannelName().ToString();
 		PageInfo.bIsEnabled = InPage.IsEnabled();
 		PageInfo.bIsPlaying = InRundown->IsPagePlaying(InPage);
@@ -1215,12 +1216,12 @@ void FAvaRundownServer::HandleChangeTemplateBP(const FAvaRundownChangeTemplateBP
 	{
 		if (Page.UpdateAsset(InMessage.AssetPath))
 		{
-			SendMessage(InContext->GetSender(), InMessage.RequestId, ELogVerbosity::Log, TEXT("Blueprint change of template: %d to %s"), InMessage.TemplateId, *InMessage.AssetPath);
+			SendMessage(InContext->GetSender(), InMessage.RequestId, ELogVerbosity::Log, TEXT("Asset change of template: %d to %s"), InMessage.TemplateId, *InMessage.AssetPath);
 			Rundown->GetOnPagesChanged().Broadcast(Rundown, Page, EAvaRundownPageChanges::Blueprint);
 			return;
 		}
 	}
-	LogAndSendMessage(InContext->GetSender(), InMessage.RequestId, ELogVerbosity::Error, TEXT("Blueprint change of template: %d to %s failed."), InMessage.TemplateId, *InMessage.AssetPath);
+	LogAndSendMessage(InContext->GetSender(), InMessage.RequestId, ELogVerbosity::Error, TEXT("Asset change of template: %d to %s failed."), InMessage.TemplateId, *InMessage.AssetPath);
 }
 
 void FAvaRundownServer::HandleGetPageDetails(const FAvaRundownGetPageDetails& InMessage, const TSharedRef<IMessageContext, ESPMode::ThreadSafe>& InContext)
diff --git a/Engine/Plugins/Experimental/Avalanche/Source/AvalancheMedia/Public/AvaMediaDefines.h b/Engine/Plugins/Experimental/Avalanche/Source/AvalancheMedia/Public/AvaMediaDefines.h
--- a/Engine/Plugins/Experimental/Avalanche/Source/AvalancheMedia/Public/AvaMediaDefines.h
+++ b/Engine/Plugins/Experimental/Avalanche/Source/AvalancheMedia/Public/AvaMediaDefines.h
@@ -6,14 +6,20 @@
 #include "Misc/EnumClassFlags.h"
 #include "AvaMediaDefines.generated.h"
 
+/**
+ * Channel state is a union summary of the output's states.
+ */
 UENUM()
 enum class EAvaBroadcastChannelState : uint8
 {
+	/** Indicates that all channel outputs are offline. */
 	Offline,
+	/** Indicates that at least some of the channel outputs are idle (but none are live). */
 	Idle,
+	/** Indicates that at least some of the channel outputs are live.*/
 	Live,
 
-	Max,
+	Max UMETA(Hidden),
 };
 
 /**
@@ -58,6 +64,10 @@ enum class EAvaBroadcastOutputState : uint8
 	Error
 };
 
+/**
+ * In case the broadcast device is Live (see EAvaBroadcastOutputState),
+ * this extra status indicates if the device is operating normally.
+ */
 UENUM()
 enum class EAvaBroadcastIssueSeverity : uint8
 {
@@ -65,7 +75,7 @@ enum class EAvaBroadcastIssueSeverity : uint8
 	Warnings,
 	Errors,
 
-	Max,
+	Max UMETA(Hidden),
 };
 
 //An enum indicating what changed in Broadcast
@@ -125,7 +135,7 @@ enum class EAvaPlayableStatus : uint8
 UENUM()
 enum class EAvaPlayableSequenceEventType : uint8
 {
-	None,
+	None UMETA(Hidden),
 	Started,
 	Paused,
 	Finished
@@ -304,6 +314,9 @@ enum class EAvaPlaybackStatus
 	Error
 };
 
+/**
+ * Rundown's page list type.
+ */
 UENUM(BlueprintType, DisplayName = "Motion Design Rundown Page List Type")
 enum class EAvaRundownPageListType : uint8
 {
