MotionDesign : DynamicMeshModifier
- Use cloner mesh builder to convert various mesh types to dynamic mesh
- Detects changes on geometry or transform and retrigger modifier

diff --git a/Engine/Plugins/Experimental/Avalanche/Source/AvalancheModifiers/AvalancheModifiers.Build.cs b/Engine/Plugins/Experimental/Avalanche/Source/AvalancheModifiers/AvalancheModifiers.Build.cs
--- a/Engine/Plugins/Experimental/Avalanche/Source/AvalancheModifiers/AvalancheModifiers.Build.cs
+++ b/Engine/Plugins/Experimental/Avalanche/Source/AvalancheModifiers/AvalancheModifiers.Build.cs
@@ -27,6 +27,7 @@ public class AvalancheModifiers : ModuleRules
 				"Avalanche",
 				"AvalancheSceneTree",
 				"AvalancheShapes",
+				"ClonerEffectorMeshBuilder",
 				"DynamicMaterial",
 				"GeometryAlgorithms",
 				"GeometryFramework",
diff --git a/Engine/Plugins/Experimental/Avalanche/Source/AvalancheModifiers/Private/Modifiers/AvaDynamicMeshConverterModifier.cpp b/Engine/Plugins/Experimental/Avalanche/Source/AvalancheModifiers/Private/Modifiers/AvaDynamicMeshConverterModifier.cpp
--- a/Engine/Plugins/Experimental/Avalanche/Source/AvalancheModifiers/Private/Modifiers/AvaDynamicMeshConverterModifier.cpp
+++ b/Engine/Plugins/Experimental/Avalanche/Source/AvalancheModifiers/Private/Modifiers/AvaDynamicMeshConverterModifier.cpp
@@ -7,14 +7,12 @@
 #include "Components/DynamicMeshComponent.h"
 #include "Components/SkeletalMeshComponent.h"
 #include "Components/StaticMeshComponent.h"
-#include "DynamicMesh/MeshTransforms.h"
-#include "DynamicMeshEditor.h"
+#include "DynamicMesh/DynamicMesh3.h"
 #include "Engine/StaticMesh.h"
+#include "Extensions/AvaSceneTreeUpdateModifierExtension.h"
 #include "GeometryScript/MeshAssetFunctions.h"
 #include "GeometryScript/MeshBasicEditFunctions.h"
-#include "GeometryScript/SceneUtilityFunctions.h"
 #include "ProceduralMeshComponent.h"
-#include "Engine/CollisionProfile.h"
 
 #if WITH_EDITOR
 #include "AssetRegistry/AssetRegistryModule.h"
@@ -29,10 +27,10 @@ FAvaDynamicMeshConverterModifierComponentState::FAvaDynamicMeshConverterModifier
 {
 	if (Component.IsValid())
 	{
-		if (const AActor* ComponentOwner = Component->GetOwner())
+		if (const AActor* ComponentOwner = InPrimitiveComponent->GetOwner())
 		{
-#if WITH_EDITOR
 			bActorHiddenInGame = ComponentOwner->IsHidden();
+#if WITH_EDITOR
 			bActorHiddenInEditor = ComponentOwner->IsTemporarilyHiddenInEditor();
 #endif
 			if (const USceneComponent* RootComponent = ComponentOwner->GetRootComponent())
@@ -40,6 +38,11 @@ FAvaDynamicMeshConverterModifierComponentState::FAvaDynamicMeshConverterModifier
 				bComponentVisible = RootComponent->IsVisible();
 				bComponentHiddenInGame = RootComponent->bHiddenInGame;
 			}
+
+			if (const AActor* ParentActor = ComponentOwner->GetAttachParentActor())
+			{
+				ActorRelativeTransform = ComponentOwner->GetActorTransform().GetRelativeTransform(ParentActor->GetActorTransform());
+			}
 		}
 	}
 }
@@ -50,8 +53,9 @@ void UAvaDynamicMeshConverterModifier::OnModifierCDOSetup(FActorModifierCoreMeta
 
 	InMetadata.SetName(TEXT("DynamicMeshConverter"));
 	InMetadata.SetCategory(TEXT("Conversion"));
+	InMetadata.AllowTick(true);
 #if WITH_EDITOR
-	InMetadata.SetDescription(LOCTEXT("ModifierDescription", "Converts a non dynamic mesh actor into a dynamic mesh"));
+	InMetadata.SetDescription(LOCTEXT("ModifierDescription", "Converts various actor mesh types into a single dynamic mesh, this is an heavy operation"));
 #endif
 	InMetadata.SetCompatibilityRule([](const AActor* InActor)->bool
 	{
@@ -64,25 +68,27 @@ void UAvaDynamicMeshConverterModifier::OnModifierAdded(EActorModifierCoreEnableR
 	Super::OnModifierAdded(InReason);
 
 	AddDynamicMeshComponent();
-	bConvertMesh = true;
+
+	AddExtension<FAvaRenderStateUpdateModifierExtension>(this);
+
+	if (FAvaSceneTreeUpdateModifierExtension* SceneExtension = AddExtension<FAvaSceneTreeUpdateModifierExtension>(this))
+	{
+		TrackedActor.ReferenceContainer = EAvaReferenceContainer::Other;
+		TrackedActor.ReferenceActorWeak = SourceActorWeak.Get();
+		TrackedActor.bSkipHiddenActors = false;
+		SceneExtension->TrackSceneTree(0, &TrackedActor);
+	}
 }
 
 void UAvaDynamicMeshConverterModifier::OnModifierEnabled(EActorModifierCoreEnableReason InReason)
 {
 	Super::OnModifierEnabled(InReason);
-
-	// On load update state by converting source again
-	if (InReason == EActorModifierCoreEnableReason::Load)
-	{
-		bConvertMesh = true;
-		MarkModifierDirty();
-	}
 }
 
 void UAvaDynamicMeshConverterModifier::RestorePreState()
 {
 	UAvaGeometryBaseModifier::RestorePreState();
-	
+
 	for (FAvaDynamicMeshConverterModifierComponentState& ConvertedComponent : ConvertedComponents)
 	{
 		if (ConvertedComponent.Component.IsValid())
@@ -91,9 +97,9 @@ void UAvaDynamicMeshConverterModifier::RestorePreState()
 
 			if (ComponentActor != GetModifiedActor())
 			{
-#if WITH_EDITOR
 				// Hide actor but do not hide ourselves
-				ComponentActor->SetActorHiddenInGame(ConvertedComponent.bActorHiddenInGame);
+				ComponentActor->SetHidden(ConvertedComponent.bActorHiddenInGame);
+#if WITH_EDITOR
 				ComponentActor->SetIsTemporarilyHiddenInEditor(ConvertedComponent.bActorHiddenInEditor);
 #endif
 			}
@@ -117,8 +123,59 @@ void UAvaDynamicMeshConverterModifier::OnModifierRemoved(EActorModifierCoreDisab
 	}
 }
 
+bool UAvaDynamicMeshConverterModifier::IsModifierDirtyable() const
+{
+	const double CurrentTime = FPlatformTime::Seconds();
+
+	if (TransformUpdateInterval > 0
+		&& CurrentTime - LastTransformUpdateTime > TransformUpdateInterval)
+	{
+		const_cast<UAvaDynamicMeshConverterModifier*>(this)->LastTransformUpdateTime = CurrentTime;
+
+		for (const FAvaDynamicMeshConverterModifierComponentState& ConvertedComponent : ConvertedComponents)
+		{
+			const UPrimitiveComponent* PrimitiveComponent = ConvertedComponent.Component.Get();
+			if (!PrimitiveComponent)
+			{
+				continue;
+			}
+
+			const AActor* ChildActor = PrimitiveComponent->GetOwner();
+			if (!ChildActor)
+			{
+				continue;
+			}
+
+			const AActor* ParentActor = ChildActor->GetAttachParentActor();
+			FTransform ExpectedTransform = FTransform::Identity;
+
+			if (ParentActor)
+			{
+				ExpectedTransform = ChildActor->GetActorTransform().GetRelativeTransform(ParentActor->GetActorTransform());
+			}
+
+			if (!ConvertedComponent.ActorRelativeTransform.Equals(ExpectedTransform, 0.01))
+			{
+				return true;
+			}
+		}
+	}
+
+	return Super::IsModifierDirtyable();
+}
+
+void UAvaDynamicMeshConverterModifier::OnSceneTreeTrackedActorChildrenChanged(int32 InIdx, const TSet<TWeakObjectPtr<AActor>>& InPreviousChildrenActors, const TSet<TWeakObjectPtr<AActor>>& InNewChildrenActors)
+{
+	if (bIncludeAttachedActors)
+	{
+		MarkModifierDirty();
+	}
+}
+
 void UAvaDynamicMeshConverterModifier::Apply()
 {
+	const AActor* ActorModified = GetModifiedActor();
+
 	if (!IsMeshValid())
 	{
 		Fail(LOCTEXT("InvalidDynamicMeshComponent", "Invalid dynamic mesh component on modified actor"));
@@ -132,88 +189,43 @@ void UAvaDynamicMeshConverterModifier::Apply()
 		Fail(LOCTEXT("InvalidDynamicMeshComponent", "Invalid dynamic mesh component on modified actor"));
 		return;
 	}
-	
-	using namespace UE::Geometry;
-	
-	if (bConvertMesh)
+
+	TArray<TWeakObjectPtr<UMaterialInterface>> MaterialsWeak;
+	if (!ConvertComponents(MaterialsWeak))
 	{
-		DynMeshComponent->SetNumMaterials(0);
-		DynMeshComponent->EditMesh([this, DynMeshComponent](FDynamicMesh3& AppendToMesh)
+		Fail(LOCTEXT("ConversionFailed", "Conversion to dynamic mesh failed"));
+		return;
+	}
+
+	for (int32 MatIndex = 0; MatIndex < MaterialsWeak.Num(); MatIndex++)
+	{
+		UMaterialInterface* Material = MaterialsWeak[MatIndex].Get();
+		DynMeshComponent->SetMaterial(MatIndex, Material);
+	}
+
+	for (const FAvaDynamicMeshConverterModifierComponentState& OutConvert : ConvertedComponents)
+	{
+		if (OutConvert.Component.IsValid())
 		{
-			AppendToMesh.Clear();
-			FMeshIndexMappings TmpMappings;
-			FDynamicMeshEditor Editor(&AppendToMesh);
-			FGeometryScriptAppendMeshOptions AppendOptions;
-			AppendOptions.CombineMode = EGeometryScriptCombineAttributesMode::EnableAllMatching;
-			int32 MaterialCount = 0;
-			// Convert meshes
-			ConvertedComponents.Empty();
-			ConvertComponents(ConvertedComponents);
-			for (const FAvaDynamicMeshConverterModifierComponentState& OutConvert : ConvertedComponents)
+			AActor* ComponentActor = OutConvert.Component->GetOwner();
+
+			// Hide converted component
+			if (bHideConvertedMesh)
 			{
-				// Enable matching attributes & append mesh
-				AppendOptions.UpdateAttributesForCombineMode(AppendToMesh, OutConvert.Mesh);
-				Editor.AppendMesh(&OutConvert.Mesh, TmpMappings);
-				
-				if (OutConvert.Mesh.HasAttributes() && OutConvert.Mesh.Attributes()->HasMaterialID())
+				if (ComponentActor != ActorModified)
 				{
-					// Fix triangles materials linking
-					const FDynamicMeshMaterialAttribute* FromMaterialIDAttrib = OutConvert.Mesh.Attributes()->GetMaterialID();
-					FDynamicMeshMaterialAttribute* ToMaterialIDAttrib = AppendToMesh.Attributes()->GetMaterialID();
-					TMap<int32, int32> MaterialMap;
-					for (const TPair<int32, int32>& FromToTId : TmpMappings.GetTriangleMap().GetForwardMap())
-					{
-						const int32 FromMatId = FromMaterialIDAttrib->GetValue(FromToTId.Key);
-						const int32 ToMatId = FromMatId + MaterialCount;
-						MaterialMap.Add(FromMatId, ToMatId);
-						ToMaterialIDAttrib->SetNewValue(FromToTId.Value, ToMatId);
-					}
-					MaterialCount += MaterialMap.Num();
-					
-					// Reapply original materials
-					if (OutConvert.Component.IsValid())
-					{
-						for (const TPair<int32, int32>& MatPair : MaterialMap)
-						{
-							DynMeshComponent->SetMaterial(MatPair.Value, OutConvert.Component->GetMaterial(MatPair.Key));
-						}
-					}
-				}
-				
-				if (OutConvert.Component.IsValid())
-				{
-					AActor* ComponentActor = OutConvert.Component->GetOwner();
-					
-					// Hide converted component
-					if (bHideConvertedMesh)
-					{
-						if (ComponentActor != GetModifiedActor())
-						{
-							ComponentActor->SetActorHiddenInGame(true);
+					ComponentActor->SetHidden(true);
 #if WITH_EDITOR
-							ComponentActor->SetIsTemporarilyHiddenInEditor(true);
+					ComponentActor->SetIsTemporarilyHiddenInEditor(true);
 #endif
-						}
-						else if (USceneComponent* RootComponent = ComponentActor->GetRootComponent())
-						{
-							RootComponent->SetHiddenInGame(true);
-							RootComponent->SetVisibility(false);
-						}
-					}
 				}
-				
-				TmpMappings.Reset();
+				else if (USceneComponent* RootComponent = ComponentActor->GetRootComponent())
+				{
+					RootComponent->SetHiddenInGame(true);
+					RootComponent->SetVisibility(false);
+				}
 			}
-			
-			ConvertedMesh = AppendToMesh;
-		});
-	}
-	else if (ConvertedMesh.IsSet())
-	{
-		DynMeshComponent->EditMesh([this](FDynamicMesh3& AppendToMesh)
-		{
-			AppendToMesh = ConvertedMesh.GetValue();
-		});
+		}
 	}
 
 	Next();
@@ -225,7 +237,7 @@ void UAvaDynamicMeshConverterModifier::PostEditChangeProperty(FPropertyChangedEv
 	Super::PostEditChangeProperty(PropertyChangedEvent);
 
 	const FName MemberName = PropertyChangedEvent.GetMemberPropertyName();
-	
+
 	static const FName SourceActorName = GET_MEMBER_NAME_CHECKED(UAvaDynamicMeshConverterModifier, SourceActorWeak);
 
 	if (MemberName == SourceActorName)
@@ -245,7 +257,7 @@ void UAvaDynamicMeshConverterModifier::ConvertToStaticMeshAsset()
 	{
 		return;
 	}
-	
+
 	// generate name for asset
 	const FString NewNameSuggestion = TEXT("SM_MotionDesign_") + OwningActor->GetActorNameOrLabel();
 	FString PackageName = FString(TEXT("/Game/Meshes/")) + NewNameSuggestion;
@@ -283,7 +295,7 @@ void UAvaDynamicMeshConverterModifier::ConvertToStaticMeshAsset()
 	{
 		return;
 	}
-	
+
 	// find/create package
 	UPackage* Package = CreatePackage(*UserPackageName);
 	check(Package);
@@ -291,14 +303,14 @@ void UAvaDynamicMeshConverterModifier::ConvertToStaticMeshAsset()
 	// Create StaticMesh object
 	UStaticMesh* DestinationMesh = NewObject<UStaticMesh>(Package, MeshName, RF_Public | RF_Standalone);
 	UDynamicMesh* SourceMesh = DynMeshComponent->GetDynamicMesh();
-	
+
 	// export options
 	FGeometryScriptCopyMeshToAssetOptions AssetOptions;
 	AssetOptions.bReplaceMaterials = false;
 	AssetOptions.bEnableRecomputeNormals = false;
 	AssetOptions.bEnableRecomputeTangents = false;
 	AssetOptions.bEnableRemoveDegenerates = true;
-	
+
 	// LOD options
 	FGeometryScriptMeshWriteLOD TargetLOD;
 	TargetLOD.LODIndex = 0;
@@ -307,7 +319,7 @@ void UAvaDynamicMeshConverterModifier::ConvertToStaticMeshAsset()
 
 	UGeometryScriptLibrary_StaticMeshFunctions::CopyMeshToStaticMesh(SourceMesh, DestinationMesh, AssetOptions, TargetLOD, OutResult);
 	DestinationMesh->GetBodySetup()->AggGeom = DynMeshComponent->GetBodySetup()->AggGeom;
-	
+
 	if (OutResult == EGeometryScriptOutcomePins::Success)
 	{
 		// Notify asset registry of new asset
@@ -391,242 +403,153 @@ void UAvaDynamicMeshConverterModifier::SetHideConvertedMesh(bool bInHide)
 	bHideConvertedMesh = bInHide;
 }
 
-void UAvaDynamicMeshConverterModifier::OnSourceActorChanged()
+void UAvaDynamicMeshConverterModifier::OnRenderStateUpdated(AActor* InActor, UActorComponent* InComponent)
 {
+	if (!IsValid(InActor) || !IsValid(InComponent))
+	{
+		return;
+	}
+
+	const UPrimitiveComponent* PrimitiveComponent = Cast<UPrimitiveComponent>(InComponent);
+	if (!PrimitiveComponent)
+	{
+		return;
+	}
+
+	const UDynamicMeshComponent* DynamicMeshComponent = GetMeshComponent();
+	if (PrimitiveComponent == DynamicMeshComponent)
+	{
+		return;
+	}
+
 	const AActor* SourceActor = SourceActorWeak.Get();
+	if (!SourceActor)
+	{
+		return;
+	}
+
+	const bool bIsSourceActor = InActor == SourceActor;
+	const bool bIsAttachedToSourceActor = bIncludeAttachedActors && InActor->IsAttachedTo(SourceActor);
+	if (!bIsSourceActor && !bIsAttachedToSourceActor)
+	{
+		return;
+	}
+
+	MarkModifierDirty();
+}
+
+void UAvaDynamicMeshConverterModifier::OnSourceActorChanged()
+{
+	AActor* SourceActor = SourceActorWeak.Get();
 	const AActor* ActorModified = GetModifiedActor();
+
 	if (!SourceActor || !ActorModified)
 	{
 		return;
 	}
+
 	bHideConvertedMesh = SourceActor == ActorModified || SourceActor->IsAttachedTo(ActorModified);
+
+	if (const FAvaSceneTreeUpdateModifierExtension* SceneExtension = GetExtension<FAvaSceneTreeUpdateModifierExtension>())
+	{
+		TrackedActor.ReferenceActorWeak = SourceActor;
+		SceneExtension->CheckTrackedActorUpdate(0);
+	}
 }
 
-void UAvaDynamicMeshConverterModifier::ConvertComponents(TArray<FAvaDynamicMeshConverterModifierComponentState>& OutResults) const
+bool UAvaDynamicMeshConverterModifier::ConvertComponents(TArray<TWeakObjectPtr<UMaterialInterface>>& OutMaterialsWeak)
 {
 	if (!IsMeshValid() || !SourceActorWeak.IsValid())
 	{
-		return;
+		return false;
 	}
-	const FTransform SourceTransform = SourceActorWeak->GetActorTransform();
-	UDynamicMeshComponent* DynMeshComponent = GetMeshComponent();
-	UDynamicMesh* OutputDynamicMesh = DynMeshComponent->GetDynamicMesh();
-	static const FGeometryScriptCopyMeshFromAssetOptions FromMeshOptions;
-	static FGeometryScriptMeshReadLOD FromMeshLOD;
-	FromMeshLOD.LODType = EGeometryScriptLODType::SourceModel;
+
+	ConvertedComponents.Empty();
+	MeshBuilder.Reset();
+
+	UDynamicMeshComponent* DynamicMeshComponent = GetMeshComponent();
+	const FTransform SourceTransform = DynamicMeshComponent->GetComponentTransform();
+
+	// Get relevant actors
 	TArray<AActor*> FilteredActors;
 	GetFilteredActors(FilteredActors);
+
 	if (HasFlag(EAvaDynamicMeshConverterModifierType::StaticMeshComponent))
 	{
 		TArray<UStaticMeshComponent*> Components;
 		GetStaticMeshComponents(FilteredActors, Components);
+
 		for (UStaticMeshComponent* Component : Components)
 		{
-			UStaticMesh* StaticMesh = Component->GetStaticMesh();
-
-			// convert to dynamic mesh
-			EGeometryScriptOutcomePins OutResult;
-
-			UGeometryScriptLibrary_StaticMeshFunctions::CopyMeshFromStaticMesh(StaticMesh, OutputDynamicMesh, FromMeshOptions, FromMeshLOD, OutResult);
-			if (OutResult == EGeometryScriptOutcomePins::Success)
+			if (MeshBuilder.AppendComponent(Component, SourceTransform))
 			{
-				// Transform the new mesh relative to the component
-				const FTransform RelativeTransform = Component->GetComponentTransform().GetRelativeTransform(SourceTransform);
-				OutputDynamicMesh->EditMesh([&OutResults, Component, RelativeTransform](FDynamicMesh3& EditMesh)
-				{
-					MeshTransforms::ApplyTransform(EditMesh, RelativeTransform);
-					FAvaDynamicMeshConverterModifierComponentState NewResult(Component);
-					NewResult.Mesh = MoveTemp(EditMesh);
-					OutResults.Add(NewResult);
-					// replace by empty mesh
-					FDynamicMesh3 EmptyMesh;
-					EditMesh = MoveTemp(EmptyMesh);
-				});
+				FAvaDynamicMeshConverterModifierComponentState State(Component);
+				ConvertedComponents.Emplace(State);
 			}
 		}
 	}
+
 	if (HasFlag(EAvaDynamicMeshConverterModifierType::DynamicMeshComponent))
 	{
 		TArray<UDynamicMeshComponent*> Components;
 		GetDynamicMeshComponents(FilteredActors, Components);
+
 		for (UDynamicMeshComponent* Component : Components)
 		{
-			const UDynamicMesh* DynamicMesh = Component->GetDynamicMesh();
-			// Transform the new mesh relative to the component
-			const FTransform RelativeTransform = Component->GetComponentTransform().GetRelativeTransform(SourceTransform);
-			// Create a copy
-			DynamicMesh->ProcessMesh([&OutResults, Component, RelativeTransform](const FDynamicMesh3& EditMesh)
+			if (MeshBuilder.AppendComponent(Component, SourceTransform))
 			{
-				FDynamicMesh3 CopyMesh = EditMesh;
-				MeshTransforms::ApplyTransform(CopyMesh, RelativeTransform);
-				FAvaDynamicMeshConverterModifierComponentState NewResult(Component);
-				NewResult.Mesh = MoveTemp(CopyMesh);
-				OutResults.Add(NewResult);
-			});
+				FAvaDynamicMeshConverterModifierComponentState State(Component);
+				ConvertedComponents.Emplace(State);
+			}
 		}
 	}
+
 	if (HasFlag(EAvaDynamicMeshConverterModifierType::SkeletalMeshComponent))
 	{
 		TArray<USkeletalMeshComponent*> Components;
 		GetSkeletalMeshComponents(FilteredActors, Components);
+
 		for (USkeletalMeshComponent* Component : Components)
 		{
-			USkeletalMesh* SkeletalMesh = Component->GetSkeletalMeshAsset();
-
-			// convert to dynamic mesh
-			EGeometryScriptOutcomePins OutResult;
-
-			UGeometryScriptLibrary_StaticMeshFunctions::CopyMeshFromSkeletalMesh(SkeletalMesh, OutputDynamicMesh, FromMeshOptions, FromMeshLOD, OutResult);
-			if (OutResult == EGeometryScriptOutcomePins::Success)
+			if (MeshBuilder.AppendComponent(Component, SourceTransform))
 			{
-				// Transform the new mesh relative to the component
-				const FTransform RelativeTransform = Component->GetComponentTransform().GetRelativeTransform(SourceTransform);
-				OutputDynamicMesh->EditMesh([&OutResults, Component, RelativeTransform](FDynamicMesh3& EditMesh)
-				{
-					MeshTransforms::ApplyTransform(EditMesh, RelativeTransform);
-					FAvaDynamicMeshConverterModifierComponentState NewResult(Component);
-					NewResult.Mesh = MoveTemp(EditMesh);
-					OutResults.Add(NewResult);
-					// replace by empty mesh
-					FDynamicMesh3 EmptyMesh;
-					EditMesh = MoveTemp(EmptyMesh);
-				});
+				FAvaDynamicMeshConverterModifierComponentState State(Component);
+				ConvertedComponents.Emplace(State);
 			}
 		}
 	}
+
 	if (HasFlag(EAvaDynamicMeshConverterModifierType::BrushComponent))
 	{
-		static FGeometryScriptCopyMeshFromComponentOptions Options;
-		Options.RequestedLOD = FromMeshLOD;
-		FTransform Transform;
 		TArray<UBrushComponent*> Components;
 		GetBrushComponents(FilteredActors, Components);
+
 		for (UBrushComponent* Component : Components)
 		{
-			// convert to dynamic mesh
-			EGeometryScriptOutcomePins OutResult;
-
-			UGeometryScriptLibrary_SceneUtilityFunctions::CopyMeshFromComponent(Component, OutputDynamicMesh, Options, false, Transform, OutResult);
-			if (OutResult == EGeometryScriptOutcomePins::Success)
+			if (MeshBuilder.AppendComponent(Component, SourceTransform))
 			{
-				// Transform the new mesh relative to the component
-				const FTransform RelativeTransform = Component->GetComponentTransform().GetRelativeTransform(SourceTransform);
-				OutputDynamicMesh->EditMesh([&OutResults, Component, RelativeTransform](FDynamicMesh3& EditMesh)
-				{
-					MeshTransforms::ApplyTransform(EditMesh, RelativeTransform);
-					FAvaDynamicMeshConverterModifierComponentState NewResult(Component);
-					NewResult.Mesh = MoveTemp(EditMesh);
-					OutResults.Add(NewResult);
-					// replace by empty mesh
-					FDynamicMesh3 EmptyMesh;
-					EditMesh = MoveTemp(EmptyMesh);
-				});
+				FAvaDynamicMeshConverterModifierComponentState State(Component);
+				ConvertedComponents.Emplace(State);
 			}
 		}
 	}
+
 	if (HasFlag(EAvaDynamicMeshConverterModifierType::ProceduralMeshComponent))
 	{
-		using namespace UE::Geometry;
 		TArray<UProceduralMeshComponent*> Components;
 		GetProceduralMeshComponents(FilteredActors, Components);
+
 		for (UProceduralMeshComponent* Component : Components)
 		{
-			const int32 SectionCount = Component->GetNumSections();
-			if (SectionCount == 0)
+			if (MeshBuilder.AppendComponent(Component, SourceTransform))
 			{
-				continue;
-			}
-
-			// Transform the new mesh relative to the component
-			const FTransform RelativeTransform = Component->GetComponentTransform().GetRelativeTransform(SourceTransform);
-			
-			FAvaDynamicMeshConverterModifierComponentState NewResult(Component);
-			NewResult.Mesh.EnableAttributes();
-			NewResult.Mesh.Attributes()->EnablePrimaryColors();
-			NewResult.Mesh.Attributes()->EnableMaterialID();
-			NewResult.Mesh.Attributes()->SetNumNormalLayers(1);
-			NewResult.Mesh.Attributes()->SetNumUVLayers(1);
-			NewResult.Mesh.Attributes()->SetNumPolygroupLayers(1);
-			NewResult.Mesh.Attributes()->EnableTangents();
-			
-			FDynamicMeshColorOverlay* ColorOverlay = NewResult.Mesh.Attributes()->PrimaryColors();
-			FDynamicMeshNormalOverlay* NormalOverlay = NewResult.Mesh.Attributes()->PrimaryNormals();
-			FDynamicMeshUVOverlay* UVOverlay = NewResult.Mesh.Attributes()->PrimaryUV();
-			FDynamicMeshMaterialAttribute* MaterialAttr = NewResult.Mesh.Attributes()->GetMaterialID();
-			FDynamicMeshPolygroupAttribute* GroupAttr = NewResult.Mesh.Attributes()->GetPolygroupLayer(0);
-			FDynamicMeshNormalOverlay* TangentOverlay = NewResult.Mesh.Attributes()->PrimaryTangents();
-			
-			for (int32 SectionIdx = 0; SectionIdx < SectionCount; SectionIdx++)
-			{
-				if (FProcMeshSection* Section = Component->GetProcMeshSection(SectionIdx))
-				{
-					if (Section->bSectionVisible)
-					{
-						TArray<int32> VtxIds;
-						TArray<int32> NormalIds;
-						TArray<int32> ColorIds;
-						TArray<int32> UVIds;
-						TArray<int32> TaIds;
-						
-						// copy vertices data (position, normal, color, UV, tangent)
-						for (FProcMeshVertex& SectionVertex : Section->ProcVertexBuffer)
-						{
-							int32 VId = NewResult.Mesh.AppendVertex(SectionVertex.Position);
-							VtxIds.Add(VId);
-							
-							int32 NId = NormalOverlay->AppendElement(static_cast<FVector3f>(SectionVertex.Normal));
-							NormalIds.Add(NId);
-							
-							int32 CId = ColorOverlay->AppendElement(static_cast<FVector4f>(SectionVertex.Color));
-							ColorIds.Add(CId);
-							
-							int32 UVId = UVOverlay->AppendElement(static_cast<FVector2f>(SectionVertex.UV0));
-							UVIds.Add(UVId);
-
-							int32 TaId = TangentOverlay->AppendElement(static_cast<FVector3f>(SectionVertex.Tangent.TangentX));
-							TaIds.Add(TaId);
-						}
-						
-						// copy tris data
-						if (Section->ProcIndexBuffer.Num() % 3 != 0)
-						{
-							continue;
-						}
-						for (int32 Idx = 0; Idx < Section->ProcIndexBuffer.Num(); Idx+=3)
-						{
-							int32 VIdx1 = Section->ProcIndexBuffer[Idx];
-							int32 VIdx2 = Section->ProcIndexBuffer[Idx + 1];
-							int32 VIdx3 = Section->ProcIndexBuffer[Idx + 2];
-							
-							int32 VId1 = VtxIds[VIdx1];
-							int32 VId2 = VtxIds[VIdx2];
-							int32 VId3 = VtxIds[VIdx3];
-
-							int32 TId = NewResult.Mesh.AppendTriangle(VId1, VId2, VId3, SectionIdx);
-							if (TId < 0)
-							{
-								continue;
-							}
-
-							NormalOverlay->SetTriangle(TId, FIndex3i(NormalIds[VIdx1], NormalIds[VIdx2], NormalIds[VIdx3]), true);
-							ColorOverlay->SetTriangle(TId, FIndex3i(ColorIds[VIdx1], ColorIds[VIdx2], ColorIds[VIdx3]), true);
-							UVOverlay->SetTriangle(TId, FIndex3i(UVIds[VIdx1], UVIds[VIdx2], UVIds[VIdx3]), true);
-							TangentOverlay->SetTriangle(TId, FIndex3i(TaIds[VIdx1], TaIds[VIdx2], TaIds[VIdx3]), true);
-							
-							MaterialAttr->SetValue(TId, SectionIdx);
-							GroupAttr->SetValue(TId, SectionIdx);
-						}
-					}
-				}
-			}
-
-			if (NewResult.Mesh.TriangleCount() > 0)
-			{
-				MeshTransforms::ApplyTransform(NewResult.Mesh, RelativeTransform);
-				OutResults.Add(NewResult);
+				FAvaDynamicMeshConverterModifierComponentState State(Component);
+				ConvertedComponents.Emplace(State);
 			}
 		}
 	}
+
+	return MeshBuilder.BuildDynamicMesh(DynamicMeshComponent->GetDynamicMesh(), OutMaterialsWeak);
 }
 
 bool UAvaDynamicMeshConverterModifier::HasFlag(EAvaDynamicMeshConverterModifierType InFlag) const
@@ -637,7 +560,7 @@ bool UAvaDynamicMeshConverterModifier::HasFlag(EAvaDynamicMeshConverterModifierT
 void UAvaDynamicMeshConverterModifier::AddDynamicMeshComponent()
 {
 	UDynamicMeshComponent* DynMeshComponent = GetMeshComponent();
-	
+
 	if (DynMeshComponent)
 	{
 		return;
@@ -649,7 +572,7 @@ void UAvaDynamicMeshConverterModifier::AddDynamicMeshComponent()
 	{
 		return;
 	}
-	
+
 #if WITH_EDITOR
 	ActorModified->Modify();
 	Modify();
@@ -685,14 +608,14 @@ void UAvaDynamicMeshConverterModifier::AddDynamicMeshComponent()
 	// Rerun construction scripts
 	ActorModified->RerunConstructionScripts();
 #endif
-	
+
 	bComponentCreated = true;
 }
 
 void UAvaDynamicMeshConverterModifier::RemoveDynamicMeshComponent()
 {
 	UDynamicMeshComponent* DynMeshComponent = GetMeshComponent();
-	
+
 	if (!DynMeshComponent)
 	{
 		return;
@@ -703,14 +626,14 @@ void UAvaDynamicMeshConverterModifier::RemoveDynamicMeshComponent()
 	{
 		return;
 	}
-	
+
 	AActor* ActorModified = GetModifiedActor();
 
 	if (!IsValid(ActorModified))
 	{
 		return;
 	}
-	
+
 #if WITH_EDITOR
 	ActorModified->Modify();
 	Modify();
@@ -718,10 +641,10 @@ void UAvaDynamicMeshConverterModifier::RemoveDynamicMeshComponent()
 
 	const FDetachmentTransformRules DetachRules(EDetachmentRule::KeepWorld, false);
 	DynMeshComponent->DetachFromComponent(DetachRules);
-	
+
 	ActorModified->RemoveInstanceComponent(DynMeshComponent);
 	DynMeshComponent->DestroyComponent(false);
-	
+
 	bComponentCreated = false;
 }
 
@@ -859,10 +782,4 @@ void UAvaDynamicMeshConverterModifier::GetProceduralMeshComponents(const TArray<
 	});
 }
 
-void UAvaDynamicMeshConverterModifier::ConvertMesh()
-{
-	bConvertMesh = true;
-	MarkModifierDirty();
-}
-
 #undef LOCTEXT_NAMESPACE
diff --git a/Engine/Plugins/Experimental/Avalanche/Source/AvalancheModifiers/Public/Modifiers/AvaDynamicMeshConverterModifier.h b/Engine/Plugins/Experimental/Avalanche/Source/AvalancheModifiers/Public/Modifiers/AvaDynamicMeshConverterModifier.h
--- a/Engine/Plugins/Experimental/Avalanche/Source/AvalancheModifiers/Public/Modifiers/AvaDynamicMeshConverterModifier.h
+++ b/Engine/Plugins/Experimental/Avalanche/Source/AvalancheModifiers/Public/Modifiers/AvaDynamicMeshConverterModifier.h
@@ -3,7 +3,9 @@
 #pragma once
 
 #include "AvaGeometryBaseModifier.h"
-#include "DynamicMesh/DynamicMesh3.h"
+#include "CEMeshBuilder.h"
+#include "Extensions/AvaRenderStateUpdateModifierExtension.h"
+#include "Extensions/AvaSceneTreeUpdateModifierExtension.h"
 #include "Templates/SubclassOf.h"
 #include "AvaDynamicMeshConverterModifier.generated.h"
 
@@ -43,9 +45,19 @@ struct FAvaDynamicMeshConverterModifierComponentState
 	explicit FAvaDynamicMeshConverterModifierComponentState() {}
 	explicit FAvaDynamicMeshConverterModifierComponentState(UPrimitiveComponent* InPrimitiveComponent);
 
+	friend uint32 GetTypeHash(const FAvaDynamicMeshConverterModifierComponentState& InItem)
+	{
+		return GetTypeHash(InItem.Component);
+	}
+
+	bool operator==(const FAvaDynamicMeshConverterModifierComponentState& InOther) const
+	{
+		return Component.IsValid() && Component == InOther.Component;
+	}
+
 	/** The component we are converting to dynamic mesh */
 	UPROPERTY()
-	TWeakObjectPtr<UPrimitiveComponent> Component = nullptr;
+	TWeakObjectPtr<UPrimitiveComponent> Component;
 
 	/** The default visibility of the actor converted component in game */
 	UPROPERTY()
@@ -63,12 +75,16 @@ struct FAvaDynamicMeshConverterModifierComponentState
 	UPROPERTY()
 	bool bComponentVisible = true;
 
-	/** The component converted dynamic mesh*/
-	UE::Geometry::FDynamicMesh3 Mesh;
+	/** Transform saved before mesh is converted */
+	UPROPERTY()
+	FTransform ActorRelativeTransform = FTransform::Identity;
 };
 
 UCLASS(MinimalAPI, BlueprintType)
-class UAvaDynamicMeshConverterModifier : public UAvaGeometryBaseModifier
+class UAvaDynamicMeshConverterModifier
+	: public UAvaGeometryBaseModifier
+	, public IAvaSceneTreeUpdateHandler
+	, public IAvaRenderStateUpdateHandler
 {
 	GENERATED_BODY()
 
@@ -144,8 +160,8 @@ protected:
 	//~ Begin UObject
 #if WITH_EDITOR
 	virtual void PostEditChangeProperty(FPropertyChangedEvent& PropertyChangedEvent) override;
+
 	/** Export current dynamic mesh to static mesh asset */
-	UFUNCTION(CallInEditor, Category="DynamicMeshConverter")
 	virtual void ConvertToStaticMeshAsset();
 #endif
 	//~ End UObject
@@ -157,11 +173,25 @@ protected:
 	virtual void RestorePreState() override;
 	virtual void Apply() override;
 	virtual void OnModifierRemoved(EActorModifierCoreDisableReason InReason) override;
+	virtual bool IsModifierDirtyable() const override;
 	//~ End UActorModifierCoreBase
 
+	//~ Begin IAvaSceneTreeUpdateHandler
+	virtual void OnSceneTreeTrackedActorChanged(int32 InIdx, AActor* InPreviousActor, AActor* InNewActor) {}
+	virtual void OnSceneTreeTrackedActorChildrenChanged(int32 InIdx, const TSet<TWeakObjectPtr<AActor>>& InPreviousChildrenActors, const TSet<TWeakObjectPtr<AActor>>& InNewChildrenActors);
+	virtual void OnSceneTreeTrackedActorDirectChildrenChanged(int32 InIdx, const TArray<TWeakObjectPtr<AActor>>& InPreviousChildrenActors, const TArray<TWeakObjectPtr<AActor>>& InNewChildrenActors) {}
+	virtual void OnSceneTreeTrackedActorParentChanged(int32 InIdx, const TArray<TWeakObjectPtr<AActor>>& InPreviousParentActor, const TArray<TWeakObjectPtr<AActor>>& InNewParentActor) {}
+	virtual void OnSceneTreeTrackedActorRearranged(int32 InIdx, AActor* InRearrangedActor) {}
+	//~ End IAvaSceneTreeUpdateHandler
+
+	//~ Begin IAvaRenderStateUpdateHandler
+	virtual void OnRenderStateUpdated(AActor* InActor, UActorComponent* InComponent);
+	virtual void OnActorVisibilityChanged(AActor* InActor) {}
+	//~ End IAvaRenderStateUpdateHandler
+
 	void OnSourceActorChanged();
 
-	void ConvertComponents(TArray<FAvaDynamicMeshConverterModifierComponentState>& OutResults) const;
+	bool ConvertComponents(TArray<TWeakObjectPtr<UMaterialInterface>>& OutMaterialsWeak);
 	bool HasFlag(EAvaDynamicMeshConverterModifierType InFlag) const;
 
 	void AddDynamicMeshComponent();
@@ -175,10 +205,6 @@ protected:
 	void GetBrushComponents(const TArray<AActor*>& InActors, TArray<UBrushComponent*>& OutComponents) const;
 	void GetProceduralMeshComponents(const TArray<AActor*>& InActors, TArray<UProceduralMeshComponent*>& OutComponents) const;
 
-	/** Triggers the update of the mesh */
-	UFUNCTION(CallInEditor, Category="DynamicMeshConverter")
-	void ConvertMesh();
-
 	/** What actor should we copy from, by default is self */
 	UPROPERTY(EditInstanceOnly, Setter="SetSourceActorWeak", Getter="GetSourceActorWeak", Category="DynamicMeshConverter", meta=(DisplayName="SourceActor", AllowPrivateAccess="true"))
 	TWeakObjectPtr<AActor> SourceActorWeak = GetModifiedActor();
@@ -212,13 +238,19 @@ protected:
 	UPROPERTY()
 	bool bComponentCreated = false;
 
+	/** Update interval to compare if a transform has changed in converted components, when value <= 0 then skipped */
+	UPROPERTY(EditInstanceOnly, Category="DynamicMeshConverter")
+	float TransformUpdateInterval = 0;
+
 	/** Components converted to dynamic mesh */
 	UPROPERTY()
-	TArray<FAvaDynamicMeshConverterModifierComponentState> ConvertedComponents;
+	TSet<FAvaDynamicMeshConverterModifierComponentState> ConvertedComponents;
+
+	UPROPERTY(Transient, DuplicateTransient, NonTransactional)
+	FCEMeshBuilder MeshBuilder;
 
-	/** Cached mesh to set state without reconverting again */
-	TOptional<UE::Geometry::FDynamicMesh3> ConvertedMesh;
+	FAvaSceneTreeActor TrackedActor;
 
-	/** Do we convert mesh on next execution */
-	bool bConvertMesh = false;
+	/** Time elasped since */
+	double LastTransformUpdateTime = 0;
 };
