[Motion Design] Broadcast Output: Game Viewport Media Output to show up by default in the broadcast editor's device list.

#rnx

diff --git a/Engine/Plugins/Experimental/Avalanche/Source/AvalancheMedia/Private/Broadcast/OutputDevices/AvaGameViewportMediaOutput.h b/Engine/Plugins/Experimental/Avalanche/Source/AvalancheMedia/Private/Broadcast/OutputDevices/AvaGameViewportMediaOutput.h
--- a/Engine/Plugins/Experimental/Avalanche/Source/AvalancheMedia/Private/Broadcast/OutputDevices/AvaGameViewportMediaOutput.h
+++ b/Engine/Plugins/Experimental/Avalanche/Source/AvalancheMedia/Private/Broadcast/OutputDevices/AvaGameViewportMediaOutput.h
@@ -8,7 +8,7 @@
 DECLARE_LOG_CATEGORY_EXTERN(LogAvaGameViewportMedia, Log, All);
 
 UCLASS(BlueprintType, ClassGroup = "Motion Design Broadcast",
-	meta = (DisplayName = "Motion Design Game Viewport Media Output"))
+	meta = (DisplayName = "Motion Design Game Viewport Media Output", MediaIOCustomLayout = "AvaGameViewport"))
 class UAvaGameViewportMediaOutput : public UMediaOutput
 {
 	GENERATED_BODY()
