Motion Design
- Some changes to the viewport tool buttons.

diff --git a/Engine/Plugins/Experimental/Avalanche/Source/AvalancheLevelViewport/Private/Widgets/SAvaLevelViewportStatusBarButtons.cpp b/Engine/Plugins/Experimental/Avalanche/Source/AvalancheLevelViewport/Private/Widgets/SAvaLevelViewportStatusBarButtons.cpp
--- a/Engine/Plugins/Experimental/Avalanche/Source/AvalancheLevelViewport/Private/Widgets/SAvaLevelViewportStatusBarButtons.cpp
+++ b/Engine/Plugins/Experimental/Avalanche/Source/AvalancheLevelViewport/Private/Widgets/SAvaLevelViewportStatusBarButtons.cpp
@@ -100,7 +100,7 @@ void SAvaLevelViewportStatusBarButtons::Construct(const FArguments& InArgs, TSha
 
 		+ SHorizontalBox::Slot()
 		.AutoWidth()
-		.Padding(0, 0, Padding, 0)
+		.Padding(0, Padding, Padding, Padding)
 		[
 			SNew(SSeparator)
 			.Orientation(EOrientation::Orient_Vertical)
@@ -115,7 +115,7 @@ void SAvaLevelViewportStatusBarButtons::Construct(const FArguments& InArgs, TSha
 
 		+ SHorizontalBox::Slot()
 		.AutoWidth()
-		.Padding(0, 0, Padding, 0)
+		.Padding(0, Padding, Padding, Padding)
 		[
 			SNew(SSeparator)
 			.Orientation(EOrientation::Orient_Vertical)
@@ -130,7 +130,7 @@ void SAvaLevelViewportStatusBarButtons::Construct(const FArguments& InArgs, TSha
 
 		+ SHorizontalBox::Slot()
 		.AutoWidth()
-		.Padding(0, 0, Padding, 0)
+		.Padding(0, Padding, Padding, Padding)
 		[
 			SNew(SSeparator)
 			.Orientation(EOrientation::Orient_Vertical)
@@ -273,7 +273,7 @@ void SAvaLevelViewportStatusBarButtons::PopulateActorButtons(TSharedPtr<SHorizon
 	TSharedRef<SComboButton> AlignmentButton = ViewportStatusBarButton::MakeMenuButton(
 		LOCTEXT("ActorAlign", "Align Actors"),
 		FOnGetContent::CreateSP(this, &SAvaLevelViewportStatusBarButtons::GetActorAlignmentMenuContent),
-		FAvaLevelViewportStyle::Get().GetBrush(TEXT("Icons.Alignment.Left")),
+		FAvaLevelViewportStyle::Get().GetBrush(TEXT("Icons.Alignment.Center_Y")),
 		TAttribute<FSlateColor>::CreateSP(this, &SAvaLevelViewportStatusBarButtons::GetActorAlignmentColor)
 	);
 
@@ -297,11 +297,11 @@ void SAvaLevelViewportStatusBarButtons::PopulateActorButtons(TSharedPtr<SHorizon
 		[
 			ViewportStatusBarButton::MakeButton(
 				this,
-				CommandsRef.DisableAnimators,
+				CommandsRef.EnableAnimators,
 				AnimatorBrush,
-				&SAvaLevelViewportStatusBarButtons::DisableAnimators,
+				&SAvaLevelViewportStatusBarButtons::EnableAnimators,
 				&SAvaLevelViewportStatusBarButtons::GetAnimatorButtonEnabled,
-				&SAvaLevelViewportStatusBarButtons::GetAnimatorButtonMuteColor
+				&SAvaLevelViewportStatusBarButtons::GetAnimatorButtonUnmuteColor
 			)
 		];
 
@@ -311,11 +311,11 @@ void SAvaLevelViewportStatusBarButtons::PopulateActorButtons(TSharedPtr<SHorizon
 		[
 			ViewportStatusBarButton::MakeButton(
 				this,
-				CommandsRef.EnableAnimators,
+				CommandsRef.DisableAnimators,
 				AnimatorBrush,
-				&SAvaLevelViewportStatusBarButtons::EnableAnimators,
+				&SAvaLevelViewportStatusBarButtons::DisableAnimators,
 				&SAvaLevelViewportStatusBarButtons::GetAnimatorButtonEnabled,
-				&SAvaLevelViewportStatusBarButtons::GetAnimatorButtonUnmuteColor
+				&SAvaLevelViewportStatusBarButtons::GetAnimatorButtonMuteColor
 			)
 		];
 }
@@ -323,6 +323,7 @@ void SAvaLevelViewportStatusBarButtons::PopulateActorButtons(TSharedPtr<SHorizon
 void SAvaLevelViewportStatusBarButtons::PopulateViewportButtons(TSharedPtr<SHorizontalBox> InContainer)
 {
 	using namespace UE::AvaLevelViewport::Private;
+	using namespace UE::Ava::LevelViewportStatusBarButtons::Private;
 
 	const FAvaLevelViewportCommands& CommandsRef = FAvaLevelViewportCommands::GetInternal();
 	const FLevelViewportCommands& ViewportActionsRef = FLevelViewportCommands::Get();
@@ -344,7 +345,7 @@ void SAvaLevelViewportStatusBarButtons::PopulateViewportButtons(TSharedPtr<SHori
 	TSharedRef<SComboButton> PostProcessButton = ViewportStatusBarButton::MakeMenuButton(
 		LOCTEXT("PostProcessEffects", "Post Process Effects"),
 		FOnGetContent::CreateSP(this, &SAvaLevelViewportStatusBarButtons::GetPostProcessMenuContent),
-		FSlateIcon(FAppStyle::GetAppStyleSetName(), "Icons.Advanced").GetIcon(),
+		RGBChannelIcon.GetIcon(),
 		TAttribute<FSlateColor>::CreateSP(this, &SAvaLevelViewportStatusBarButtons::GetPostProcessColor)
 	);
 
@@ -399,6 +400,28 @@ void SAvaLevelViewportStatusBarButtons::PopulateViewportButtons(TSharedPtr<SHori
 			)
 		];
 
+	InContainer->AddSlot()
+		.AutoWidth()
+		.Padding(5.f)
+		[
+			SNew(SSeparator)
+			.Orientation(EOrientation::Orient_Vertical)
+		];
+
+	InContainer->AddSlot()
+		.AutoWidth()
+		.Padding(ViewportStatusBarButton::Padding)
+		[
+			ViewportStatusBarButton::MakeButton(
+				this,
+				CommandsRef.ToggleOverlay,
+				FAppStyle::GetBrush("Icons.Visible"),
+				&SAvaLevelViewportStatusBarButtons::ToggleOverlay,
+				&SAvaLevelViewportStatusBarButtons::GetToggleOverlayEnabled,
+				&SAvaLevelViewportStatusBarButtons::GetToggleOverlayColor
+			)
+		];
+
 	InContainer->AddSlot()
 		.AutoWidth()
 		.Padding(ViewportStatusBarButton::Padding)
@@ -466,16 +489,10 @@ void SAvaLevelViewportStatusBarButtons::PopulateViewportButtons(TSharedPtr<SHori
 
 	InContainer->AddSlot()
 		.AutoWidth()
-		.Padding(ViewportStatusBarButton::Padding)
+		.Padding(5.f)
 		[
-			ViewportStatusBarButton::MakeButton(
-				this,
-				CommandsRef.ToggleOverlay,
-				FAppStyle::GetBrush("Icons.Visible"),
-				&SAvaLevelViewportStatusBarButtons::ToggleOverlay,
-				&SAvaLevelViewportStatusBarButtons::GetToggleOverlayEnabled,
-				&SAvaLevelViewportStatusBarButtons::GetToggleOverlayColor
-			)
+			SNew(SSeparator)
+			.Orientation(EOrientation::Orient_Vertical)
 		];
 
 	TSharedRef<SAvaMultiComboButton> SnapButton = ViewportStatusBarButton::MakeMultiMenuButton(
diff --git a/Engine/Plugins/Experimental/Avalanche/Source/AvalancheViewport/Private/AvaViewportSettings.cpp b/Engine/Plugins/Experimental/Avalanche/Source/AvalancheViewport/Private/AvaViewportSettings.cpp
--- a/Engine/Plugins/Experimental/Avalanche/Source/AvalancheViewport/Private/AvaViewportSettings.cpp
+++ b/Engine/Plugins/Experimental/Avalanche/Source/AvalancheViewport/Private/AvaViewportSettings.cpp
@@ -15,7 +15,7 @@ UAvaViewportSettings::UAvaViewportSettings()
 	ViewportCheckerboardColor1 = FLinearColor(FVector3f(0.177888f));
 	ViewportCheckerboardSize = 8.0f;
 
-	bGridEnabled = true;
+	bGridEnabled = false;
 	bGridAlwaysVisible = false;
 	GridSize = 50;
 	GridColor = FLinearColor(0.6, 0.0, 0.0, 0.4);
@@ -24,7 +24,7 @@ UAvaViewportSettings::UAvaViewportSettings()
 	bPixelGridEnabled = true;
 	PixelGridColor = FLinearColor(0.5, 0.5, 0.5, 0.4);
 
-	SnapState = static_cast<int32>(EAvaViewportSnapState::Global | EAvaViewportSnapState::Screen | EAvaViewportSnapState::Grid);
+	SnapState = static_cast<int32>(EAvaViewportSnapState::Off);
 	bSnapIndicatorsEnabled = true;
 	SnapIndicatorColor = FLinearColor::Red;
 	SnapIndicatorThickness = 2.f;
