Motion Design: Set a default template TL tree

#virtualized

diff --git a/Engine/Plugins/Experimental/Avalanche/Source/AvalancheTransitionEditor/Private/Settings/AvaTransitionEditorSettings.cpp b/Engine/Plugins/Experimental/Avalanche/Source/AvalancheTransitionEditor/Private/Settings/AvaTransitionEditorSettings.cpp
--- a/Engine/Plugins/Experimental/Avalanche/Source/AvalancheTransitionEditor/Private/Settings/AvaTransitionEditorSettings.cpp
+++ b/Engine/Plugins/Experimental/Avalanche/Source/AvalancheTransitionEditor/Private/Settings/AvaTransitionEditorSettings.cpp
@@ -8,6 +8,7 @@ UAvaTransitionEditorSettings::UAvaTransitionEditorSettings()
 {
 	CategoryName = TEXT("Motion Design");
 	SectionName  = TEXT("Transition Logic");
+	DefaultTemplate = FSoftObjectPath(TEXT("/Avalanche/TransitionLogic/TL_TemplateTree.TL_TemplateTree"));
 }
 
 UAvaTransitionTreeEditorData* UAvaTransitionEditorSettings::LoadDefaultTemplateEditorData() const
