Motion Design
- Some more changes to the toolbar.

diff --git a/Engine/Plugins/Experimental/Avalanche/Source/AvalancheLevelViewport/Private/Widgets/SAvaLevelViewportStatusBarButtons.cpp b/Engine/Plugins/Experimental/Avalanche/Source/AvalancheLevelViewport/Private/Widgets/SAvaLevelViewportStatusBarButtons.cpp
--- a/Engine/Plugins/Experimental/Avalanche/Source/AvalancheLevelViewport/Private/Widgets/SAvaLevelViewportStatusBarButtons.cpp
+++ b/Engine/Plugins/Experimental/Avalanche/Source/AvalancheLevelViewport/Private/Widgets/SAvaLevelViewportStatusBarButtons.cpp
@@ -346,7 +346,7 @@ void SAvaLevelViewportStatusBarButtons::PopulateViewportButtons(TSharedPtr<SHori
 		LOCTEXT("PostProcessEffects", "Post Process Effects"),
 		FOnGetContent::CreateSP(this, &SAvaLevelViewportStatusBarButtons::GetPostProcessMenuContent),
 		RGBChannelIcon.GetIcon(),
-		TAttribute<FSlateColor>::CreateSP(this, &SAvaLevelViewportStatusBarButtons::GetPostProcessColor)
+		FLinearColor::White
 	);
 
 	PostProcessButton->SetEnabled(TAttribute<bool>::CreateSP(this, &SAvaLevelViewportStatusBarButtons::GetPostProcessEnabled));
@@ -386,20 +386,6 @@ void SAvaLevelViewportStatusBarButtons::PopulateViewportButtons(TSharedPtr<SHori
 			)
 		];
 
-	InContainer->AddSlot()
-		.AutoWidth()
-		.Padding(ViewportStatusBarButton::Padding)
-		[
-			ViewportStatusBarButton::MakeButton(
-				this,
-				CommandsRef.ToggleSafeFrames,
-				FAvaLevelViewportStyle::Get().GetBrush(TEXT("Button.SafeFrames")),
-				&SAvaLevelViewportStatusBarButtons::ToggleSafeFrames,
-				&SAvaLevelViewportStatusBarButtons::GetToggleSafeFramesEnabled,
-				&SAvaLevelViewportStatusBarButtons::GetToggleSafeFramesColor
-			)
-		];
-
 	InContainer->AddSlot()
 		.AutoWidth()
 		.Padding(5.f)
@@ -422,6 +408,20 @@ void SAvaLevelViewportStatusBarButtons::PopulateViewportButtons(TSharedPtr<SHori
 			)
 		];
 
+	InContainer->AddSlot()
+		.AutoWidth()
+		.Padding(ViewportStatusBarButton::Padding)
+		[
+			ViewportStatusBarButton::MakeButton(
+				this,
+				CommandsRef.ToggleSafeFrames,
+				FAvaLevelViewportStyle::Get().GetBrush(TEXT("Button.SafeFrames")),
+				&SAvaLevelViewportStatusBarButtons::ToggleSafeFrames,
+				&SAvaLevelViewportStatusBarButtons::GetToggleSafeFramesEnabled,
+				&SAvaLevelViewportStatusBarButtons::GetToggleSafeFramesColor
+			)
+		];
+
 	InContainer->AddSlot()
 		.AutoWidth()
 		.Padding(ViewportStatusBarButton::Padding)
@@ -467,6 +467,14 @@ void SAvaLevelViewportStatusBarButtons::PopulateViewportButtons(TSharedPtr<SHori
 			TextureOverlayButton
 		];
 
+	InContainer->AddSlot()
+		.AutoWidth()
+		.Padding(5.f)
+		[
+			SNew(SSeparator)
+			.Orientation(EOrientation::Orient_Vertical)
+		];
+
 	TSharedRef<SAvaMultiComboButton> GridButton = ViewportStatusBarButton::MakeMultiMenuButton(
 		CommandsRef.ToggleGrid->GetDescription(),
 		FOnGetContent::CreateSP(this, &SAvaLevelViewportStatusBarButtons::GetGridMenuContent),
@@ -487,14 +495,6 @@ void SAvaLevelViewportStatusBarButtons::PopulateViewportButtons(TSharedPtr<SHori
 			GridButton
 		];
 
-	InContainer->AddSlot()
-		.AutoWidth()
-		.Padding(5.f)
-		[
-			SNew(SSeparator)
-			.Orientation(EOrientation::Orient_Vertical)
-		];
-
 	TSharedRef<SAvaMultiComboButton> SnapButton = ViewportStatusBarButton::MakeMultiMenuButton(
 		CommandsRef.ToggleSnapping->GetDescription(),
 		FOnGetContent::CreateSP(this, &SAvaLevelViewportStatusBarButtons::GetSnappingMenuContent),
@@ -546,25 +546,6 @@ void SAvaLevelViewportStatusBarButtons::PopulateViewportButtons(TSharedPtr<SHori
 		];
 }
 
-FSlateColor SAvaLevelViewportStatusBarButtons::GetPostProcessColor() const
-{
-	using namespace UE::AvaLevelViewport::Private;
-
-	const FAvaLevelViewportGuideFrameAndClient FrameAndClient(ViewportFrameWeak);
-
-	if (FrameAndClient.IsValid())
-	{
-		if (const TSharedPtr<FAvaViewportPostProcessManager> PostProcessManager = FrameAndClient.ViewportClient->GetPostProcessManager())
-		{
-			return PostProcessManager->GetType() != EAvaViewportPostProcessType::None
-				? ViewportStatusBarButton::ActiveColor
-				: ViewportStatusBarButton::EnabledColor;
-		}
-	}
-
-	return ViewportStatusBarButton::DisabledColor;
-}
-
 bool SAvaLevelViewportStatusBarButtons::GetPostProcessEnabled() const
 {
 	const FAvaLevelViewportGuideFrameAndClient FrameAndClient(ViewportFrameWeak);
diff --git a/Engine/Plugins/Experimental/Avalanche/Source/AvalancheLevelViewport/Private/Widgets/SAvaLevelViewportStatusBarButtons.h b/Engine/Plugins/Experimental/Avalanche/Source/AvalancheLevelViewport/Private/Widgets/SAvaLevelViewportStatusBarButtons.h
--- a/Engine/Plugins/Experimental/Avalanche/Source/AvalancheLevelViewport/Private/Widgets/SAvaLevelViewportStatusBarButtons.h
+++ b/Engine/Plugins/Experimental/Avalanche/Source/AvalancheLevelViewport/Private/Widgets/SAvaLevelViewportStatusBarButtons.h
@@ -54,7 +54,6 @@ protected:
 	FReply DisableAnimators();
 
 	// Post Process
-	FSlateColor GetPostProcessColor() const;
 	bool GetPostProcessEnabled() const;
 	TSharedRef<SWidget> GetPostProcessMenuContent();
 
