[Motion Design] Fixing Rundown Editor's Page Detail Undo/Redo

- Correctly recording transactions on RC Property/Controller changes using NotifyHook. Transactions are already handled by PropertyHandlers, just need to mark Rundown->Modify() at the right time.
- PageLists refresh is more efficient. Since it doesn't rebuild the list completely if it doesn't need to, refreshing on an RC value undo doesn't change the list nor it's selection and is much nicer. Nothing visually changes except the value that was undone.
- Undo Backup helper is for sending proper Value Change event so that Page details can refresh and the Runtime Preview Render.

diff --git a/Engine/Plugins/Experimental/Avalanche/Source/AvalancheMedia/Private/Rundown/AvaRundown.cpp b/Engine/Plugins/Experimental/Avalanche/Source/AvalancheMedia/Private/Rundown/AvaRundown.cpp
--- a/Engine/Plugins/Experimental/Avalanche/Source/AvalancheMedia/Private/Rundown/AvaRundown.cpp
+++ b/Engine/Plugins/Experimental/Avalanche/Source/AvalancheMedia/Private/Rundown/AvaRundown.cpp
@@ -163,6 +163,52 @@ void UAvaRundown::PostLoad()
 }
 
 #if WITH_EDITOR
+
+// Undo backup helper.
+class UAvaRundown::FPreUndoBackup
+{
+public:
+	FAvaRundownPageCollection TemplatePages;
+	FAvaRundownPageCollection InstancedPages;
+
+	static const FAvaRundownPage& GetPage(int32 InPageId, const FAvaRundownPageCollection& InCollection)
+	{
+		const int32 PageIndex = InCollection.GetPageIndex(InPageId);
+		return InCollection.Pages.IsValidIndex(PageIndex) ? InCollection.Pages[PageIndex] : FAvaRundownPage::NullPage;
+	}
+
+	static void NofifyPageValueChanges(UAvaRundown* InRundown, const FAvaRundownPageCollection& InCollection, const FAvaRundownPageCollection& InOtherCollection)
+	{
+		for (const FAvaRundownPage& Page : InCollection.Pages)
+		{
+			const FAvaRundownPage& BackupPage = GetPage(Page.GetPageId(), InOtherCollection);
+			if (BackupPage.IsValidPage())
+			{
+				EAvaPlayableRemoteControlChanges ValueChanges = EAvaPlayableRemoteControlChanges::None;
+				if (!BackupPage.GetRemoteControlValues().HasSameControllerValues(Page.GetRemoteControlValues()))
+				{
+					ValueChanges |= EAvaPlayableRemoteControlChanges::ControllerValues;
+				}
+				if (!BackupPage.GetRemoteControlValues().HasSameEntityValues(Page.GetRemoteControlValues()))
+				{
+					ValueChanges |= EAvaPlayableRemoteControlChanges::EntityValues;
+				}
+				if (ValueChanges != EAvaPlayableRemoteControlChanges::None)
+				{
+					InRundown->NotifyPageRemoteControlValueChanged(Page.GetPageId(), ValueChanges);
+				}
+			}
+		}
+	}
+};
+
+void UAvaRundown::PreEditUndo()
+{
+	PreUndoBackup = MakePimpl<FPreUndoBackup>();
+	PreUndoBackup->TemplatePages = TemplatePages;
+	PreUndoBackup->InstancedPages = InstancedPages;
+}
+
 void UAvaRundown::PostEditUndo()
 {
 	UObject::PostEditUndo();
@@ -180,6 +226,13 @@ void UAvaRundown::PostEditUndo()
 	}
 
 	GetOnActiveListChanged().Broadcast();
+
+	if (PreUndoBackup)
+	{
+		FPreUndoBackup::NofifyPageValueChanges(this, TemplatePages, PreUndoBackup->TemplatePages);
+		FPreUndoBackup::NofifyPageValueChanges(this, InstancedPages, PreUndoBackup->InstancedPages);
+		PreUndoBackup.Reset();
+	}
 }
 #endif
 
@@ -1684,7 +1737,7 @@ void UAvaRundown::NotifyPageRemoteControlValueChanged(int32 InPageId, EAvaPlayab
 		// will only set the value of the entity if it changed.
 		PushRuntimeRemoteControlValues(InPageId, true);
 	}
-	OnPagesChanged.Broadcast(this, {InPageId}, EAvaRundownPageChanges::RemoteControlValues);
+	OnPagesChanged.Broadcast(this, GetPage(InPageId), EAvaRundownPageChanges::RemoteControlValues);
 }
 
 #if WITH_EDITOR
diff --git a/Engine/Plugins/Experimental/Avalanche/Source/AvalancheMedia/Public/Rundown/AvaRundown.h b/Engine/Plugins/Experimental/Avalanche/Source/AvalancheMedia/Public/Rundown/AvaRundown.h
--- a/Engine/Plugins/Experimental/Avalanche/Source/AvalancheMedia/Public/Rundown/AvaRundown.h
+++ b/Engine/Plugins/Experimental/Avalanche/Source/AvalancheMedia/Public/Rundown/AvaRundown.h
@@ -333,6 +333,7 @@ public:
 	virtual void BeginDestroy() override;
 	virtual void PostLoad() override;
 #if WITH_EDITOR
+	virtual void PreEditUndo() override;
 	virtual void PostEditUndo() override;
 #endif
 	//~ End UObject Interface
@@ -789,4 +790,9 @@ protected:
 	mutable FOnCanClosePlaybackContext OnCanClosePlaybackContext;
 
 	TMap<FGuid, int32> SubListIndices;
+
+#if WITH_EDITOR
+	class FPreUndoBackup;
+	TPimplPtr<FPreUndoBackup> PreUndoBackup;
+#endif
 };
diff --git a/Engine/Plugins/Experimental/Avalanche/Source/AvalancheMediaEditor/Private/Rundown/AvaRundownEditor.cpp b/Engine/Plugins/Experimental/Avalanche/Source/AvalancheMediaEditor/Private/Rundown/AvaRundownEditor.cpp
--- a/Engine/Plugins/Experimental/Avalanche/Source/AvalancheMediaEditor/Private/Rundown/AvaRundownEditor.cpp
+++ b/Engine/Plugins/Experimental/Avalanche/Source/AvalancheMediaEditor/Private/Rundown/AvaRundownEditor.cpp
@@ -944,11 +944,19 @@ UAvaRundown* FAvaRundownEditor::GetRundown() const
 	return AvaRundown.Get();
 }
 
-void FAvaRundownEditor::MarkAsModified()
+void FAvaRundownEditor::BeginModify()
 {
 	if (AvaRundown.IsValid())
 	{
-		AvaRundown->Modify();
+		AvaRundown->Modify(/*bAlwaysMarkDirty*/false);
+	}
+}
+
+void FAvaRundownEditor::MarkAsModified()
+{
+	if (AvaRundown.IsValid() && AvaRundown->CanModify())
+	{
+		AvaRundown->MarkPackageDirty();
 	}
 }
 
diff --git a/Engine/Plugins/Experimental/Avalanche/Source/AvalancheMediaEditor/Private/Rundown/DetailsView/RemoteControl/Controllers/SAvaRundownRCControllerPanel.cpp b/Engine/Plugins/Experimental/Avalanche/Source/AvalancheMediaEditor/Private/Rundown/DetailsView/RemoteControl/Controllers/SAvaRundownRCControllerPanel.cpp
--- a/Engine/Plugins/Experimental/Avalanche/Source/AvalancheMediaEditor/Private/Rundown/DetailsView/RemoteControl/Controllers/SAvaRundownRCControllerPanel.cpp
+++ b/Engine/Plugins/Experimental/Avalanche/Source/AvalancheMediaEditor/Private/Rundown/DetailsView/RemoteControl/Controllers/SAvaRundownRCControllerPanel.cpp
@@ -292,7 +292,8 @@ void SAvaRundownRCControllerPanel::UpdateDefaultValuesAndRefresh(const TArray<in
 
 void SAvaRundownRCControllerPanel::OnRemoteControlControllerModified(URemoteControlPreset* InPreset, const TSet<FGuid>& InModifiedControllerIds)
 {
-	if (!IsValid(InPreset) || !HasRemoteControlPreset(InPreset))
+	// Note: Ignore changes from the RCP Transaction listener.
+	if (!IsValid(InPreset) || !HasRemoteControlPreset(InPreset) || GIsTransacting)
 	{
 		return;
 	}
@@ -472,10 +473,20 @@ void SAvaRundownRCControllerPanel::FPropertyRowGeneratorWrapper::NotifyPostChang
 	
 	if (URemoteControlPreset* Preset = PresetWeak.Get())
 	{
+		const TSharedPtr<FAvaRundownEditor> RundownEditor = ParentPanel ? ParentPanel->RundownEditorWeak.Pin() : nullptr;
+
+		// Only capture a modification when scrubbing starts.
+		if (!OngoingPropertyChanges.Contains(InPropertyThatChanged) && RundownEditor)
+		{
+			OngoingPropertyChanges.Add(InPropertyThatChanged);
+			RundownEditor->BeginModify();
+		}
+
 		Preset->OnModifyController(InPropertyChangedEvent);
 		if (ParentPanel && InPropertyChangedEvent.ChangeType & EPropertyChangeType::ValueSet)
 		{
 			ParentPanel->UpdatePageSummary(true);
+			OngoingPropertyChanges.Remove(InPropertyThatChanged);
 		}
 	}
 }
diff --git a/Engine/Plugins/Experimental/Avalanche/Source/AvalancheMediaEditor/Private/Rundown/DetailsView/RemoteControl/Controllers/SAvaRundownRCControllerPanel.h b/Engine/Plugins/Experimental/Avalanche/Source/AvalancheMediaEditor/Private/Rundown/DetailsView/RemoteControl/Controllers/SAvaRundownRCControllerPanel.h
--- a/Engine/Plugins/Experimental/Avalanche/Source/AvalancheMediaEditor/Private/Rundown/DetailsView/RemoteControl/Controllers/SAvaRundownRCControllerPanel.h
+++ b/Engine/Plugins/Experimental/Avalanche/Source/AvalancheMediaEditor/Private/Rundown/DetailsView/RemoteControl/Controllers/SAvaRundownRCControllerPanel.h
@@ -90,6 +90,7 @@ private:
 		TSharedPtr<IPropertyRowGenerator> PropertyRowGenerator;
 		TWeakObjectPtr<URemoteControlPreset> PresetWeak;
 		SAvaRundownRCControllerPanel* ParentPanel = nullptr;
+		TSet<FProperty*> OngoingPropertyChanges;
 
 		FPropertyRowGeneratorWrapper(SAvaRundownRCControllerPanel* InParentPanel);
 		virtual ~FPropertyRowGeneratorWrapper();
diff --git a/Engine/Plugins/Experimental/Avalanche/Source/AvalancheMediaEditor/Private/Rundown/DetailsView/RemoteControl/Properties/SAvaRundownPageRemoteControlProps.cpp b/Engine/Plugins/Experimental/Avalanche/Source/AvalancheMediaEditor/Private/Rundown/DetailsView/RemoteControl/Properties/SAvaRundownPageRemoteControlProps.cpp
--- a/Engine/Plugins/Experimental/Avalanche/Source/AvalancheMediaEditor/Private/Rundown/DetailsView/RemoteControl/Properties/SAvaRundownPageRemoteControlProps.cpp
+++ b/Engine/Plugins/Experimental/Avalanche/Source/AvalancheMediaEditor/Private/Rundown/DetailsView/RemoteControl/Properties/SAvaRundownPageRemoteControlProps.cpp
@@ -2,18 +2,15 @@
 
 #include "SAvaRundownPageRemoteControlProps.h"
 #include "AvaRundownRCPropertyItem.h"
-#include "IAvaMediaEditorModule.h"
 #include "Playable/AvaPlayableRemoteControl.h"
 #include "Playable/AvaPlayableRemoteControlValues.h"
 #include "RemoteControlEntity.h"
 #include "RemoteControlPreset.h"
-#include "Rundown/AvaRundownCommands.h"
 #include "Rundown/AvaRundownEditor.h"
 #include "Rundown/AvaRundownEditorUtils.h"
 #include "Rundown/AvaRundownManagedInstanceCache.h"
 #include "Rundown/AvaRundownPage.h"
 #include "Rundown/DetailsView/RemoteControl/Properties/AvaRundownPagePropertyContextMenu.h"
-#include "Rundown/Pages/AvaRundownPagePropertyContext.h"
 #include "SAvaRundownRCPropertyItemRow.h"
 #include "SlateOptMacros.h"
 #include "UObject/NameTypes.h"
@@ -34,12 +31,52 @@ TArray<FAvaRundownRCPropertyTableRowExtensionDelegate>& SAvaRundownPageRemoteCon
 	return TableRowExtensionDelegates.FindOrAdd(InExtensionName);
 }
 
+class FAvaRundownPageRCPropsNotifyHookImpl : public FAvaRundownPageRCPropsNotifyHook
+{
+public:
+	FAvaRundownPageRCPropsNotifyHookImpl(const TWeakPtr<SAvaRundownPageRemoteControlProps>& InPanelWeak)
+		: PanelWeak(InPanelWeak)
+	{
+	}
+
+	virtual ~FAvaRundownPageRCPropsNotifyHookImpl() override = default;
+	
+	//~ Begin FNotifyHook
+	virtual void NotifyPostChange(const FPropertyChangedEvent& InPropertyChangedEvent, FProperty* InPropertyThatChanged) override
+	{
+		if (const TSharedPtr<SAvaRundownPageRemoteControlProps> Panel = PanelWeak.Pin())
+		{
+			const TSharedPtr<FAvaRundownEditor> RundownEditor = Panel->RundownEditorWeak.Pin();
+			
+			// Only capture a modification when scrubbing starts.
+			if (!OngoingPropertyChanges.Contains(InPropertyThatChanged) && RundownEditor)
+			{
+				OngoingPropertyChanges.Add(InPropertyThatChanged);
+				RundownEditor->BeginModify();
+			}
+		
+			// Apply change to page immediately to capture it in the transaction.
+			Panel->OnPostPropertyChanged(InPropertyThatChanged);
+
+			if (InPropertyChangedEvent.ChangeType & EPropertyChangeType::ValueSet)
+			{
+				OngoingPropertyChanges.Remove(InPropertyThatChanged);
+			}
+		}
+	}
+	//~ End FNotifyHook
+
+	TWeakPtr<SAvaRundownPageRemoteControlProps> PanelWeak;
+	TSet<FProperty*> OngoingPropertyChanges;
+};
+
 void SAvaRundownPageRemoteControlProps::Construct(const FArguments& InArgs, TSharedPtr<FAvaRundownEditor> InRundownEditor)
 {
 	RundownEditorWeak = InRundownEditor;
 	ActivePageId = FAvaRundownPage::InvalidPageId;
 
 	CommandList = MakeShared<FUICommandList>();
+	NotifyHook = MakeShared<FAvaRundownPageRCPropsNotifyHookImpl>(SharedThis<SAvaRundownPageRemoteControlProps>(this).ToWeakPtr());
 
 	ContextMenu = MakeShared<FAvaRundownPagePropertyContextMenu>(CommandList);
 
@@ -282,7 +319,8 @@ void SAvaRundownPageRemoteControlProps::Refresh(const TArray<int32>& InSelectedP
 // However, it doesn't seem to be called (or not always) if the entity is modified by a controller action. 
 void SAvaRundownPageRemoteControlProps::OnRemoteControlExposedPropertiesModified(URemoteControlPreset* InPreset, const TSet<FGuid>& InModifiedProperties)
 {
-	if (!IsValid(InPreset) ||!HasRemoteControlPreset(InPreset))
+	// Note: Ignore changes from the RCP Transaction listener.
+	if (!IsValid(InPreset) || !HasRemoteControlPreset(InPreset) || GIsTransacting)
 	{
 		return;
 	}
@@ -333,7 +371,8 @@ void SAvaRundownPageRemoteControlProps::OnRemoteControlExposedPropertiesModified
 
 void SAvaRundownPageRemoteControlProps::OnRemoteControlControllerModified(URemoteControlPreset* InPreset, const TSet<FGuid>& InModifiedControllerIds)
 {
-	if (!IsValid(InPreset) || !HasRemoteControlPreset(InPreset))
+	// Note: Ignore changes from the RCP Transaction listener.
+	if (!IsValid(InPreset) || !HasRemoteControlPreset(InPreset) || GIsTransacting)
 	{
 		return;
 	}
@@ -353,6 +392,27 @@ void SAvaRundownPageRemoteControlProps::OnRemoteControlControllerModified(URemot
 	OnRemoteControlExposedPropertiesModified(InPreset, EntityIds);
 }
 
+void SAvaRundownPageRemoteControlProps::OnPostPropertyChanged(FProperty* InPropertyThatChanged)
+{
+	// Find which property of which preset this is.
+	for (const TSharedPtr<FAvaRundownManagedInstance>& ManagedInstance : ManagedInstances)
+	{
+		if (URemoteControlPreset* Preset = ManagedInstance ? ManagedInstance->GetRemoteControlPreset() : nullptr)
+		{
+			for (const TWeakPtr<FRemoteControlProperty>& ExposedPropertyWeak : Preset->GetExposedEntities<FRemoteControlProperty>())
+			{
+				if (const TSharedPtr<FRemoteControlProperty> ExposedProperty = ExposedPropertyWeak.Pin())
+				{
+					if (ExposedProperty->GetProperty() == InPropertyThatChanged)
+					{
+						OnRemoteControlExposedPropertiesModified(Preset, {ExposedProperty->GetId()});
+					}
+				}
+			}
+		}
+	}
+}
+
 void SAvaRundownPageRemoteControlProps::BindRemoteControlDelegates(URemoteControlPreset* InPreset)
 {
 	if (IsValid(InPreset))
@@ -451,6 +511,11 @@ const TArray<FAvaRundownRCPropertyItemPtr> SAvaRundownPageRemoteControlProps::Ge
 	return PropertyContainer->GetSelectedItems();
 }
 
+TSharedPtr<FAvaRundownPageRCPropsNotifyHook> SAvaRundownPageRemoteControlProps::GetNotifyHook() const
+{
+	return NotifyHook;
+}
+
 TSharedPtr<SWidget> SAvaRundownPageRemoteControlProps::GetContextMenuContent()
 {
 	TArray<FAvaRundownRCPropertyItemPtr> SelectedItems = GetSelectedPropertyItems();
diff --git a/Engine/Plugins/Experimental/Avalanche/Source/AvalancheMediaEditor/Private/Rundown/DetailsView/RemoteControl/Properties/SAvaRundownPageRemoteControlProps.h b/Engine/Plugins/Experimental/Avalanche/Source/AvalancheMediaEditor/Private/Rundown/DetailsView/RemoteControl/Properties/SAvaRundownPageRemoteControlProps.h
--- a/Engine/Plugins/Experimental/Avalanche/Source/AvalancheMediaEditor/Private/Rundown/DetailsView/RemoteControl/Properties/SAvaRundownPageRemoteControlProps.h
+++ b/Engine/Plugins/Experimental/Avalanche/Source/AvalancheMediaEditor/Private/Rundown/DetailsView/RemoteControl/Properties/SAvaRundownPageRemoteControlProps.h
@@ -5,10 +5,11 @@
 #include "Containers/Array.h"
 #include "Containers/Map.h"
 #include "Containers/Set.h"
+#include "Misc/NotifyHook.h"
 #include "Templates/SharedPointer.h"
 #include "Widgets/SCompoundWidget.h"
-#include "Widgets/Views/SListView.h"
 #include "Widgets/Views/ITableRow.h"
+#include "Widgets/Views/SListView.h"
 
 class FAvaRundownEditor;
 class FAvaRundownManagedInstance;
@@ -33,6 +34,20 @@ DECLARE_MULTICAST_DELEGATE_TwoParams(FAvaRundownRCPropertyHeaderRowExtensionDele
 DECLARE_DELEGATE_ThreeParams(FAvaRundownRCPropertyTableRowExtensionDelegate, TSharedRef<SAvaRundownPageRemoteControlProps> Panel,
 	TSharedRef<const FAvaRundownRCPropertyItem> ItemPtr, TSharedPtr<SWidget>& CurrentWidget)
 
+/**
+ * Note:
+ *  Making a separate object for notify hook to break the SharedPtr circular dependency:
+ *  PanelWidget ->(strong) ItemWidget -> (strong) Hook -> (weak) PanelWidget
+ *  The ItemWidget has a weak ptr on the PanelWidget, but the Hook needs to be strong to match
+ *  ownership with the PropertyRowGenerator that is using the Hook.
+ */
+class FAvaRundownPageRCPropsNotifyHook : public FNotifyHook
+{
+public:
+	// Need to define a virtual destructor because FNotifyHook doesn't.
+	virtual ~FAvaRundownPageRCPropsNotifyHook() = default;
+};
+
 /**
  * The page props implementation for remote control fields.
  */
@@ -58,6 +73,8 @@ public:
 	void Refresh(const TArray<int32>& InSelectedPageIds);
 
 	const TArray<FAvaRundownRCPropertyItemPtr> GetSelectedPropertyItems() const;
+	
+	TSharedPtr<FAvaRundownPageRCPropsNotifyHook> GetNotifyHook() const;
 
 private:
 	static FAvaRundownRCPropertyHeaderRowExtensionDelegate HeaderRowExtensionDelegate;
@@ -69,6 +86,8 @@ private:
 	void OnRemoteControlExposedPropertiesModified(URemoteControlPreset* InPreset, const TSet<FGuid>& InModifiedProperties);
 	void OnRemoteControlControllerModified(URemoteControlPreset* InPreset, const TSet<FGuid>& InModifiedControllerIds);
 
+	void OnPostPropertyChanged(FProperty* InPropertyThatChanged);
+
 	void BindRemoteControlDelegates(URemoteControlPreset* InPreset);
 
 	bool HasRemoteControlPreset(const URemoteControlPreset* InPreset) const;
@@ -112,4 +131,8 @@ private:
 	TSharedPtr<FUICommandList> CommandList;
 
 	TSharedPtr<FAvaRundownPagePropertyContextMenu> ContextMenu;
+
+	TSharedPtr<FAvaRundownPageRCPropsNotifyHook> NotifyHook;
+
+	friend class FAvaRundownPageRCPropsNotifyHookImpl;
 };
diff --git a/Engine/Plugins/Experimental/Avalanche/Source/AvalancheMediaEditor/Private/Rundown/DetailsView/RemoteControl/Properties/SAvaRundownRCPropertyItemRow.cpp b/Engine/Plugins/Experimental/Avalanche/Source/AvalancheMediaEditor/Private/Rundown/DetailsView/RemoteControl/Properties/SAvaRundownRCPropertyItemRow.cpp
--- a/Engine/Plugins/Experimental/Avalanche/Source/AvalancheMediaEditor/Private/Rundown/DetailsView/RemoteControl/Properties/SAvaRundownRCPropertyItemRow.cpp
+++ b/Engine/Plugins/Experimental/Avalanche/Source/AvalancheMediaEditor/Private/Rundown/DetailsView/RemoteControl/Properties/SAvaRundownRCPropertyItemRow.cpp
@@ -12,8 +12,6 @@
 #include "RemoteControlEntity.h"
 #include "RemoteControlField.h"
 #include "RemoteControlPreset.h"
-#include "Rundown/DetailsView/RemoteControl/Properties/AvaRundownPageRemoteControlWidgetUtils.h"
-#include "Widgets/Layout/SScissorRectBox.h"
 #include "Widgets/SNullWidget.h"
 #include "Widgets/Text/STextBlock.h"
 
@@ -24,6 +22,7 @@ void SAvaRundownRCPropertyItemRow::Construct(const FArguments& InArgs, TSharedRe
 {
 	ItemPtrWeak = InRowItem;
 	PropertyPanelWeak = InPropertyPanel;
+	NotifyHook = InPropertyPanel->GetNotifyHook();
 	Generator = nullptr;
 	ValueContainer = nullptr;
 	ValueWidget = nullptr;
@@ -112,6 +111,7 @@ TSharedRef<SWidget> SAvaRundownRCPropertyItemRow::CreateValue()
 			if ((FieldPtr->FieldType == EExposedFieldType::Property) && (Objects.Num() > 0))
 			{
 				FPropertyRowGeneratorArgs Args;
+				Args.NotifyHook = NotifyHook.Get();
 				Generator = FModuleManager::GetModuleChecked<FPropertyEditorModule>("PropertyEditor").CreatePropertyRowGenerator(Args);
 				Generator->SetObjects({Objects[0]});
 
diff --git a/Engine/Plugins/Experimental/Avalanche/Source/AvalancheMediaEditor/Private/Rundown/DetailsView/RemoteControl/Properties/SAvaRundownRCPropertyItemRow.h b/Engine/Plugins/Experimental/Avalanche/Source/AvalancheMediaEditor/Private/Rundown/DetailsView/RemoteControl/Properties/SAvaRundownRCPropertyItemRow.h
--- a/Engine/Plugins/Experimental/Avalanche/Source/AvalancheMediaEditor/Private/Rundown/DetailsView/RemoteControl/Properties/SAvaRundownRCPropertyItemRow.h
+++ b/Engine/Plugins/Experimental/Avalanche/Source/AvalancheMediaEditor/Private/Rundown/DetailsView/RemoteControl/Properties/SAvaRundownRCPropertyItemRow.h
@@ -36,6 +36,7 @@ protected:
 
 	TWeakPtr<const FAvaRundownRCPropertyItem> ItemPtrWeak;
 	TWeakPtr<SAvaRundownPageRemoteControlProps> PropertyPanelWeak;
+	TSharedPtr<FAvaRundownPageRCPropsNotifyHook> NotifyHook;
 	TSharedPtr<IPropertyRowGenerator> Generator;
 	TSharedPtr<SBox> ValueContainer;
 	TSharedPtr<SWidget> ValueWidget;
diff --git a/Engine/Plugins/Experimental/Avalanche/Source/AvalancheMediaEditor/Private/Rundown/DetailsView/SAvaRundownPageDetails.cpp b/Engine/Plugins/Experimental/Avalanche/Source/AvalancheMediaEditor/Private/Rundown/DetailsView/SAvaRundownPageDetails.cpp
--- a/Engine/Plugins/Experimental/Avalanche/Source/AvalancheMediaEditor/Private/Rundown/DetailsView/SAvaRundownPageDetails.cpp
+++ b/Engine/Plugins/Experimental/Avalanche/Source/AvalancheMediaEditor/Private/Rundown/DetailsView/SAvaRundownPageDetails.cpp
@@ -190,8 +190,6 @@ void SAvaRundownPageDetails::Construct(const FArguments& InArgs, const TSharedPt
 			]
 		]
 	];
-
-	ActivePageId = FAvaRundownPage::InvalidPageId;
 }
 
 SAvaRundownPageDetails::~SAvaRundownPageDetails()
diff --git a/Engine/Plugins/Experimental/Avalanche/Source/AvalancheMediaEditor/Private/Rundown/Pages/PageViews/AvaRundownPageViewImpl.h b/Engine/Plugins/Experimental/Avalanche/Source/AvalancheMediaEditor/Private/Rundown/Pages/PageViews/AvaRundownPageViewImpl.h
--- a/Engine/Plugins/Experimental/Avalanche/Source/AvalancheMediaEditor/Private/Rundown/Pages/PageViews/AvaRundownPageViewImpl.h
+++ b/Engine/Plugins/Experimental/Avalanche/Source/AvalancheMediaEditor/Private/Rundown/Pages/PageViews/AvaRundownPageViewImpl.h
@@ -19,6 +19,12 @@ public:
 
 	FAvaRundownPageViewImpl(int32 InPageId, UAvaRundown* InRundown, const TSharedPtr<SAvaRundownPageList>& InPageList);
 
+	// Fast-path for in-place UI refresh. 
+	void RefreshPageId(int32 InPageId)
+	{
+		PageId = InPageId;
+	}
+
 	virtual UAvaRundown* GetRundown() const override;
 	
 	virtual int32 GetPageId() const override;
diff --git a/Engine/Plugins/Experimental/Avalanche/Source/AvalancheMediaEditor/Private/Rundown/Pages/Slate/SAvaRundownInstancedPageList.cpp b/Engine/Plugins/Experimental/Avalanche/Source/AvalancheMediaEditor/Private/Rundown/Pages/Slate/SAvaRundownInstancedPageList.cpp
--- a/Engine/Plugins/Experimental/Avalanche/Source/AvalancheMediaEditor/Private/Rundown/Pages/Slate/SAvaRundownInstancedPageList.cpp
+++ b/Engine/Plugins/Experimental/Avalanche/Source/AvalancheMediaEditor/Private/Rundown/Pages/Slate/SAvaRundownInstancedPageList.cpp
@@ -140,27 +140,68 @@ void SAvaRundownInstancedPageList::Refresh()
 		if (PageListReference.Type == EAvaRundownPageListType::Instance)
 		{
 			const TArray<FAvaRundownPage>& Pages = Rundown->GetInstancedPages().Pages;
+			const int32 VisiblePageCount = VisiblePageIds.IsEmpty() ? Pages.Num() : VisiblePageIds.Num();
 			
-			PageViews.Reset(Pages.Num());
+			if (PageViews.Num() != VisiblePageCount)
+			{
+				PageViews.Reset(Pages.Num());
 
-			for (const FAvaRundownPage& Page : Pages)
+				for (const FAvaRundownPage& Page : Pages)
+				{
+					if (IsPageVisible(Page))
+					{
+						PageViews.Emplace(MakeShared<FAvaRundownInstancedPageViewImpl>(Page.GetPageId(), Rundown, SharedThis(this)));
+					}
+				}
+			}
+			else
 			{
-				if (IsPageVisible(Page))
+				// Number of page didn't change, just refresh ids.
+				int32 PageViewIndex = 0;
+				for (const FAvaRundownPage& Page : Pages)
 				{
-					PageViews.Emplace(MakeShared<FAvaRundownInstancedPageViewImpl>(Page.GetPageId(), Rundown, SharedThis(this)));
+					if (IsPageVisible(Page))
+					{
+						if (FAvaRundownPageViewImpl* PageView = PageViews[PageViewIndex]->CastTo<FAvaRundownPageViewImpl>())
+						{
+							PageView->RefreshPageId(Page.GetPageId());
+						}
+						++PageViewIndex;
+					}
 				}
 			}
 		}
 		else if (Rundown->IsValidSubList(PageListReference))
 		{
 			const FAvaRundownSubList& SubList = Rundown->GetSubList(PageListReference);
-			PageViews.Reset(SubList.PageIds.Num());
+			const int32 VisiblePageCount = VisiblePageIds.IsEmpty() ? SubList.PageIds.Num() : VisiblePageIds.Num();
+
+			if (SubList.PageIds.Num() != VisiblePageCount)
+			{
+				PageViews.Reset(SubList.PageIds.Num());
 			
-			for (const int32 PageId : SubList.PageIds)
+				for (const int32 PageId : SubList.PageIds)
+				{
+					if (IsPageVisible(PageId))
+					{
+						PageViews.Emplace(MakeShared<FAvaRundownInstancedPageViewImpl>(PageId, Rundown, SharedThis(this)));
+					}
+				}
+			}
+			else
 			{
-				if (IsPageVisible(PageId))
+				// Number of page didn't change, just refresh ids.
+				int32 PageViewIndex = 0;
+				for (const int32 PageId : SubList.PageIds)
 				{
-					PageViews.Emplace(MakeShared<FAvaRundownInstancedPageViewImpl>(PageId, Rundown, SharedThis(this)));
+					if (IsPageVisible(PageId))
+					{
+						if (FAvaRundownPageViewImpl* PageView = PageViews[PageViewIndex]->CastTo<FAvaRundownPageViewImpl>())
+						{
+							PageView->RefreshPageId(PageId);
+						}
+						++PageViewIndex;
+					}
 				}
 			}
 		}
diff --git a/Engine/Plugins/Experimental/Avalanche/Source/AvalancheMediaEditor/Private/Rundown/Pages/Slate/SAvaRundownTemplatePageList.cpp b/Engine/Plugins/Experimental/Avalanche/Source/AvalancheMediaEditor/Private/Rundown/Pages/Slate/SAvaRundownTemplatePageList.cpp
--- a/Engine/Plugins/Experimental/Avalanche/Source/AvalancheMediaEditor/Private/Rundown/Pages/Slate/SAvaRundownTemplatePageList.cpp
+++ b/Engine/Plugins/Experimental/Avalanche/Source/AvalancheMediaEditor/Private/Rundown/Pages/Slate/SAvaRundownTemplatePageList.cpp
@@ -59,14 +59,35 @@ void SAvaRundownTemplatePageList::Refresh()
 			return;
 		}
 
-		const FAvaRundownPageCollection& PageCollection = Rundown->GetTemplatePages();
-		PageViews.Reset(PageCollection.Pages.Num());
+		const TArray<FAvaRundownPage>& Pages = Rundown->GetTemplatePages().Pages;
+		const int32 VisiblePageCount = VisiblePageIds.IsEmpty() ? Pages.Num() : VisiblePageIds.Num();
+			
+		if (PageViews.Num() != VisiblePageCount)
+		{
+			PageViews.Reset(Pages.Num());
 
-		for (const FAvaRundownPage& Page : PageCollection.Pages)
+			for (const FAvaRundownPage& Page : Pages)
+			{
+				if (IsPageVisible(Page))
+				{
+					PageViews.Emplace(MakeShared<FAvaRundownTemplatePageViewImpl>(Page.GetPageId(), Rundown, SharedThis(this)));
+				}
+			}
+		}
+		else
 		{
-			if (IsPageVisible(Page))
+			// Number of page didn't change, just refresh ids.
+			int32 PageViewIndex = 0;
+			for (const FAvaRundownPage& Page : Pages)
 			{
-				PageViews.Emplace(MakeShared<FAvaRundownTemplatePageViewImpl>(Page.GetPageId(), Rundown, SharedThis(this)));
+				if (IsPageVisible(Page))
+				{
+					if (FAvaRundownPageViewImpl* PageView = PageViews[PageViewIndex]->CastTo<FAvaRundownPageViewImpl>())
+					{
+						PageView->RefreshPageId(Page.GetPageId());
+					}
+					++PageViewIndex;
+				}
 			}
 		}
 
diff --git a/Engine/Plugins/Experimental/Avalanche/Source/AvalancheMediaEditor/Public/Rundown/AvaRundownEditor.h b/Engine/Plugins/Experimental/Avalanche/Source/AvalancheMediaEditor/Public/Rundown/AvaRundownEditor.h
--- a/Engine/Plugins/Experimental/Avalanche/Source/AvalancheMediaEditor/Public/Rundown/AvaRundownEditor.h
+++ b/Engine/Plugins/Experimental/Avalanche/Source/AvalancheMediaEditor/Public/Rundown/AvaRundownEditor.h
@@ -41,6 +41,11 @@ public:
 
 	bool IsRundownValid() const { return AvaRundown.IsValid(); }
 	AVALANCHEMEDIAEDITOR_API UAvaRundown* GetRundown() const;
+
+	/** Begin transaction buffer, but don't mark as dirty yet. */
+	void BeginModify();
+	
+	/** Mark the rundown as dirty. */
 	void MarkAsModified();
 
 	DECLARE_MULTICAST_DELEGATE_TwoParams(FOnPageEvent, const TArray<int32>&, UE::AvaRundown::EPageEvent);
