MotionDesign : Modifiers
- Visibility modifier marked dirty due to multiple render state dirty events and causing performance issue

diff --git a/Engine/Plugins/Experimental/Avalanche/Source/AvalancheModifiers/Private/Extensions/AvaRenderStateUpdateModifierExtension.cpp b/Engine/Plugins/Experimental/Avalanche/Source/AvalancheModifiers/Private/Extensions/AvaRenderStateUpdateModifierExtension.cpp
--- a/Engine/Plugins/Experimental/Avalanche/Source/AvalancheModifiers/Private/Extensions/AvaRenderStateUpdateModifierExtension.cpp
+++ b/Engine/Plugins/Experimental/Avalanche/Source/AvalancheModifiers/Private/Extensions/AvaRenderStateUpdateModifierExtension.cpp
@@ -103,13 +103,13 @@ void FAvaRenderStateUpdateModifierExtension::OnRenderStateDirty(UActorComponent&
 		return;
 	}
 
-	const UActorModifierCoreBase* Modifier = GetModifier();
-	if (!Modifier || !Modifier->IsModifierEnabled())
+	if (ActorDirty->GetLevel() != ModifierActor->GetLevel())
 	{
 		return;
 	}
 
-	if (ActorDirty->GetLevel() != ModifierActor->GetLevel())
+	const UActorModifierCoreBase* Modifier = GetModifier();
+	if (!Modifier || !Modifier->IsModifierEnabled() || !Modifier->IsModifierIdle())
 	{
 		return;
 	}
diff --git a/Engine/Plugins/Experimental/Avalanche/Source/AvalancheModifiers/Private/Modifiers/AvaVisibilityModifier.cpp b/Engine/Plugins/Experimental/Avalanche/Source/AvalancheModifiers/Private/Modifiers/AvaVisibilityModifier.cpp
--- a/Engine/Plugins/Experimental/Avalanche/Source/AvalancheModifiers/Private/Modifiers/AvaVisibilityModifier.cpp
+++ b/Engine/Plugins/Experimental/Avalanche/Source/AvalancheModifiers/Private/Modifiers/AvaVisibilityModifier.cpp
@@ -113,6 +113,9 @@ void UAvaVisibilityModifier::Apply()
 
 	ChildrenActorsWeak = NewChildrenActorsWeak;
 
+	FAvaRenderStateUpdateModifierExtension* RenderStateExtension = GetExtension<FAvaRenderStateUpdateModifierExtension>();
+	RenderStateExtension->SetTrackedActorsVisibility(ChildrenActorsWeak);
+
 	Next();
 }
 
@@ -174,18 +177,18 @@ void UAvaVisibilityModifier::OnModifiedActorTransformed()
 	// Overwrite parent class behaviour don't do anything when moved
 }
 
-void UAvaVisibilityModifier::OnRenderStateUpdated(AActor* InActor, UActorComponent* InComponent)
+void UAvaVisibilityModifier::OnActorVisibilityChanged(AActor* InActor)
 {
-	Super::OnRenderStateUpdated(InActor, InComponent);
+	Super::OnActorVisibilityChanged(InActor);
 
 	AActor* ActorModified = GetModifiedActor();
 
-	// Only handle what is linked to us
 	if (!IsValid(ActorModified))
 	{
 		return;
 	}
 
+	// Only handle what is linked to us
 	const bool bThisActorUpdated = InActor == ActorModified;
 	const bool bActorAttachedToThisUpdated = InActor->IsAttachedTo(ActorModified);
 
@@ -194,7 +197,7 @@ void UAvaVisibilityModifier::OnRenderStateUpdated(AActor* InActor, UActorCompone
 		return;
 	}
 
-	// If no modifier is found above us, then we handle this case
+	// If no modifier is found above us, then we handle this case otherwise let the other modifier handle it
 	const UAvaVisibilityModifier* Modifier = GetFirstModifierAbove(ActorModified);
 
 	if (bThisActorUpdated && Modifier)
diff --git a/Engine/Plugins/Experimental/Avalanche/Source/AvalancheModifiers/Public/Modifiers/AvaVisibilityModifier.h b/Engine/Plugins/Experimental/Avalanche/Source/AvalancheModifiers/Public/Modifiers/AvaVisibilityModifier.h
--- a/Engine/Plugins/Experimental/Avalanche/Source/AvalancheModifiers/Public/Modifiers/AvaVisibilityModifier.h
+++ b/Engine/Plugins/Experimental/Avalanche/Source/AvalancheModifiers/Public/Modifiers/AvaVisibilityModifier.h
@@ -68,7 +68,7 @@ protected:
 	//~ End UActorModifierCoreBase
 
 	//~ Begin IAvaRenderStateUpdateExtension
-	virtual void OnRenderStateUpdated(AActor* InActor, UActorComponent* InComponent) override;
+	virtual void OnActorVisibilityChanged(AActor* InActor) override;
 	//~ End IAvaRenderStateUpdateExtension
 
 	/** Used by other modifiers of this class to check if we are hiding an actor */
