Motion Design: Fix issue with World GC Leak crash on Sequences with Director BP after having played them at least once with Ava Sequence Player

diff --git a/Engine/Plugins/Experimental/Avalanche/Source/AvalancheSequence/Private/AvaSequencePlayer.cpp b/Engine/Plugins/Experimental/Avalanche/Source/AvalancheSequence/Private/AvaSequencePlayer.cpp
--- a/Engine/Plugins/Experimental/Avalanche/Source/AvalancheSequence/Private/AvaSequencePlayer.cpp
+++ b/Engine/Plugins/Experimental/Avalanche/Source/AvalancheSequence/Private/AvaSequencePlayer.cpp
@@ -86,6 +86,8 @@ void UAvaSequencePlayer::SetPlaySettings(const FAvaSequencePlayParams& InPlaySet
 
 	SetTimeRange(StartTimeSeconds, DurationSeconds);
 
+	bTearDownOnFinished = InPlaySettings.AdvancedSettings.bTearDownOnFinished;
+
 	PlaybackSettings = FMovieSceneSequencePlaybackSettings();
 	PlaybackSettings.PlayRate = InPlaySettings.AdvancedSettings.PlaybackSpeed;
 	PlaybackSettings.LoopCount.Value = InPlaySettings.AdvancedSettings.LoopCount;
@@ -226,4 +228,9 @@ void UAvaSequencePlayer::NotifySequencePaused()
 void UAvaSequencePlayer::NotifySequenceFinished()
 {
 	OnSequenceFinishedDelegate.Broadcast(this, GetAvaSequence());
+
+	if (bTearDownOnFinished)
+	{
+		TearDown();
+	}
 }
diff --git a/Engine/Plugins/Experimental/Avalanche/Source/AvalancheSequence/Private/Director/AvaSequenceDirector.cpp b/Engine/Plugins/Experimental/Avalanche/Source/AvalancheSequence/Private/Director/AvaSequenceDirector.cpp
--- a/Engine/Plugins/Experimental/Avalanche/Source/AvalancheSequence/Private/Director/AvaSequenceDirector.cpp
+++ b/Engine/Plugins/Experimental/Avalanche/Source/AvalancheSequence/Private/Director/AvaSequenceDirector.cpp
@@ -9,19 +9,19 @@
 
 TScriptInterface<IAvaSequencePlaybackObject> UAvaSequenceDirector::GetPlaybackObject() const
 {
-	if (PlaybackObject)
+	if (PlaybackObjectInterfaceWeak.IsValid())
 	{
-		return PlaybackObject;
+		return PlaybackObjectInterfaceWeak.ToScriptInterface();
 	}
 
 	// If Playback Object is null, try to update from SequencePlayer PlaybackObject
 	const_cast<UAvaSequenceDirector*>(this)->UpdatePlaybackObject();
-
-	return PlaybackObject;
+	return PlaybackObjectInterfaceWeak.ToScriptInterface();
 }
 
 void UAvaSequenceDirector::UpdatePlaybackObject()
 {
+	UAvaSequencePlayer* SequencePlayer = SequencePlayerWeak.Get();
 	if (!SequencePlayer)
 	{
 		return;
@@ -29,8 +29,7 @@ void UAvaSequenceDirector::UpdatePlaybackObject()
 
 	if (IAvaSequencePlaybackObject* PlaybackObjectInterface = SequencePlayer->GetPlaybackObject())
 	{
-		PlaybackObject.SetInterface(PlaybackObjectInterface);
-		PlaybackObject.SetObject(PlaybackObjectInterface->ToUObject());
+		PlaybackObjectInterfaceWeak = PlaybackObjectInterface;
 	}
 }
 
@@ -56,7 +55,8 @@ void UAvaSequenceDirector::PostDuplicate(EDuplicateMode::Type InDuplicateMode)
 
 void UAvaSequenceDirector::Initialize(IMovieScenePlayer& InPlayer, IAvaSequenceProvider* InSequenceProvider)
 {
-	SequencePlayer = Cast<UAvaSequencePlayer>(InPlayer.AsUObject());
+	UAvaSequencePlayer* SequencePlayer = Cast<UAvaSequencePlayer>(InPlayer.AsUObject());
+	SequencePlayerWeak = SequencePlayer;
 
 	IAvaSequencePlaybackObject* NewPlaybackObject = nullptr;
 
@@ -74,12 +74,11 @@ void UAvaSequenceDirector::Initialize(IMovieScenePlayer& InPlayer, IAvaSequenceP
 
 	if (NewPlaybackObject)
 	{
-		PlaybackObject.SetInterface(NewPlaybackObject);
-		PlaybackObject.SetObject(NewPlaybackObject->ToUObject());
+		PlaybackObjectInterfaceWeak = NewPlaybackObject;
 	}
 	else
 	{
-		PlaybackObject.SetObject(nullptr);
+		PlaybackObjectInterfaceWeak = nullptr;
 	}
 
 	UpdateProperties();
diff --git a/Engine/Plugins/Experimental/Avalanche/Source/AvalancheSequence/Public/AvaSequencePlayer.h b/Engine/Plugins/Experimental/Avalanche/Source/AvalancheSequence/Public/AvaSequencePlayer.h
--- a/Engine/Plugins/Experimental/Avalanche/Source/AvalancheSequence/Public/AvaSequencePlayer.h
+++ b/Engine/Plugins/Experimental/Avalanche/Source/AvalancheSequence/Public/AvaSequencePlayer.h
@@ -95,4 +95,6 @@ private:
 	TSharedPtr<IAvaSequenceController> SequenceController;
 
 	TWeakObjectPtr<ULevel> PlaybackLevelWeak;
+
+	bool bTearDownOnFinished = false;
 };
diff --git a/Engine/Plugins/Experimental/Avalanche/Source/AvalancheSequence/Public/AvaSequenceShared.h b/Engine/Plugins/Experimental/Avalanche/Source/AvalancheSequence/Public/AvaSequenceShared.h
--- a/Engine/Plugins/Experimental/Avalanche/Source/AvalancheSequence/Public/AvaSequenceShared.h
+++ b/Engine/Plugins/Experimental/Avalanche/Source/AvalancheSequence/Public/AvaSequenceShared.h
@@ -130,6 +130,9 @@ struct FAvaSequencePlayAdvancedSettings
 
 	UPROPERTY(EditAnywhere, BlueprintReadWrite, Category = "Motion Design Sequence")
 	bool bRestoreState = false;
+
+	UPROPERTY()
+	bool bTearDownOnFinished = false;
 };
 
 USTRUCT(BlueprintType, DisplayName="Motion Design Sequence Play Parameters")
diff --git a/Engine/Plugins/Experimental/Avalanche/Source/AvalancheSequence/Public/Director/AvaSequenceDirector.h b/Engine/Plugins/Experimental/Avalanche/Source/AvalancheSequence/Public/Director/AvaSequenceDirector.h
--- a/Engine/Plugins/Experimental/Avalanche/Source/AvalancheSequence/Public/Director/AvaSequenceDirector.h
+++ b/Engine/Plugins/Experimental/Avalanche/Source/AvalancheSequence/Public/Director/AvaSequenceDirector.h
@@ -3,6 +3,8 @@
 #pragma once
 
 #include "LevelSequenceDirector.h"
+#include "UObject/WeakInterfacePtr.h"
+#include "UObject/WeakObjectPtrTemplates.h"
 #include "AvaSequenceDirector.generated.h"
 
 class IAvaSequencePlaybackObject;
@@ -40,9 +42,7 @@ private:
 
 	void UpdatePlaybackObject();
 
-	UPROPERTY()
-	TObjectPtr<UAvaSequencePlayer> SequencePlayer;
+	TWeakObjectPtr<UAvaSequencePlayer> SequencePlayerWeak;
 
-	UPROPERTY()
-	TScriptInterface<IAvaSequencePlaybackObject> PlaybackObject;
+	TWeakInterfacePtr<IAvaSequencePlaybackObject> PlaybackObjectInterfaceWeak;
 };
diff --git a/Engine/Plugins/Experimental/Avalanche/Source/AvalancheSequencer/Private/AvaSequencer.cpp b/Engine/Plugins/Experimental/Avalanche/Source/AvalancheSequencer/Private/AvaSequencer.cpp
--- a/Engine/Plugins/Experimental/Avalanche/Source/AvalancheSequencer/Private/AvaSequencer.cpp
+++ b/Engine/Plugins/Experimental/Avalanche/Source/AvalancheSequencer/Private/AvaSequencer.cpp
@@ -1022,6 +1022,7 @@ void FAvaSequencer::PlaySelected_Execute()
 
 	FAvaSequencePlayParams PlaySettings;
 	PlaySettings.AdvancedSettings.bRestoreState = true;
+	PlaySettings.AdvancedSettings.bTearDownOnFinished = true;
 
 	for (const FAvaSequenceItemPtr& Item : SelectedItems)
 	{
