[Motion Design] Fix for a broadcast update infinite loop between client and servers making the system non responsive.

This happens when there are multiple servers on the same channel, as would be the expected case for an nDisplay configuration with multiple render nodes.

Problem:
Both client and server had too many spurious channel configuration update event propagation. In particular, the server would send incomplete channel configuration update while it is still adding the missing outputs causing the client to send the missing outputs again.

Fix:
Server side: This is fixed by blocking the updates until the configuration is fnished. See bBlockChannelStatusUpdate in UpdateChannelOutputConfig.
Client side: The configuration updates are only sent to the server that has a missing output instead of all servers. For a display cluster with many nodes, this should help reduce the event propagation.

diff --git a/Engine/Plugins/Experimental/Avalanche/Source/AvalancheMedia/Private/Playback/AvaPlaybackClient.cpp b/Engine/Plugins/Experimental/Avalanche/Source/AvalancheMedia/Private/Playback/AvaPlaybackClient.cpp
--- a/Engine/Plugins/Experimental/Avalanche/Source/AvalancheMedia/Private/Playback/AvaPlaybackClient.cpp
+++ b/Engine/Plugins/Experimental/Avalanche/Source/AvalancheMedia/Private/Playback/AvaPlaybackClient.cpp
@@ -440,7 +440,7 @@ void FAvaPlaybackClient::RequestPlayableTransitionStop(const FGuid& InTransition
 
 void FAvaPlaybackClient::RequestBroadcast(const FString& InProfile, const FName& InChannel,
 											   const TArray<UMediaOutput*>& InRemoteMediaOutputs,
-											   EAvaBroadcastAction InAction)
+											   EAvaBroadcastAction InAction, const FString& InServerName)
 {
 	FAvaBroadcastRequest* Request = FMessageEndpoint::MakeMessage<FAvaBroadcastRequest>();
 	Request->Profile = InProfile;
@@ -496,16 +496,17 @@ void FAvaPlaybackClient::RequestBroadcast(const FString& InProfile, const FName&
 		}
 	}
 
-	const TArray<FMessageAddress> ServerAddressesForChannel = GetServerAddressesForChannel(InChannel);
+	const TArray<FMessageAddress> ServerAddresses = InServerName.IsEmpty() ?
+		GetServerAddressesForChannel(InChannel) : TArray<FMessageAddress>({ GetServerAddress(InServerName) });
 
 	// Also update the channel settings.
 	if (InAction == EAvaBroadcastAction::Start || InAction == EAvaBroadcastAction::UpdateConfig)
 	{
-		SendBroadcastChannelSettingsUpdate(ServerAddressesForChannel, UAvaBroadcast::Get().GetCurrentProfile().GetChannel(InChannel));
+		SendBroadcastChannelSettingsUpdate(ServerAddresses, UAvaBroadcast::Get().GetCurrentProfile().GetChannel(InChannel));
 	}
 
 	// Send to server(s) that have output for this channel.
-	SendRequest(Request, ServerAddressesForChannel);
+	SendRequest(Request, ServerAddresses);
 }
 
 bool FAvaPlaybackClient::IsMediaOutputRemoteFallback(const UMediaOutput* InMediaOutput)
@@ -1011,7 +1012,7 @@ void FAvaPlaybackClient::HandleBroadcastStatusMessage(const FAvaBroadcastStatus&
 					   TEXT("Playback Server: \"%s\" Channel: \"%s\" is missing outputs. Requesting configuration update."),
 					   *InMessage.ServerName, *InMessage.ChannelName);
 
-				RequestBroadcast(CurrentProfileName, ChannelName, RemoteOutputs, EAvaBroadcastAction::UpdateConfig);
+				RequestBroadcast(CurrentProfileName, ChannelName, RemoteOutputs, EAvaBroadcastAction::UpdateConfig, InMessage.ServerName);
 			}
 
 			// Note: this will broadcast to delegates which may then request states of media outputs.
@@ -1029,7 +1030,7 @@ void FAvaPlaybackClient::HandleBroadcastStatusMessage(const FAvaBroadcastStatus&
 				   *InMessage.ServerName, *InMessage.ChannelName);
 
 			// Request this channel be deleted.
-			RequestBroadcast(CurrentProfileName, ChannelName, {}, EAvaBroadcastAction::DeleteChannel);
+			RequestBroadcast(CurrentProfileName, ChannelName, {}, EAvaBroadcastAction::DeleteChannel, InMessage.ServerName);
 		}
 	}
 
@@ -1057,7 +1058,7 @@ void FAvaPlaybackClient::HandleBroadcastStatusMessage(const FAvaBroadcastStatus&
 				if (ServerInfo.GetBroadcastChannelInfo(Channel->GetChannelName().ToString()) == nullptr)
 				{
 					// The server is missing this channel, we need to update it's config.
-					RequestBroadcast(CurrentProfileName, Channel->GetChannelName(), RemoteOutputs, EAvaBroadcastAction::UpdateConfig);
+					RequestBroadcast(CurrentProfileName, Channel->GetChannelName(), RemoteOutputs, EAvaBroadcastAction::UpdateConfig, InMessage.ServerName);
 				}
 			}
 		}
diff --git a/Engine/Plugins/Experimental/Avalanche/Source/AvalancheMedia/Private/Playback/AvaPlaybackClient.h b/Engine/Plugins/Experimental/Avalanche/Source/AvalancheMedia/Private/Playback/AvaPlaybackClient.h
--- a/Engine/Plugins/Experimental/Avalanche/Source/AvalancheMedia/Private/Playback/AvaPlaybackClient.h
+++ b/Engine/Plugins/Experimental/Avalanche/Source/AvalancheMedia/Private/Playback/AvaPlaybackClient.h
@@ -39,7 +39,7 @@ public:
 	virtual void RequestRemoteControlUpdate(const FGuid& InInstanceId, const FSoftObjectPath& InAssetPath, const FString& InChannelName, const FAvaPlayableRemoteControlValues& InRemoteControlValues) override;
 	virtual void RequestPlayableTransitionStart(const FGuid& InTransitionId, TArray<FGuid>&& InEnterInstanceIds, TArray<FGuid>&& InPlayingInstanceIds, TArray<FGuid>&& InExitInstanceIds, TArray<FAvaPlayableRemoteControlValues>&& InEnterValues, const FName& InChannelName, EAvaPlayableTransitionFlags InTransitionFlags) override;
 	virtual void RequestPlayableTransitionStop(const FGuid& InTransitionId, const FName& InChannelName) override;
-	virtual void RequestBroadcast(const FString& InProfile, const FName& InChannel, const TArray<UMediaOutput*>& InRemoteMediaOutputs, EAvaBroadcastAction InAction) override;
+	virtual void RequestBroadcast(const FString& InProfile, const FName& InChannel, const TArray<UMediaOutput*>& InRemoteMediaOutputs, EAvaBroadcastAction InAction, const FString& InServerName = FString()) override;
 	virtual bool IsMediaOutputRemoteFallback(const UMediaOutput* InMediaOutput) override;
 	virtual EAvaBroadcastIssueSeverity GetMediaOutputIssueSeverity(const FString& InServerName, const FString& InChannelName, const FGuid& InOutputGuid) const override;
 	virtual const TArray<FString>& GetMediaOutputIssueMessages(const FString& InServerName, const FString& InChannelName, const FGuid& InOutputGuid) const override;
diff --git a/Engine/Plugins/Experimental/Avalanche/Source/AvalancheMedia/Private/Playback/AvaPlaybackClientDummy.h b/Engine/Plugins/Experimental/Avalanche/Source/AvalancheMedia/Private/Playback/AvaPlaybackClientDummy.h
--- a/Engine/Plugins/Experimental/Avalanche/Source/AvalancheMedia/Private/Playback/AvaPlaybackClientDummy.h
+++ b/Engine/Plugins/Experimental/Avalanche/Source/AvalancheMedia/Private/Playback/AvaPlaybackClientDummy.h
@@ -33,7 +33,7 @@ public:
 	virtual void RequestRemoteControlUpdate(const FGuid& InInstanceId, const FSoftObjectPath& InAssetPath, const FString& InChannelName, const FAvaPlayableRemoteControlValues& InRemoteControlValues) override {};
 	virtual void RequestPlayableTransitionStart(const FGuid& InTransitionId, TArray<FGuid>&& InEnterInstanceIds, TArray<FGuid>&& InPlayingInstanceIds, TArray<FGuid>&& InExitInstanceIds, TArray<FAvaPlayableRemoteControlValues>&& InEnterValues, const FName& InChannelName, EAvaPlayableTransitionFlags InTransitionFlags) override {}
 	virtual void RequestPlayableTransitionStop(const FGuid& InTransitionId, const FName& InChannelName) override {}
-	virtual void RequestBroadcast(const FString& InProfile, const FName& InChannel, const TArray<UMediaOutput*>& InRemoteMediaOutputs, EAvaBroadcastAction InAction) override {}
+	virtual void RequestBroadcast(const FString& InProfile, const FName& InChannel, const TArray<UMediaOutput*>& InRemoteMediaOutputs, EAvaBroadcastAction InAction, const FString& InServerName) override {}
 	virtual bool IsMediaOutputRemoteFallback(const UMediaOutput* InMediaOutput) override { return false;}
 	virtual EAvaBroadcastIssueSeverity GetMediaOutputIssueSeverity(const FString& InServerName, const FString& InChannelName, const FGuid& InOutputGuid) const override { return EAvaBroadcastIssueSeverity::None;}
 	virtual const TArray<FString>& GetMediaOutputIssueMessages(const FString& InServerName, const FString& InChannelName, const FGuid& InOutputGuid) const override { return EmptyStringArray; }
diff --git a/Engine/Plugins/Experimental/Avalanche/Source/AvalancheMedia/Private/Playback/AvaPlaybackServer.cpp b/Engine/Plugins/Experimental/Avalanche/Source/AvalancheMedia/Private/Playback/AvaPlaybackServer.cpp
--- a/Engine/Plugins/Experimental/Avalanche/Source/AvalancheMedia/Private/Playback/AvaPlaybackServer.cpp
+++ b/Engine/Plugins/Experimental/Avalanche/Source/AvalancheMedia/Private/Playback/AvaPlaybackServer.cpp
@@ -832,7 +832,7 @@ void FAvaPlaybackServer::HandleBroadcastStatusRequest(const FAvaBroadcastStatusR
 	// Make sure all the channel status are refreshed
 	{
 		// Block channel status update while we refresh, we want to send one clean update at the end.
-		TGuardValue BlockChannelStatusUpdate(bBlockChannelStatusUpdate, true);
+		TGuardValue<bool> BlockChannelStatusUpdate(bBlockChannelStatusUpdate, true);
 		for (FAvaBroadcastOutputChannel* Channel : UAvaBroadcast::Get().GetCurrentProfile().GetChannels())
 		{
 			Channel->RefreshState();
@@ -1732,18 +1732,22 @@ bool FAvaPlaybackServer::UpdateChannelOutputConfig(FAvaBroadcastOutputChannel& I
 		if (!NewOutputs.IsEmpty())
 		{
 			{
+				// Both RemoveMediaOutput and AddMediaOutput will broadcast channel events
+				// we don't want those temporary states to propagate to the playback client.
+				TGuardValue<bool> BlockChannelStatusUpdate(bBlockChannelStatusUpdate, true);
+			
 				TArray<UMediaOutput*> MediaOutputs = InChannel.GetMediaOutputs();
 				for (UMediaOutput* MediaOutput : MediaOutputs)
 				{
 					InChannel.RemoveMediaOutput(MediaOutput);
 				}
-			}
 
-			for (int32 Index = 0; Index < NewOutputs.Num(); ++Index)
-			{
-				// Make the device info "local" for this server.
-				NewOutputInfos[Index].ServerName = FAvaBroadcastDeviceProviderProxyManager::LocalServerName;
-				InChannel.AddMediaOutput(NewOutputs[Index].Get(), NewOutputInfos[Index]);
+				for (int32 Index = 0; Index < NewOutputs.Num(); ++Index)
+				{
+					// Make the device info "local" for this server.
+					NewOutputInfos[Index].ServerName = FAvaBroadcastDeviceProviderProxyManager::LocalServerName;
+					InChannel.AddMediaOutput(NewOutputs[Index].Get(), NewOutputInfos[Index]);
+				}
 			}
 
 			// We may not desired refresh state here to avoid spurious states if
diff --git a/Engine/Plugins/Experimental/Avalanche/Source/AvalancheMedia/Public/Playback/IAvaPlaybackClient.h b/Engine/Plugins/Experimental/Avalanche/Source/AvalancheMedia/Public/Playback/IAvaPlaybackClient.h
--- a/Engine/Plugins/Experimental/Avalanche/Source/AvalancheMedia/Public/Playback/IAvaPlaybackClient.h
+++ b/Engine/Plugins/Experimental/Avalanche/Source/AvalancheMedia/Public/Playback/IAvaPlaybackClient.h
@@ -11,7 +11,7 @@ class UMediaOutput;
 struct FAvaPlayableRemoteControlValues;
 struct FAvaPlaybackAnimPlaySettings;
 
-class AVALANCHEMEDIA_API IAvaPlaybackClient
+class IAvaPlaybackClient
 {
 public:
 	IAvaPlaybackClient() = default;
@@ -52,7 +52,7 @@ public:
 	virtual void RequestRemoteControlUpdate(const FGuid& InInstanceId, const FSoftObjectPath& InAssetPath, const FString& InChannelName, const FAvaPlayableRemoteControlValues& InRemoteControlValues) = 0;
 	virtual void RequestPlayableTransitionStart(const FGuid& InTransitionId, TArray<FGuid>&& InEnterInstanceIds, TArray<FGuid>&& InPlayingInstanceIds, TArray<FGuid>&& InExitInstanceIds, TArray<FAvaPlayableRemoteControlValues>&& InEnterValues, const FName& InChannelName, EAvaPlayableTransitionFlags InTransitionFlags) = 0;
 	virtual void RequestPlayableTransitionStop(const FGuid& InTransitionId, const FName& InChannelName) = 0;
-	virtual void RequestBroadcast(const FString& InProfile, const FName& InChannel, const TArray<UMediaOutput*>& InRemoteMediaOutputs, EAvaBroadcastAction InAction) = 0;
+	virtual void RequestBroadcast(const FString& InProfile, const FName& InChannel, const TArray<UMediaOutput*>& InRemoteMediaOutputs, EAvaBroadcastAction InAction, const FString& InServerName = FString()) = 0;
 	/**
 	 * Use the device provider to determine if the media output is remote.
 	 * This function should only be used as fallback if the channel's media output info is not valid.
