MotionDesign : PropertyAnimatorCore
- Fixed animator is not removed from operator stack when the delete button is pressed and animator is selected via outliner chip

diff --git a/Engine/Plugins/Experimental/Avalanche/Source/AvalancheEditor/Private/OperatorStack/Widgets/SAvaOperatorStackTab.cpp b/Engine/Plugins/Experimental/Avalanche/Source/AvalancheEditor/Private/OperatorStack/Widgets/SAvaOperatorStackTab.cpp
--- a/Engine/Plugins/Experimental/Avalanche/Source/AvalancheEditor/Private/OperatorStack/Widgets/SAvaOperatorStackTab.cpp
+++ b/Engine/Plugins/Experimental/Avalanche/Source/AvalancheEditor/Private/OperatorStack/Widgets/SAvaOperatorStackTab.cpp
@@ -6,6 +6,7 @@
 #include "Components/PropertyAnimatorCoreComponent.h"
 #include "Contexts/OperatorStackEditorContext.h"
 #include "DetailView/IAvaDetailsProvider.h"
+#include "Editor.h"
 #include "EditorModeManager.h"
 #include "Items/OperatorStackEditorItem.h"
 #include "Items/OperatorStackEditorObjectItem.h"
@@ -36,7 +37,7 @@ void SAvaOperatorStackTab::Construct(const FArguments& InArgs
 
 	// Property controllers delegates
 	UPropertyAnimatorCoreBase::OnPropertyAnimatorAdded().AddSP(this, &SAvaOperatorStackTab::OnAnimatorUpdated);
-	UPropertyAnimatorCoreBase::OnPropertyAnimatorRemoved().AddSP(this, &SAvaOperatorStackTab::OnAnimatorUpdated);
+	UPropertyAnimatorCoreBase::OnPropertyAnimatorRemoved().AddSP(this, &SAvaOperatorStackTab::OnAnimatorRemoved);
 	UPropertyAnimatorCoreBase::OnPropertyAnimatorRenamed().AddSP(this, &SAvaOperatorStackTab::OnAnimatorUpdated);
 
 	const TSharedPtr<IDetailKeyframeHandler> KeyframeHandler = InProvider->GetDetailsKeyframeHandler();
@@ -115,6 +116,27 @@ void SAvaOperatorStackTab::OnAnimatorUpdated(UPropertyAnimatorCoreComponent* InC
 	}
 }
 
+void SAvaOperatorStackTab::OnAnimatorRemoved(UPropertyAnimatorCoreComponent* InComponent, UPropertyAnimatorCoreBase* InRemovedItem) const
+{
+	if (!GEditor)
+	{
+		return;
+	}
+
+	if (USelection* SelectionSet = GEditor->GetSelectedObjects())
+	{
+		if (SelectionSet->CountSelections(UPropertyAnimatorCoreComponent::StaticClass())
+			|| SelectionSet->CountSelections(UPropertyAnimatorCoreBase::StaticClass()))
+		{
+			RefreshSelection(SelectionSet);
+		}
+		else
+		{
+			OnAnimatorUpdated(InComponent, InRemovedItem);
+		}
+	}
+}
+
 void SAvaOperatorStackTab::RefreshCurrentSelection(const UObject* InObject) const
 {
 	const TSharedPtr<IAvaDetailsProvider> DetailsProvider = DetailsProviderWeak.Pin();
diff --git a/Engine/Plugins/Experimental/Avalanche/Source/AvalancheEditor/Private/OperatorStack/Widgets/SAvaOperatorStackTab.h b/Engine/Plugins/Experimental/Avalanche/Source/AvalancheEditor/Private/OperatorStack/Widgets/SAvaOperatorStackTab.h
--- a/Engine/Plugins/Experimental/Avalanche/Source/AvalancheEditor/Private/OperatorStack/Widgets/SAvaOperatorStackTab.h
+++ b/Engine/Plugins/Experimental/Avalanche/Source/AvalancheEditor/Private/OperatorStack/Widgets/SAvaOperatorStackTab.h
@@ -31,6 +31,7 @@ private:
 	void RefreshSelection(UObject* InSelectionObject) const;
 	void OnModifierUpdated(UActorModifierCoreBase* InUpdatedItem) const;
 	void OnAnimatorUpdated(UPropertyAnimatorCoreComponent* InComponent, UPropertyAnimatorCoreBase* InUpdatedItem) const;
+	void OnAnimatorRemoved(UPropertyAnimatorCoreComponent* InComponent, UPropertyAnimatorCoreBase* InRemovedItem) const;
 	void RefreshCurrentSelection(const UObject* InObject) const;
 
 	TWeakPtr<IAvaDetailsProvider> DetailsProviderWeak;
diff --git a/Engine/Plugins/Experimental/Avalanche/Source/AvalanchePropertyAnimatorEditor/Private/Outliner/AvaPropertyAnimatorEditorOutliner.cpp b/Engine/Plugins/Experimental/Avalanche/Source/AvalanchePropertyAnimatorEditor/Private/Outliner/AvaPropertyAnimatorEditorOutliner.cpp
--- a/Engine/Plugins/Experimental/Avalanche/Source/AvalanchePropertyAnimatorEditor/Private/Outliner/AvaPropertyAnimatorEditorOutliner.cpp
+++ b/Engine/Plugins/Experimental/Avalanche/Source/AvalanchePropertyAnimatorEditor/Private/Outliner/AvaPropertyAnimatorEditorOutliner.cpp
@@ -4,6 +4,7 @@
 
 #include "Animators/PropertyAnimatorCoreBase.h"
 #include "Components/PropertyAnimatorCoreComponent.h"
+#include "IAvaOutliner.h"
 #include "Selection/AvaOutlinerScopedSelection.h"
 #include "Styling/SlateIconFinder.h"
 #include "Subsystems/PropertyAnimatorCoreSubsystem.h"
@@ -15,6 +16,16 @@ FAvaPropertyAnimatorEditorOutliner::FAvaPropertyAnimatorEditorOutliner(IAvaOutli
 	ItemName = FText::FromName(PropertyAnimator->GetAnimatorDisplayName());
 	ItemIcon = FSlateIconFinder::FindIconForClass(UPropertyAnimatorCoreComponent::StaticClass());
 	ItemTooltip = FText::FromName(PropertyAnimator->GetAnimatorOriginalName());
+
+	UPropertyAnimatorCoreBase::OnPropertyAnimatorRemoved().AddRaw(this, &FAvaPropertyAnimatorEditorOutliner::OnAnimatorRemoved);
+}
+
+FAvaPropertyAnimatorEditorOutliner::~FAvaPropertyAnimatorEditorOutliner()
+{
+	if (UObjectInitialized())
+	{
+		UPropertyAnimatorCoreBase::OnPropertyAnimatorRemoved().RemoveAll(this);
+	}
 }
 
 void FAvaPropertyAnimatorEditorOutliner::Select(FAvaOutlinerScopedSelection& InSelection) const
@@ -89,3 +100,14 @@ void FAvaPropertyAnimatorEditorOutliner::SetObject_Impl(UObject* InObject)
 	FAvaOutlinerObject::SetObject_Impl(InObject);
 	PropertyAnimator = Cast<UPropertyAnimatorCoreBase>(InObject);
 }
+
+void FAvaPropertyAnimatorEditorOutliner::OnAnimatorRemoved(UPropertyAnimatorCoreComponent* InComponent, UPropertyAnimatorCoreBase* InAnimator) const
+{
+	if (InAnimator && PropertyAnimator.Get(/** EvenIfPendingKill */true) == InAnimator)
+	{
+		if (const TSharedPtr<IAvaOutliner> OwnerOutliner = GetOwnerOutliner())
+		{
+			OwnerOutliner->RequestRefresh();
+		}
+	}
+}
diff --git a/Engine/Plugins/Experimental/Avalanche/Source/AvalanchePropertyAnimatorEditor/Private/Outliner/AvaPropertyAnimatorEditorOutliner.h b/Engine/Plugins/Experimental/Avalanche/Source/AvalanchePropertyAnimatorEditor/Private/Outliner/AvaPropertyAnimatorEditorOutliner.h
--- a/Engine/Plugins/Experimental/Avalanche/Source/AvalanchePropertyAnimatorEditor/Private/Outliner/AvaPropertyAnimatorEditorOutliner.h
+++ b/Engine/Plugins/Experimental/Avalanche/Source/AvalanchePropertyAnimatorEditor/Private/Outliner/AvaPropertyAnimatorEditorOutliner.h
@@ -7,6 +7,7 @@
 #include "UObject/WeakObjectPtr.h"
 
 class UPropertyAnimatorCoreBase;
+class UPropertyAnimatorCoreComponent;
 
 /** Item representing a property animator item */
 class FAvaPropertyAnimatorEditorOutliner : public FAvaOutlinerObject
@@ -15,6 +16,7 @@ public:
 	UE_AVA_INHERITS_WITH_SUPER(FAvaPropertyAnimatorEditorOutliner, FAvaOutlinerObject);
 
 	FAvaPropertyAnimatorEditorOutliner(IAvaOutliner& InOutliner, UPropertyAnimatorCoreBase* InObject);
+	virtual ~FAvaPropertyAnimatorEditorOutliner() override;
 
 	UPropertyAnimatorCoreBase* GetPropertyAnimator() const
 	{
@@ -38,6 +40,8 @@ protected:
 	virtual void SetObject_Impl(UObject* InObject) override;
 	//~ End FAvaOutlinerObjectItem
 
+	void OnAnimatorRemoved(UPropertyAnimatorCoreComponent* InComponent, UPropertyAnimatorCoreBase* InAnimator) const;
+
 	TWeakObjectPtr<UPropertyAnimatorCoreBase> PropertyAnimator;
 
 	FText ItemName;
