Motion Design: Fix ensure fail on Tick Manager.
This also moves the Playback actor to the runtime module as it was mistakenly placed in the editor module.

diff --git a/Engine/Plugins/Experimental/Avalanche/Source/AvalancheSequence/Private/AvaSequenceActor.cpp b/Engine/Plugins/Experimental/Avalanche/Source/AvalancheSequence/Private/AvaSequenceActor.cpp
new file mode 100644
--- /dev/null
+++ b/Engine/Plugins/Experimental/Avalanche/Source/AvalancheSequence/Private/AvaSequenceActor.cpp
@@ -0,0 +1,60 @@
+// Copyright Epic Games, Inc. All Rights Reserved.
+
+#include "AvaSequenceActor.h"
+#include "AvaSequence.h"
+#include "AvaSequencePlayer.h"
+#include "AvaSequenceSubsystem.h"
+
+AAvaSequenceActor::AAvaSequenceActor(const FObjectInitializer& InObjectInitializer)
+	: Super(InObjectInitializer.SetDefaultSubobjectClass(TEXT("AnimationPlayer"), UAvaSequencePlayer::StaticClass()))
+{
+}
+
+void AAvaSequenceActor::Initialize(UAvaSequence* InSequence)
+{
+	if (ensure(!GetSequencePlayer()->IsPlaying()))
+	{
+		LevelSequenceAsset = InSequence;
+		InitSequencePlayer(InSequence);
+	}
+}
+
+void AAvaSequenceActor::PostInitializeComponents()
+{
+	Super::PostInitializeComponents();
+	InitSequencePlayer(Cast<UAvaSequence>(LevelSequenceAsset));
+}
+
+void AAvaSequenceActor::InitSequencePlayer(UAvaSequence* InSequence)
+{
+	if (!InSequence)
+	{
+		return;
+	}
+
+	ULevel* Level = GetLevel();
+	if (!Level)
+	{
+		return;
+	}
+
+	UAvaSequencePlayer* Player = Cast<UAvaSequencePlayer>(GetSequencePlayer());
+	if (!Player)
+	{
+		return;
+	}
+
+	UAvaSequenceSubsystem* SequenceSubsystem = UAvaSequenceSubsystem::Get(this);
+	if (!SequenceSubsystem)
+	{
+		return;
+	}
+
+	IAvaSequencePlaybackObject* PlaybackObject = SequenceSubsystem->FindPlaybackObject(Level);
+	if (!PlaybackObject)
+	{
+		return;
+	}
+
+	Player->InitSequence(InSequence, PlaybackObject, Level, CameraSettings);
+}
diff --git a/Engine/Plugins/Experimental/Avalanche/Source/AvalancheSequence/Private/AvaSequencePlaybackActor.cpp b/Engine/Plugins/Experimental/Avalanche/Source/AvalancheSequence/Private/AvaSequencePlaybackActor.cpp
--- a/Engine/Plugins/Experimental/Avalanche/Source/AvalancheSequence/Private/AvaSequencePlaybackActor.cpp
+++ b/Engine/Plugins/Experimental/Avalanche/Source/AvalancheSequence/Private/AvaSequencePlaybackActor.cpp
@@ -17,7 +17,9 @@ AAvaSequencePlaybackActor::AAvaSequencePlaybackActor()
 	PrimaryActorTick.TickInterval = 0.f;
 	PrimaryActorTick.bTickEvenWhenPaused = true;
 
-	UAvaSequencePlayer::OnSequenceFinished().AddUObject(this, &AAvaSequencePlaybackActor::OnSequenceFinished);
+	OnSequenceFinishedDelegate = UAvaSequencePlayer::OnSequenceFinished().AddUObject(this, &AAvaSequencePlaybackActor::OnSequenceFinished);
+
+	OnWorldCleanupDelegate = FWorldDelegates::OnWorldCleanup.AddUObject(this, &AAvaSequencePlaybackActor::OnWorldCleanup);
 }
 
 void AAvaSequencePlaybackActor::SetSequenceProvider(IAvaSequenceProvider& InSequenceProvider)
@@ -486,6 +488,26 @@ void AAvaSequencePlaybackActor::UnregisterPlaybackObject()
 	CleanupPlayers();
 }
 
+void AAvaSequencePlaybackActor::OnWorldCleanup(UWorld* InWorld, bool bInSessionEnded, bool bInCleanupResources)
+{
+	if (bInCleanupResources && GetWorld() == InWorld)
+	{
+		// Ensure the players tear down to unregister themselves from Tick Manager holding them as external references
+		CleanupPlayers();
+	}
+}
+
+void AAvaSequencePlaybackActor::BeginDestroy()
+{
+	Super::BeginDestroy();
+
+	UAvaSequencePlayer::OnSequenceFinished().Remove(OnSequenceFinishedDelegate);
+	OnSequenceFinishedDelegate.Reset();
+
+	FWorldDelegates::OnWorldCleanup.Remove(OnWorldCleanupDelegate);
+	OnWorldCleanupDelegate.Reset();
+}
+
 void AAvaSequencePlaybackActor::EndPlay(const EEndPlayReason::Type InEndPlayReason)
 {
 	Super::EndPlay(InEndPlayReason);
diff --git a/Engine/Plugins/Experimental/Avalanche/Source/AvalancheSequence/Private/AvaSequencePlaybackActor.h b/Engine/Plugins/Experimental/Avalanche/Source/AvalancheSequence/Private/AvaSequencePlaybackActor.h
--- a/Engine/Plugins/Experimental/Avalanche/Source/AvalancheSequence/Private/AvaSequencePlaybackActor.h
+++ b/Engine/Plugins/Experimental/Avalanche/Source/AvalancheSequence/Private/AvaSequencePlaybackActor.h
@@ -54,6 +54,7 @@ protected:
 	//~ End IAvaPlaybackObject
 
 	//~ Begin AActor
+	virtual void BeginDestroy() override;
 	virtual void EndPlay(const EEndPlayReason::Type InEndPlayReason) override;
 #if WITH_EDITOR
 	virtual bool IsSelectable() const override { return false; }
@@ -69,6 +70,9 @@ protected:
 	/** Unregisters this Playback Object from the World's Sequence Tick Manager */
 	void UnregisterPlaybackObject();
 
+	/** Called when on FWorldDelegates::OnWorldCleanup */
+	void OnWorldCleanup(UWorld* InWorld, bool bInSessionEnded, bool bInCleanupResources);
+
 	void OnSequenceFinished(UAvaSequencePlayer* InPlayer, UAvaSequence* InSequence);
 
 private:
@@ -102,5 +106,8 @@ private:
 	/** Called when a Camera Cut occurs. */
 	FOnCameraCut OnCameraCut;
 
+	FDelegateHandle OnSequenceFinishedDelegate;
+	FDelegateHandle OnWorldCleanupDelegate;
+
 	bool bStoppingAllSequences = false;
 };
diff --git a/Engine/Plugins/Experimental/Avalanche/Source/AvalancheSequence/Private/AvaSequencePlayer.cpp b/Engine/Plugins/Experimental/Avalanche/Source/AvalancheSequence/Private/AvaSequencePlayer.cpp
--- a/Engine/Plugins/Experimental/Avalanche/Source/AvalancheSequence/Private/AvaSequencePlayer.cpp
+++ b/Engine/Plugins/Experimental/Avalanche/Source/AvalancheSequence/Private/AvaSequencePlayer.cpp
@@ -86,8 +86,6 @@ void UAvaSequencePlayer::SetPlaySettings(const FAvaSequencePlayParams& InPlaySet
 
 	SetTimeRange(StartTimeSeconds, DurationSeconds);
 
-	bTearDownOnFinished = InPlaySettings.AdvancedSettings.bTearDownOnFinished;
-
 	PlaybackSettings = FMovieSceneSequencePlaybackSettings();
 	PlaybackSettings.PlayRate = InPlaySettings.AdvancedSettings.PlaybackSpeed;
 	PlaybackSettings.LoopCount.Value = InPlaySettings.AdvancedSettings.LoopCount;
@@ -184,6 +182,10 @@ void UAvaSequencePlayer::OnStopped()
 
 	// At the moment, Stop means to completely finish
 	NotifySequenceFinished();
+
+	// Defer cleanup as there is an action flush that assumes the Tick Manager is still alive after this Stop callback
+	// see UMovieSceneSequencePlayer::RunLatentActions
+	QueueLatentAction(FMovieSceneSequenceLatentActionDelegate::CreateUObject(this, &UAvaSequencePlayer::Cleanup));
 }
 
 void UAvaSequencePlayer::TickFromSequenceTickManager(float InDeltaSeconds, FMovieSceneEntitySystemRunner* InRunner)
@@ -228,9 +230,4 @@ void UAvaSequencePlayer::NotifySequencePaused()
 void UAvaSequencePlayer::NotifySequenceFinished()
 {
 	OnSequenceFinishedDelegate.Broadcast(this, GetAvaSequence());
-
-	if (bTearDownOnFinished)
-	{
-		TearDown();
-	}
 }
diff --git a/Engine/Plugins/Experimental/Avalanche/Source/AvalancheSequence/Public/AvaSequenceActor.h b/Engine/Plugins/Experimental/Avalanche/Source/AvalancheSequence/Public/AvaSequenceActor.h
new file mode 100644
--- /dev/null
+++ b/Engine/Plugins/Experimental/Avalanche/Source/AvalancheSequence/Public/AvaSequenceActor.h
@@ -0,0 +1,28 @@
+// Copyright Epic Games, Inc. All Rights Reserved.
+
+#pragma once
+
+#include "LevelSequenceActor.h"
+#include "AvaSequenceActor.generated.h"
+
+class UAvaSequence;
+class UAvaSequencePlayer;
+
+UCLASS(MinimalAPI)
+class AAvaSequenceActor : public ALevelSequenceActor
+{
+	GENERATED_BODY()
+
+public:
+	AAvaSequenceActor(const FObjectInitializer& InObjectInitializer);
+
+	AVALANCHESEQUENCE_API void Initialize(UAvaSequence* InSequence);
+
+protected:
+	//~ Begin AActor
+	virtual void PostInitializeComponents() override;
+	//~ End AActor
+
+private:
+	void InitSequencePlayer(UAvaSequence* InSequence);
+};
diff --git a/Engine/Plugins/Experimental/Avalanche/Source/AvalancheSequence/Public/AvaSequencePlayer.h b/Engine/Plugins/Experimental/Avalanche/Source/AvalancheSequence/Public/AvaSequencePlayer.h
--- a/Engine/Plugins/Experimental/Avalanche/Source/AvalancheSequence/Public/AvaSequencePlayer.h
+++ b/Engine/Plugins/Experimental/Avalanche/Source/AvalancheSequence/Public/AvaSequencePlayer.h
@@ -95,6 +95,4 @@ private:
 	TSharedPtr<IAvaSequenceController> SequenceController;
 
 	TWeakObjectPtr<ULevel> PlaybackLevelWeak;
-
-	bool bTearDownOnFinished = false;
 };
diff --git a/Engine/Plugins/Experimental/Avalanche/Source/AvalancheSequence/Public/AvaSequenceShared.h b/Engine/Plugins/Experimental/Avalanche/Source/AvalancheSequence/Public/AvaSequenceShared.h
--- a/Engine/Plugins/Experimental/Avalanche/Source/AvalancheSequence/Public/AvaSequenceShared.h
+++ b/Engine/Plugins/Experimental/Avalanche/Source/AvalancheSequence/Public/AvaSequenceShared.h
@@ -130,9 +130,6 @@ struct FAvaSequencePlayAdvancedSettings
 
 	UPROPERTY(EditAnywhere, BlueprintReadWrite, Category = "Motion Design Sequence")
 	bool bRestoreState = false;
-
-	UPROPERTY()
-	bool bTearDownOnFinished = false;
 };
 
 USTRUCT(BlueprintType, DisplayName="Motion Design Sequence Play Parameters")
diff --git a/Engine/Plugins/Experimental/Avalanche/Source/AvalancheSequencer/Private/AvaSequencer.cpp b/Engine/Plugins/Experimental/Avalanche/Source/AvalancheSequencer/Private/AvaSequencer.cpp
--- a/Engine/Plugins/Experimental/Avalanche/Source/AvalancheSequencer/Private/AvaSequencer.cpp
+++ b/Engine/Plugins/Experimental/Avalanche/Source/AvalancheSequencer/Private/AvaSequencer.cpp
@@ -2,6 +2,7 @@
 
 #include "AvaSequencer.h"
 #include "AvaSequence.h"
+#include "AvaSequenceActor.h"
 #include "AvaSequencePlaybackObject.h"
 #include "AvaSequencePlayer.h"
 #include "AvaSequencerArgs.h"
@@ -36,7 +37,6 @@
 #include "MovieScene.h"
 #include "Playback/AvaSequencerCleanView.h"
 #include "Playback/AvaSequencerController.h"
-#include "Player/AvaSequenceActor.h"
 #include "PropertyHandle.h"
 #include "ScopedTransaction.h"
 #include "Selection/AvaEditorSelection.h"
@@ -1022,7 +1022,6 @@ void FAvaSequencer::PlaySelected_Execute()
 
 	FAvaSequencePlayParams PlaySettings;
 	PlaySettings.AdvancedSettings.bRestoreState = true;
-	PlaySettings.AdvancedSettings.bTearDownOnFinished = true;
 
 	for (const FAvaSequenceItemPtr& Item : SelectedItems)
 	{
diff --git a/Engine/Plugins/Experimental/Avalanche/Source/AvalancheSequencer/Private/Player/AvaSequenceActor.cpp b/Engine/Plugins/Experimental/Avalanche/Source/AvalancheSequencer/Private/Player/AvaSequenceActor.cpp
deleted file mode 100644
diff --git a/Engine/Plugins/Experimental/Avalanche/Source/AvalancheSequencer/Private/Player/AvaSequenceActor.h b/Engine/Plugins/Experimental/Avalanche/Source/AvalancheSequencer/Private/Player/AvaSequenceActor.h
deleted file mode 100644
