Motion Design: Fix transition logic rundown playback crash when setting gc.CollectGarbageEveryFrame 1 and running transition logic containing blueprint tasks

diff --git a/Engine/Plugins/Experimental/Avalanche/Source/AvalancheTransition/Private/Behavior/AvaTransitionBehaviorInstance.cpp b/Engine/Plugins/Experimental/Avalanche/Source/AvalancheTransition/Private/Behavior/AvaTransitionBehaviorInstance.cpp
--- a/Engine/Plugins/Experimental/Avalanche/Source/AvalancheTransition/Private/Behavior/AvaTransitionBehaviorInstance.cpp
+++ b/Engine/Plugins/Experimental/Avalanche/Source/AvalancheTransition/Private/Behavior/AvaTransitionBehaviorInstance.cpp
@@ -158,6 +158,11 @@ void FAvaTransitionBehaviorInstance::SetLogContext(const FString& InContext)
 	LogContext = InContext;
 }
 
+void FAvaTransitionBehaviorInstance::AddReferencedObjects(FReferenceCollector& InCollector)
+{
+	InstanceData.AddStructReferencedObjects(InCollector);
+}
+
 void FAvaTransitionBehaviorInstance::ConditionallyStop()
 {
 	// if the run status was set in Start/Tick with anything other than running (e.g. Succeeded/Failed/Stop) it should call stop
diff --git a/Engine/Plugins/Experimental/Avalanche/Source/AvalancheTransition/Private/Execution/AvaTransitionExecutor.cpp b/Engine/Plugins/Experimental/Avalanche/Source/AvalancheTransition/Private/Execution/AvaTransitionExecutor.cpp
--- a/Engine/Plugins/Experimental/Avalanche/Source/AvalancheTransition/Private/Execution/AvaTransitionExecutor.cpp
+++ b/Engine/Plugins/Experimental/Avalanche/Source/AvalancheTransition/Private/Execution/AvaTransitionExecutor.cpp
@@ -233,6 +233,20 @@ bool FAvaTransitionExecutor::IsTickable() const
 	return IsRunning();
 }
 
+FString FAvaTransitionExecutor::GetReferencerName() const
+{
+	return TEXT("FAvaTransitionExecutor");
+}
+
+void FAvaTransitionExecutor::AddReferencedObjects(FReferenceCollector& InCollector)
+{
+	ForEachInstance(
+		[&InCollector](FAvaTransitionBehaviorInstance& InInstance)
+		{
+			InInstance.AddReferencedObjects(InCollector);
+		});
+}
+
 void FAvaTransitionExecutor::ForEachInstance(TFunctionRef<void(FAvaTransitionBehaviorInstance&)> InFunc)
 {
 	for (FAvaTransitionBehaviorInstance& Instance : Instances)
diff --git a/Engine/Plugins/Experimental/Avalanche/Source/AvalancheTransition/Private/Execution/AvaTransitionExecutor.h b/Engine/Plugins/Experimental/Avalanche/Source/AvalancheTransition/Private/Execution/AvaTransitionExecutor.h
--- a/Engine/Plugins/Experimental/Avalanche/Source/AvalancheTransition/Private/Execution/AvaTransitionExecutor.h
+++ b/Engine/Plugins/Experimental/Avalanche/Source/AvalancheTransition/Private/Execution/AvaTransitionExecutor.h
@@ -9,13 +9,14 @@
 #include "Delegates/Delegate.h"
 #include "Execution/IAvaTransitionExecutor.h"
 #include "Tickable.h"
+#include "UObject/GCObject.h"
 #include "UObject/WeakInterfacePtr.h"
 
 class FAvaTransitionExecutorBuilder;
 class IAvaTransitionBehavior;
 
 /** Base Implementation of an Executor dealing with multiple behaviors going out (Exit Instances) and multiple behaviors going in (Enter Instances) */
-class FAvaTransitionExecutor : public IAvaTransitionExecutor, public FTickableGameObject
+class FAvaTransitionExecutor : public IAvaTransitionExecutor, public FTickableGameObject, public FGCObject
 {
 public:
 	explicit FAvaTransitionExecutor(FAvaTransitionExecutorBuilder& InBuilder);
@@ -42,6 +43,11 @@ protected:
 	virtual ETickableTickType GetTickableTickType() const override { return ETickableTickType::Conditional; }
 	//~ End FTickableGameObject
 
+	//~ Begin FGCObject
+	virtual FString GetReferencerName() const override;
+	virtual void AddReferencedObjects(FReferenceCollector& InCollector) override;
+	//~ End FGCObject
+
 private:
 	void ForEachInstance(TFunctionRef<void(FAvaTransitionBehaviorInstance&)> InFunc);
 	void ForEachInstance(TFunctionRef<void(const FAvaTransitionBehaviorInstance&)> InFunc) const;
diff --git a/Engine/Plugins/Experimental/Avalanche/Source/AvalancheTransition/Public/Behavior/AvaTransitionBehaviorInstance.h b/Engine/Plugins/Experimental/Avalanche/Source/AvalancheTransition/Public/Behavior/AvaTransitionBehaviorInstance.h
--- a/Engine/Plugins/Experimental/Avalanche/Source/AvalancheTransition/Public/Behavior/AvaTransitionBehaviorInstance.h
+++ b/Engine/Plugins/Experimental/Avalanche/Source/AvalancheTransition/Public/Behavior/AvaTransitionBehaviorInstance.h
@@ -9,6 +9,7 @@
 #include "UObject/WeakInterfacePtr.h"
 
 class FAvaTransitionExecutorBuilder;
+class FReferenceCollector;
 class IAvaTransitionBehavior;
 class UAvaTransitionTree;
 struct FAvaTransitionExecutionContext;
@@ -67,6 +68,8 @@ struct FAvaTransitionBehaviorInstance
 
 	void SetLogContext(const FString& InContext);
 
+	void AddReferencedObjects(FReferenceCollector& InCollector);
+
 private:
 	/** Stops Execution if the Tree is no longer running */
 	void ConditionallyStop();
