Motion Design: Fix issue with Get Changed Controllers not returning the same results after calling an RC Update Controllers that would cause graphics to not behave as expected

diff --git a/Engine/Plugins/Experimental/Avalanche/Source/Avalanche/Private/Transition/AvaTransitionRCLibrary.cpp b/Engine/Plugins/Experimental/Avalanche/Source/Avalanche/Private/Transition/AvaTransitionRCLibrary.cpp
--- a/Engine/Plugins/Experimental/Avalanche/Source/Avalanche/Private/Transition/AvaTransitionRCLibrary.cpp
+++ b/Engine/Plugins/Experimental/Avalanche/Source/Avalanche/Private/Transition/AvaTransitionRCLibrary.cpp
@@ -82,7 +82,9 @@ bool UAvaTransitionRCLibrary::CompareRCControllerValues(const FAvaTransitionCont
 	{
 		check(BehaviorInstance);
 
-		const FAvaTransitionScene* OtherTransitionScene = BehaviorInstance->GetTransitionContext().GetTransitionScene();
+		const FAvaTransitionContext& OtherTransitionContext = BehaviorInstance->GetTransitionContext();
+
+		const FAvaTransitionScene* OtherTransitionScene = OtherTransitionContext.GetTransitionScene();
 		if (!OtherTransitionScene)
 		{
 			continue;
@@ -92,8 +94,8 @@ bool UAvaTransitionRCLibrary::CompareRCControllerValues(const FAvaTransitionCont
 		if (RCTransitionExtension)
 		{
 			Result = RCTransitionExtension->CompareControllers(Controller->Id
-				, *TransitionScene
-				, *OtherTransitionScene);
+				, InTransitionContext
+				, OtherTransitionContext);
 		}
 		else if (URCVirtualPropertyBase* OtherController = UE::Ava::Private::GetController(InControllerId, *SceneSubsystem, OtherTransitionScene->GetLevel()))
 		{
diff --git a/Engine/Plugins/Experimental/Avalanche/Source/Avalanche/Public/Transition/Extensions/IAvaTransitionRCExtension.h b/Engine/Plugins/Experimental/Avalanche/Source/Avalanche/Public/Transition/Extensions/IAvaTransitionRCExtension.h
--- a/Engine/Plugins/Experimental/Avalanche/Source/Avalanche/Public/Transition/Extensions/IAvaTransitionRCExtension.h
+++ b/Engine/Plugins/Experimental/Avalanche/Source/Avalanche/Public/Transition/Extensions/IAvaTransitionRCExtension.h
@@ -2,10 +2,11 @@
 
 #pragma once
 
+#include "AvaTransitionEnums.h"
 #include "HAL/Platform.h"
 #include "IAvaTransitionExtension.h"
 
-enum class EAvaTransitionComparisonResult : uint8;
+struct FAvaTransitionContext;
 struct FAvaTransitionScene;
 struct FGuid;
 
@@ -14,7 +15,22 @@ class IAvaRCTransitionExtension : public IAvaTransitionExtension
 public:
 	static constexpr const TCHAR* ExtensionIdentifier = TEXT("IAvaRCTransitionExtension");
 
+	UE_DEPRECATED(5.5, "Please use CompareControllers that takes in FAvaTransitionContext instead")
 	virtual EAvaTransitionComparisonResult CompareControllers(const FGuid& InControllerId
 		, const FAvaTransitionScene& InMyScene
-		, const FAvaTransitionScene& InOtherScene) const = 0;
+		, const FAvaTransitionScene& InOtherScene) const
+	{
+		return EAvaTransitionComparisonResult::None;
+	}
+
+	/**
+	 * Compares the values of an RC controller with the given ControllerId between two transition contexts
+	 * @param InControllerId the Id of the controller to compare
+	 * @param InMyContext the first transition context to use
+	 * @param InOtherContext the second transition context to use
+	 * @return the comparison result whether same, different, or none (none meaning that the values could not be properly compared)
+	 */
+	virtual EAvaTransitionComparisonResult CompareControllers(const FGuid& InControllerId
+		, const FAvaTransitionContext& InMyContext
+		, const FAvaTransitionContext& InOtherContext) const = 0;
 };
diff --git a/Engine/Plugins/Experimental/Avalanche/Source/AvalancheMedia/Private/Playable/AvaPlayableLibrary.cpp b/Engine/Plugins/Experimental/Avalanche/Source/AvalancheMedia/Private/Playable/AvaPlayableLibrary.cpp
--- a/Engine/Plugins/Experimental/Avalanche/Source/AvalancheMedia/Private/Playable/AvaPlayableLibrary.cpp
+++ b/Engine/Plugins/Experimental/Avalanche/Source/AvalancheMedia/Private/Playable/AvaPlayableLibrary.cpp
@@ -88,7 +88,10 @@ bool UAvaPlayableLibrary::UpdatePlayableRemoteControlValues(const UObject* InWor
 	{
 		if (UAvaPlayableTransition* Transition = GetPlayableTransition(Playable))
 		{
-			if (const TSharedPtr<FAvaPlayableRemoteControlValues> RemoteControlValues = Transition->GetValuesForPlayable(Playable))
+			// Assume that if the RC Values need to be updated, it's because of an Enter Playable not yet having its Update RC called
+			constexpr bool bIsEnterPlayable = true;
+
+			if (const TSharedPtr<FAvaPlayableRemoteControlValues> RemoteControlValues = Transition->GetValuesForPlayable(Playable, bIsEnterPlayable))
 			{
 				if (Playable->UpdateRemoteControlCommand(RemoteControlValues.ToSharedRef()) == EAvaPlayableCommandResult::Executed)
 				{
diff --git a/Engine/Plugins/Experimental/Avalanche/Source/AvalancheMedia/Private/Playable/Transition/AvaPlayableLocalTransition.cpp b/Engine/Plugins/Experimental/Avalanche/Source/AvalancheMedia/Private/Playable/Transition/AvaPlayableLocalTransition.cpp
--- a/Engine/Plugins/Experimental/Avalanche/Source/AvalancheMedia/Private/Playable/Transition/AvaPlayableLocalTransition.cpp
+++ b/Engine/Plugins/Experimental/Avalanche/Source/AvalancheMedia/Private/Playable/Transition/AvaPlayableLocalTransition.cpp
@@ -277,6 +277,20 @@ bool UAvaPlayableLocalTransition::Start()
 		return false;
 	}
 
+	// Store the Latest Remote Control Values for Playables that are not Enter Playables
+	// as the Enter Playable ones are stored in EnterPlayableValues
+	OtherPlayableValues.Empty(PlayingPlayables.Num() + ExitPlayables.Num());
+
+	for (UAvaPlayable* Playable : PlayingPlayables)
+	{
+		OtherPlayableValues.Add(Playable, Playable->GetLatestRemoteControlValues());
+	}
+
+	for (UAvaPlayable* Playable : ExitPlayables)
+	{
+		OtherPlayableValues.Add(Playable, Playable->GetLatestRemoteControlValues());
+	}
+
 	int32 ArrayIndex = 0;
 	for (UAvaPlayable* Playable : EnterPlayables)
 	{
diff --git a/Engine/Plugins/Experimental/Avalanche/Source/AvalancheMedia/Private/Playable/Transition/AvaPlayableTransition.cpp b/Engine/Plugins/Experimental/Avalanche/Source/AvalancheMedia/Private/Playable/Transition/AvaPlayableTransition.cpp
--- a/Engine/Plugins/Experimental/Avalanche/Source/AvalancheMedia/Private/Playable/Transition/AvaPlayableTransition.cpp
+++ b/Engine/Plugins/Experimental/Avalanche/Source/AvalancheMedia/Private/Playable/Transition/AvaPlayableTransition.cpp
@@ -94,18 +94,26 @@ void UAvaPlayableTransition::SetEnterPlayableValues(TArray<TSharedPtr<FAvaPlayab
 	EnterPlayableValues = MoveTemp(InPlayableValues);
 }
 
-TSharedPtr<FAvaPlayableRemoteControlValues> UAvaPlayableTransition::GetValuesForPlayable(const UAvaPlayable* InPlayable)
+TSharedPtr<FAvaPlayableRemoteControlValues> UAvaPlayableTransition::GetValuesForPlayable(const UAvaPlayable* InPlayable, bool bInIsEnterPlayable)
 {
-	int32 ArrayIndex = 0;
-	for (const TWeakObjectPtr<UAvaPlayable> PlayableWeak : EnterPlayablesWeak)
+	if (bInIsEnterPlayable)
 	{
-		const UAvaPlayable* Playable = PlayableWeak.Get();
-		if (Playable == InPlayable && EnterPlayableValues.IsValidIndex(ArrayIndex))
+		int32 ArrayIndex = 0;
+		for (const TWeakObjectPtr<UAvaPlayable> PlayableWeak : EnterPlayablesWeak)
 		{
-			return EnterPlayableValues[ArrayIndex];
+			const UAvaPlayable* Playable = PlayableWeak.Get();
+			if (Playable == InPlayable && EnterPlayableValues.IsValidIndex(ArrayIndex))
+			{
+				return EnterPlayableValues[ArrayIndex];
+			}
+			++ArrayIndex;
 		}
-		++ArrayIndex;
 	}
+	else if (TSharedPtr<FAvaPlayableRemoteControlValues>* FoundValues = OtherPlayableValues.Find(InPlayable))
+	{
+		return *FoundValues;
+	}
+
 	return nullptr;
 }
 
diff --git a/Engine/Plugins/Experimental/Avalanche/Source/AvalancheMedia/Private/Playable/Transition/AvaPlayableTransitionScene.cpp b/Engine/Plugins/Experimental/Avalanche/Source/AvalancheMedia/Private/Playable/Transition/AvaPlayableTransitionScene.cpp
--- a/Engine/Plugins/Experimental/Avalanche/Source/AvalancheMedia/Private/Playable/Transition/AvaPlayableTransitionScene.cpp
+++ b/Engine/Plugins/Experimental/Avalanche/Source/AvalancheMedia/Private/Playable/Transition/AvaPlayableTransitionScene.cpp
@@ -2,6 +2,7 @@
 
 #include "Playable/Transition/AvaPlayableTransitionScene.h"
 
+#include "AvaTransitionContext.h"
 #include "IAvaSceneInterface.h"
 #include "Playable/AvaPlayable.h"
 #include "Playable/Transition/AvaPlayableTransition.h"
@@ -12,34 +13,46 @@ namespace UE::AvaMedia::Private
 	/** Controllers aren't applied to the Preset, instead, this compares the latest remote control values for a given playable */
 	class FAvaRCTransitionPlayableExtension : public IAvaRCTransitionExtension
 	{
-		virtual EAvaTransitionComparisonResult CompareControllers(const FGuid& InControllerId, const FAvaTransitionScene& InMyScene, const FAvaTransitionScene& InOtherScene) const override
+		virtual EAvaTransitionComparisonResult CompareControllers(const FGuid& InControllerId, const FAvaTransitionContext& InMyContext, const FAvaTransitionContext& InOtherContext) const override
 		{
-			const UAvaPlayable* MyPlayable = InMyScene.GetDataView().GetPtr<UAvaPlayable>();
-			const UAvaPlayable* OtherPlayable = InOtherScene.GetDataView().GetPtr<UAvaPlayable>();
+			const FAvaTransitionScene* MyScene = InMyContext.GetTransitionScene();
+			const FAvaTransitionScene* OtherScene = InOtherContext.GetTransitionScene();
+
+			if (!MyScene || !OtherScene)
+			{
+				return EAvaTransitionComparisonResult::None;
+			}
+
+			const UAvaPlayable* MyPlayable = MyScene->GetDataView().GetPtr<UAvaPlayable>();
+			const UAvaPlayable* OtherPlayable = OtherScene->GetDataView().GetPtr<UAvaPlayable>();
+
 			if (!MyPlayable || !OtherPlayable)
 			{
 				return EAvaTransitionComparisonResult::None;
 			}
 
-			const FAvaPlayableRemoteControlValue* MyValue = MyPlayable->GetLatestRemoteControlValues().ControllerValues.Find(InControllerId);
-			const FAvaPlayableRemoteControlValue* OtherValue = nullptr;
+			UAvaPlayableTransition* MyTransition = static_cast<const FAvaPlayableTransitionScene*>(MyScene)->PlayableTransitionWeak.Get();
+			UAvaPlayableTransition* OtherTransition = static_cast<const FAvaPlayableTransitionScene*>(OtherScene)->PlayableTransitionWeak.Get();
 
-			if (OtherPlayable != MyPlayable)
+			if (!MyTransition || !OtherTransition)
 			{
-				OtherValue = OtherPlayable->GetLatestRemoteControlValues().ControllerValues.Find(InControllerId);
+				return EAvaTransitionComparisonResult::None;
 			}
-			else
+
+			const bool bIsMyPlayableEnter = InMyContext.GetTransitionType() == EAvaTransitionType::In;
+			const bool bIsOtherPlayableEnter = InOtherContext.GetTransitionType() == EAvaTransitionType::In;
+
+			TSharedPtr<const FAvaPlayableRemoteControlValues> MyValues = MyTransition->GetValuesForPlayable(MyPlayable, bIsMyPlayableEnter);
+			TSharedPtr<const FAvaPlayableRemoteControlValues> OtherValues = OtherTransition->GetValuesForPlayable(OtherPlayable, bIsOtherPlayableEnter);
+
+			if (!MyValues.IsValid() || !OtherValues.IsValid())
 			{
-				// In case of scene reuse, compare against RC values incoming in the transition.
-				if (UAvaPlayableTransition* Transition = static_cast<const FAvaPlayableTransitionScene&>(InOtherScene).PlayableTransitionWeak.Get())
-				{
-					if (const TSharedPtr<FAvaPlayableRemoteControlValues> OtherValues = Transition->GetValuesForPlayable(OtherPlayable))
-					{
-						OtherValue = OtherValues->ControllerValues.Find(InControllerId);
-					}
-				}
+				return EAvaTransitionComparisonResult::None;
 			}
 
+			const FAvaPlayableRemoteControlValue* MyValue = MyValues->ControllerValues.Find(InControllerId);
+			const FAvaPlayableRemoteControlValue* OtherValue = OtherValues->ControllerValues.Find(InControllerId);
+
 			if (!MyValue || !OtherValue)
 			{
 				return EAvaTransitionComparisonResult::None;
diff --git a/Engine/Plugins/Experimental/Avalanche/Source/AvalancheMedia/Public/Playable/AvaPlayable.h b/Engine/Plugins/Experimental/Avalanche/Source/AvalancheMedia/Public/Playable/AvaPlayable.h
--- a/Engine/Plugins/Experimental/Avalanche/Source/AvalancheMedia/Public/Playable/AvaPlayable.h
+++ b/Engine/Plugins/Experimental/Avalanche/Source/AvalancheMedia/Public/Playable/AvaPlayable.h
@@ -147,9 +147,9 @@ public:
 	virtual void SetUserData(const FString& InUserData) { UserData = InUserData; }
 	const FString& GetUserData() const { return UserData; }
 
-	const FAvaPlayableRemoteControlValues& GetLatestRemoteControlValues() const
+	TSharedPtr<FAvaPlayableRemoteControlValues> GetLatestRemoteControlValues() const
 	{
-		return LatestRemoteControlValues.IsValid() ? *LatestRemoteControlValues : FAvaPlayableRemoteControlValues::GetDefaultEmpty();
+		return LatestRemoteControlValues;
 	}
 
 protected:
diff --git a/Engine/Plugins/Experimental/Avalanche/Source/AvalancheMedia/Public/Playable/Transition/AvaPlayableTransition.h b/Engine/Plugins/Experimental/Avalanche/Source/AvalancheMedia/Public/Playable/Transition/AvaPlayableTransition.h
--- a/Engine/Plugins/Experimental/Avalanche/Source/AvalancheMedia/Public/Playable/Transition/AvaPlayableTransition.h
+++ b/Engine/Plugins/Experimental/Avalanche/Source/AvalancheMedia/Public/Playable/Transition/AvaPlayableTransition.h
@@ -5,6 +5,7 @@
 #include "AvaMediaDefines.h"
 #include "Playable/AvaPlayableRemoteControlValues.h"
 #include "UObject/Object.h"
+#include "UObject/ObjectKey.h"
 #include "AvaPlayableTransition.generated.h"
 
 class UAvaPlayable;
@@ -46,8 +47,14 @@ public:
 
 	void SetEnterPlayableValues(TArray<TSharedPtr<FAvaPlayableRemoteControlValues>>&& InPlayableValues);
 
-	TSharedPtr<FAvaPlayableRemoteControlValues> GetValuesForPlayable(const UAvaPlayable* InPlayable);
-	
+	/**
+	 * Finds the stored values for a given playable
+	 * @param InPlayable the playable to look for
+	 * @param bInIsEnterPlayable whether to look for the Enter Playable Values
+	 * @return the remote control values for the playable if found. null otherwise
+	 */
+	TSharedPtr<FAvaPlayableRemoteControlValues> GetValuesForPlayable(const UAvaPlayable* InPlayable, bool bInIsEnterPlayable);
+
 	/** This is called during the transition evaluation to indicate discarded playables. */
 	void MarkPlayableAsDiscard(UAvaPlayable* InPlayable);
 
@@ -66,7 +73,10 @@ protected:
 	EAvaPlayableTransitionFlags TransitionFlags = EAvaPlayableTransitionFlags::None;
 	
 	TArray<TSharedPtr<FAvaPlayableRemoteControlValues>> EnterPlayableValues;
-	
+
+	/** Other Playable Values that are not the Enter Playable's (i.e. Exiting or Playing Playable Values) */
+	TMap<TObjectKey<UAvaPlayable>, TSharedPtr<FAvaPlayableRemoteControlValues>> OtherPlayableValues;
+
 	TArray<TWeakObjectPtr<UAvaPlayable>> EnterPlayablesWeak;
 	TArray<TWeakObjectPtr<UAvaPlayable>> PlayingPlayablesWeak;
 	TArray<TWeakObjectPtr<UAvaPlayable>> ExitPlayablesWeak;
