Motion Design: Fixed issue with Outliner duplicate unit test failing.

The reason why the tests were failing is that the Outliner in the test environment does the entire outliner registration and test once -- and there were issues with how Level Items were not being registered and so the items to duplicate (that should've been parented to the Level Items) remained parentless -- and caused issues with outliner sorting.
The items did indeed duplicate but because they were not properly setup in the hierarchy, the hierarchy-based sorting was returning the wrong results.

Registereding Level items  solved the issue, but revealed an issue with grouping logic and its tests as the grouping algo requests the items to be 'reselected' and, because the grouping algo has these items in reverse order, the items in the test match but they're reversed. The ideal fix for this would be to change the "Add Item Param" to accept N  items rather than a single one and then having to create N "Add Item" actions in reverse order (to fix in UE-228496)

diff --git a/Engine/Plugins/Experimental/Avalanche/Source/AvalancheOutliner/Private/AvaOutliner.cpp b/Engine/Plugins/Experimental/Avalanche/Source/AvalancheOutliner/Private/AvaOutliner.cpp
--- a/Engine/Plugins/Experimental/Avalanche/Source/AvalancheOutliner/Private/AvaOutliner.cpp
+++ b/Engine/Plugins/Experimental/Avalanche/Source/AvalancheOutliner/Private/AvaOutliner.cpp
@@ -26,6 +26,7 @@
 #include "Item/AvaOutlinerActor.h"
 #include "Item/AvaOutlinerComponent.h"
 #include "Item/AvaOutlinerItemUtils.h"
+#include "Item/AvaOutlinerLevel.h"
 #include "Item/AvaOutlinerMaterial.h"
 #include "Item/AvaOutlinerTreeRoot.h"
 #include "ItemActions/AvaOutlinerAddItem.h"
@@ -396,7 +397,11 @@ void FAvaOutliner::Refresh()
 
 				ULevelInstanceSubsystem* const LevelInstanceSubsystem = InWorld->GetSubsystem<ULevelInstanceSubsystem>();
 
-				const ULevel* Level = InActor->GetLevel();
+				ULevel* const Level = InActor->GetLevel();
+
+				// Ensure that the Actor's Level is registered as an Item
+				FindOrAdd<FAvaOutlinerLevel>(Level);
+
 				AActor* Parent = InActor->GetSceneOutlinerParent();
 
 				// Try to find the Level Instance Actor to use as Parent for actors that aren't attached to anything
diff --git a/Engine/Plugins/Experimental/Avalanche/Source/AvalancheOutliner/Private/ItemActions/AvaOutlinerAddItem.cpp b/Engine/Plugins/Experimental/Avalanche/Source/AvalancheOutliner/Private/ItemActions/AvaOutlinerAddItem.cpp
--- a/Engine/Plugins/Experimental/Avalanche/Source/AvalancheOutliner/Private/ItemActions/AvaOutlinerAddItem.cpp
+++ b/Engine/Plugins/Experimental/Avalanche/Source/AvalancheOutliner/Private/ItemActions/AvaOutlinerAddItem.cpp
@@ -88,17 +88,24 @@ void FAvaOutlinerAddItem::Execute(FAvaOutliner& InOutliner)
 		{
 			const FAvaOutlinerScopedSelection ScopedSelection(*ModeTools, EAvaOutlinerScopedSelectionPurpose::Read);
 
-			//Automatically Select Item if it's Selected in Mode Tools
-			if (AddParams.Item->IsSelected(ScopedSelection))
-			{
-				// Select in Outliner but don't signal selection as we already have it selected in Mode Tools
-				AddParams.Flags = EAvaOutlinerAddItemFlags::Select;
-				AddParams.SelectionFlags &= ~EAvaOutlinerItemSelectionFlags::SignalSelectionChange;
-			}
-			//Signal Selection Change when we attempt to select this item in the Outliner but it isn't selected in Mode Tools
-			else if (EnumHasAnyFlags(AddParams.Flags, EAvaOutlinerAddItemFlags::Select))
+			const bool bSelectedInModeTools = AddParams.Item->IsSelected(ScopedSelection);
+			const bool bSelectedInOutliner  = InOutliner.GetSelectedItems().Contains(AddParams.Item);
+
+			// Update Selection if there's a discrepancy between the Mode Tools Selection & the Outliner View Selected Items
+			if (bSelectedInModeTools != bSelectedInOutliner)
 			{
-				AddParams.SelectionFlags |= EAvaOutlinerItemSelectionFlags::SignalSelectionChange;
+				// Automatically Select Item if it's Selected in Mode Tools and not yet in the Outliner
+				if (bSelectedInModeTools)
+				{
+					// Select in Outliner but don't signal selection as we already have it selected in Mode Tools
+					AddParams.Flags |= EAvaOutlinerAddItemFlags::Select;
+					AddParams.SelectionFlags &= ~EAvaOutlinerItemSelectionFlags::SignalSelectionChange;
+				}
+				// Signal Selection Change when we attempt to select this item in the Outliner but it isn't selected in Mode Tools
+				else if (EnumHasAnyFlags(AddParams.Flags, EAvaOutlinerAddItemFlags::Select))
+				{
+					AddParams.SelectionFlags |= EAvaOutlinerItemSelectionFlags::SignalSelectionChange;
+				}
 			}
 		}
 
