Motion Design: Actor spawn, duplication behavior in Outliner is now designed so that items are added below rather than above their template.

diff --git a/Engine/Plugins/Experimental/Avalanche/Source/AvalancheOutliner/Private/AvaOutliner.cpp b/Engine/Plugins/Experimental/Avalanche/Source/AvalancheOutliner/Private/AvaOutliner.cpp
--- a/Engine/Plugins/Experimental/Avalanche/Source/AvalancheOutliner/Private/AvaOutliner.cpp
+++ b/Engine/Plugins/Experimental/Avalanche/Source/AvalancheOutliner/Private/AvaOutliner.cpp
@@ -547,7 +547,7 @@ void FAvaOutliner::OnActorsDuplicated(const TMap<AActor*, AActor*>& InDuplicateA
 		}
 	}
 
-	const EItemDropZone RelativeDropZone = InRelativeDropZone.Get(EItemDropZone::AboveItem);
+	const EItemDropZone RelativeDropZone = InRelativeDropZone.Get(EItemDropZone::BelowItem);
 
 	// Newly Duplicated Items don't have a Parent Item set yet, so it's not going to call RemoveFromParent, which handles Detachment
 	bool bShouldDetachActors = false;
@@ -574,22 +574,9 @@ void FAvaOutliner::OnActorsDuplicated(const TMap<AActor*, AActor*>& InDuplicateA
 		return UE::AvaOutliner::CompareOutlinerItemOrder(TemplateItemA, TemplateItemB);
 	});
 
-	/*
-	 * REVERSE the ordering of Duplicate Actors
-	 * The order needs to be reversed when items are not added above the target item (e.g. RelativeDropZone is either below or onto)
-	 * Some context:
-	 * ABOVE- items are added just above the target, so the first one will always be topmost. Order is preserved
-	 * BELOW- items are added just below the target, so the last one added will be the topmost. Order is in reverse.
-	 * ONTO - items are added at index 0, so last one added will be topmost. Order is in reverse
-	 */
-	if (RelativeDropZone != EItemDropZone::AboveItem)
-	{
-		Algo::Reverse(DuplicateActors);
-	}
-	
 	TArray<TSharedPtr<IAvaOutlinerAction>> ItemActions;
 	ItemActions.Reserve(DuplicateActors.Num());
-	
+
 	FAvaOutlinerAddItemParams Params;
 	Params.RelativeDropZone = RelativeDropZone;
 	Params.Flags            = EAvaOutlinerAddItemFlags::AddChildren;
@@ -640,18 +627,8 @@ void FAvaOutliner::OnActorsDuplicated(const TMap<AActor*, AActor*>& InDuplicateA
 		}
 
 		TSharedRef<FAvaOutlinerAddItem> AddItemAction = NewItemAction<FAvaOutlinerAddItem>(Params);
+		ItemActions.Add(AddItemAction);
 
-		// If parent is duplicated, make sure we add this item to index 0 so that children are added in reverse order
-		// this is because ParentItem->AddChild(...) adds the Child at Index 0, so it's correcting the ordering issue
-		if (bIsParentDuplicated)
-		{
-			ItemActions.EmplaceAt(0, AddItemAction);
-		}
-		else
-		{
-			ItemActions.Add(AddItemAction);
-		}		
-		
 		// Append after first selected item is added
 		Params.SelectionFlags |= EAvaOutlinerItemSelectionFlags::AppendToCurrentSelection;
 	}
@@ -712,9 +689,8 @@ void FAvaOutliner::GroupSelection(AActor* InGroupingActor, const TOptional<FAtta
 		return;
 	}
 
-	// Sort Items in reverse order of hierarchy as Add Item will add them at index 0 by default
-	SortItems(SelectedItems, true);
-		
+	SortItems(SelectedItems);
+
 	FAvaOutlinerAddItemParams AddChildItemParams;
 	AddChildItemParams.RelativeItem     = GroupingItem;
 	AddChildItemParams.RelativeDropZone = EItemDropZone::OntoItem;
@@ -1163,11 +1139,11 @@ FAvaSceneItem FAvaOutliner::MakeSceneItemFromOutlinerItem(const FAvaOutlinerItem
 	return FAvaSceneItem(InItem->GetItemId().GetStringId());
 }
 
-void FAvaOutliner::SortItems(TArray<FAvaOutlinerItemPtr>& OutOutlinerItems, bool bInReverseOrder)
+void FAvaOutliner::SortItems(TArray<FAvaOutlinerItemPtr>& OutOutlinerItems)
 {
-	OutOutlinerItems.Sort([bInReverseOrder](const FAvaOutlinerItemPtr& ItemA, const FAvaOutlinerItemPtr& ItemB)
+	OutOutlinerItems.Sort([](const FAvaOutlinerItemPtr& ItemA, const FAvaOutlinerItemPtr& ItemB)
 	{
-		return UE::AvaOutliner::CompareOutlinerItemOrder(ItemA, ItemB) != bInReverseOrder;
+		return UE::AvaOutliner::CompareOutlinerItemOrder(ItemA, ItemB);
 	});
 }
 
diff --git a/Engine/Plugins/Experimental/Avalanche/Source/AvalancheOutliner/Private/AvaOutliner.h b/Engine/Plugins/Experimental/Avalanche/Source/AvalancheOutliner/Private/AvaOutliner.h
--- a/Engine/Plugins/Experimental/Avalanche/Source/AvalancheOutliner/Private/AvaOutliner.h
+++ b/Engine/Plugins/Experimental/Avalanche/Source/AvalancheOutliner/Private/AvaOutliner.h
@@ -229,7 +229,7 @@ public:
 	 * Helper function to sort the given array of items based on their ordering in the Outliner
 	 * @see FAvaOutliner::CompareOutlinerItemOrder
 	 */
-	static void SortItems(TArray<FAvaOutlinerItemPtr>& OutOutlinerItems, bool bInReverseOrder = false);
+	static void SortItems(TArray<FAvaOutlinerItemPtr>& OutOutlinerItems);
 
 	/** Normalizes the given Items by removing selected items that have their parent item also in the selection */
 	static void NormalizeItems(TArray<FAvaOutlinerItemPtr>& InOutItems);
diff --git a/Engine/Plugins/Experimental/Avalanche/Source/AvalancheOutliner/Private/DragDropOps/Handlers/AvaOutlinerActorDropHandler.cpp b/Engine/Plugins/Experimental/Avalanche/Source/AvalancheOutliner/Private/DragDropOps/Handlers/AvaOutlinerActorDropHandler.cpp
--- a/Engine/Plugins/Experimental/Avalanche/Source/AvalancheOutliner/Private/DragDropOps/Handlers/AvaOutlinerActorDropHandler.cpp
+++ b/Engine/Plugins/Experimental/Avalanche/Source/AvalancheOutliner/Private/DragDropOps/Handlers/AvaOutlinerActorDropHandler.cpp
@@ -72,12 +72,6 @@ void FAvaOutlinerActorDropHandler::MoveItems(EItemDropZone InDropZone, FAvaOutli
 			return !InItem.IsValid() || !InItem->GetParent().IsValid() || DraggedItemSet.Contains(InItem->GetParent());
 		});
 
-	// Reverse order for Onto since Item->AddChild(...) adds it to Index 0, so last item would be at the top, which is reversed
-	if (InDropZone == EItemDropZone::OntoItem)
-	{
-		Algo::Reverse(Items);
-	}
-
 	TSharedRef<FAvaOutliner> Outliner = StaticCastSharedRef<FAvaOutliner>(InTargetItem->GetOwnerOutliner());
 
 	for (const FAvaOutlinerItemPtr& Item : Items)
diff --git a/Engine/Plugins/Experimental/Avalanche/Source/AvalancheOutliner/Private/Item/AvaOutlinerItem.cpp b/Engine/Plugins/Experimental/Avalanche/Source/AvalancheOutliner/Private/Item/AvaOutlinerItem.cpp
--- a/Engine/Plugins/Experimental/Avalanche/Source/AvalancheOutliner/Private/Item/AvaOutlinerItem.cpp
+++ b/Engine/Plugins/Experimental/Avalanche/Source/AvalancheOutliner/Private/Item/AvaOutlinerItem.cpp
@@ -334,12 +334,12 @@ void FAvaOutlinerItem::AddChildChecked(const FAvaOutlinerAddItemParams& InAddIte
 		}
 		else
 		{
-			Children.EmplaceAt(0, InAddItemParams.Item);
+			Children.Emplace(InAddItemParams.Item);
 		}
 	}
 	else
 	{
-		Children.EmplaceAt(0, InAddItemParams.Item);
+		Children.Emplace(InAddItemParams.Item);
 	}
 
 	InAddItemParams.Item->SetParent(SharedThis(this));
diff --git a/Engine/Plugins/Experimental/Avalanche/Source/AvalancheOutliner/Private/Tests/AvaOutliner.spec.cpp b/Engine/Plugins/Experimental/Avalanche/Source/AvalancheOutliner/Private/Tests/AvaOutliner.spec.cpp
--- a/Engine/Plugins/Experimental/Avalanche/Source/AvalancheOutliner/Private/Tests/AvaOutliner.spec.cpp
+++ b/Engine/Plugins/Experimental/Avalanche/Source/AvalancheOutliner/Private/Tests/AvaOutliner.spec.cpp
@@ -170,7 +170,7 @@ void FAvaOutlinerSpec::Define()
 							, TestItems[TestItemIndex]);
 					});
 
-					It("should have each duplicate item right above its respective template item", [this]()
+					It("should have each duplicate item right below its respective template item", [this]()
 					{
 						constexpr int32 TestItemIndex = 0;
 
@@ -188,8 +188,8 @@ void FAvaOutlinerSpec::Define()
 						const int32 TemplateItemIndex  = ParentItem->GetChildIndex(TemplateItem);
 						const int32 DuplicateItemIndex = ParentItem->GetChildIndex(DuplicateItem);
 
-						TestEqual(TEXT("The duplicate item is right above template item")
-							, TemplateItemIndex - 1
+						TestEqual(TEXT("The duplicate item is right below template item")
+							, TemplateItemIndex + 1
 							, DuplicateItemIndex);
 					});
 				});
