[Motion Design] Fix crash on exit. The playback graph objects where not cleaning up the game instances correctly.

Fixing regression introduced in CL 30944594. "EndPlayWorld" must be requested when shutting down to avoid dangling world crash.
Fixing regression introduced in CL 35747313: In EndPlayWorld, cleaning up the local players added in BeginPlayWorld. Note: This also fixes the warning  "UWorld::DestroyActor: World has no context!" when stopping playback instances.

diff --git a/Engine/Plugins/Experimental/Avalanche/Source/Avalanche/Private/Framework/AvaGameInstance.cpp b/Engine/Plugins/Experimental/Avalanche/Source/Avalanche/Private/Framework/AvaGameInstance.cpp
--- a/Engine/Plugins/Experimental/Avalanche/Source/Avalanche/Private/Framework/AvaGameInstance.cpp
+++ b/Engine/Plugins/Experimental/Avalanche/Source/Avalanche/Private/Framework/AvaGameInstance.cpp
@@ -366,6 +366,12 @@ void UAvaGameInstance::EndPlayWorld()
 	OnEndPlay.Broadcast(this, PlayingChannelName);
 	PlayingChannelName = NAME_None;
 
+	if (LocalPlayer)
+	{
+		RemoveLocalPlayer(LocalPlayer);
+		LocalPlayer = nullptr;
+	}
+
 	if (PlayWorld)
 	{
 		PlayWorld->EndPlay(EEndPlayReason::LevelTransition);
diff --git a/Engine/Plugins/Experimental/Avalanche/Source/AvalancheMedia/Private/Playback/AvaPlaybackGraph.cpp b/Engine/Plugins/Experimental/Avalanche/Source/AvalancheMedia/Private/Playback/AvaPlaybackGraph.cpp
--- a/Engine/Plugins/Experimental/Avalanche/Source/AvalancheMedia/Private/Playback/AvaPlaybackGraph.cpp
+++ b/Engine/Plugins/Experimental/Avalanche/Source/AvalancheMedia/Private/Playback/AvaPlaybackGraph.cpp
@@ -84,7 +84,12 @@ void UAvaPlaybackGraph::Stop(EAvaPlaybackStopOptions InStopOptions)
 {
 	const bool bIsForceImmediate = EnumHasAnyFlags(InStopOptions, EAvaPlaybackStopOptions::ForceImmediate);
 	const bool bIsUnload = EnumHasAnyFlags(InStopOptions, EAvaPlaybackStopOptions::Unload);
-
+	bool bIsShuttingDown = false;
+	if (const TSharedPtr<FAvaPlaybackManager> PlaybackManager = PlaybackManagerWeak.Pin())
+	{
+		bIsShuttingDown = PlaybackManager->IsShuttingDown();
+	}
+	
 	// We need to unload all the playables, not just the ones currently traversed by the graph.
 	// We keep track of all the local playables in ChannelPlayableGroups. This now includes the
 	// remote playables.
@@ -96,10 +101,10 @@ void UAvaPlaybackGraph::Stop(EAvaPlaybackStopOptions InStopOptions)
 		
 		for (UAvaPlayable* Playable : Playables)
 		{
-			// EndPlayWorld is no longer issued here as some of the engine systems don't seem
+			// EndPlayWorld is no longer issued here (unless shutting down) as some of the engine systems don't seem
 			// to support a EndPlay, BeginPlay sequence anymore. It is unknown if this situation
 			// will change in the future, so we keep the option.
-			Playable->EndPlay(EAvaPlayableEndPlayOptions::None);
+			Playable->EndPlay(EAvaPlayableEndPlayOptions::None); //bIsShuttingDown ? EAvaPlayableEndPlayOptions::ConditionalEndPlayWorld : EAvaPlayableEndPlayOptions::None);
 
 			if (bIsUnload)
 			{
diff --git a/Engine/Plugins/Experimental/Avalanche/Source/AvalancheMedia/Public/Playback/AvaPlaybackManager.h b/Engine/Plugins/Experimental/Avalanche/Source/AvalancheMedia/Public/Playback/AvaPlaybackManager.h
--- a/Engine/Plugins/Experimental/Avalanche/Source/AvalancheMedia/Public/Playback/AvaPlaybackManager.h
+++ b/Engine/Plugins/Experimental/Avalanche/Source/AvalancheMedia/Public/Playback/AvaPlaybackManager.h
@@ -236,6 +236,8 @@ public:
 	 */
 	void StartShuttingDown() { bIsShuttingDown = true;}
 
+	bool IsShuttingDown() const { return bIsShuttingDown; }
+	
 	EAvaPlaybackStopOptions GetPlaybackStopOptions(bool bInUnload) const
 	{
 		EAvaPlaybackStopOptions Options = bIsShuttingDown ? EAvaPlaybackStopOptions::ForceImmediate : EAvaPlaybackStopOptions::None;
