[Motion Design] Rundown Editor: Fix Reset Controller to default will update runtime preview.

When the controllers are reset to the template value, the controller behaviors need to be executed and modified properties captured back into the page values. This will lead to the runtime (either preview or program) to update properly. Also need to update the page summary (derived from the controller values) as part of the same transaction.

diff --git a/Engine/Plugins/Experimental/Avalanche/Source/AvalancheMediaEditor/Private/Rundown/AvaRundownEditorUtils.cpp b/Engine/Plugins/Experimental/Avalanche/Source/AvalancheMediaEditor/Private/Rundown/AvaRundownEditorUtils.cpp
--- a/Engine/Plugins/Experimental/Avalanche/Source/AvalancheMediaEditor/Private/Rundown/AvaRundownEditorUtils.cpp
+++ b/Engine/Plugins/Experimental/Avalanche/Source/AvalancheMediaEditor/Private/Rundown/AvaRundownEditorUtils.cpp
@@ -12,6 +12,7 @@
 #include "IContentBrowserSingleton.h"
 #include "JsonObjectConverter.h"
 #include "Misc/PathViews.h"
+#include "RCVirtualProperty.h"
 #include "Rundown/AvaRundown.h"
 #include "Rundown/AvaRundownManagedInstanceCache.h"
 #include "StructDeserializer.h"
@@ -632,4 +633,60 @@ EAvaPlayableRemoteControlChanges UE::AvaRundownEditor::Utils::UpdateDefaultRemot
 	return Changes;
 }
 
+// Todo: merge with FAvaRundownPage::UpdatePageSummary(InRundown).
+// This will also be needed for rundown server.
+// It would need to add a "force update" parameter.
+bool UE::AvaRundownEditor::Utils::ForceUpdatePageSummary(int32 InPageId, UAvaRundown* InRundown)
+{
+	FAvaRundownPage& Page = InRundown ? InRundown->GetPage(InPageId) : FAvaRundownPage::NullPage;
+	if (!Page.IsValidPage())
+	{
+		return false;
+	}
+
+	// Important: the managed instance cache is used by the page details panel, which assumes the instances it holds
+	// for the current page are not modified by something else.
+	// This is why we need to backup the controllers and restore them at the end of the function. 
+	const TArray<TSharedPtr<FAvaRundownManagedInstance>> ManagedInstances = GetManagedInstancesForPage(InRundown, Page);
+	TArray<const URemoteControlPreset*> Presets;
+	Presets.Reserve(ManagedInstances.Num());
+
+	FAvaPlayableRemoteControlValues AllControllerValuesBackup;
+	
+	for (const TSharedPtr<FAvaRundownManagedInstance>& Instance : ManagedInstances)
+	{
+		if (URemoteControlPreset* Preset = Instance->GetRemoteControlPreset())
+		{
+			// Backup controllers
+			for (const URCVirtualPropertyBase* Controller : Preset->GetControllers())
+			{
+				if (!AllControllerValuesBackup.SetControllerValue(Controller->Id, Preset, /*bInIsDefault*/ false))
+				{
+					// Don't risk stomping preset controllers.
+					UE_LOG(LogAvaRundown, Error, TEXT("Unable to backup preset controllers. Aborting page summary update."));
+					return false;
+				}
+			}
+			
+			// Apply page controller value to preset without running the behaviors. We only care about controllers for page summary. 
+			Page.GetRemoteControlValues().ApplyControllerValuesToRemoteControlPreset(Preset, /*bInForceDisableBehaviors*/ true);
+			
+			Presets.Add(Preset);
+		}
+	}
+
+	const bool bSummaryUpdated = Page.UpdatePageSummary(Presets, /*bInIsPresetChanged*/ true);
+
+	// Restore managed instances controllers
+	for (const TSharedPtr<FAvaRundownManagedInstance>& Instance : ManagedInstances)
+	{
+		if (URemoteControlPreset* Preset = Instance->GetRemoteControlPreset())
+		{
+			AllControllerValuesBackup.ApplyControllerValuesToRemoteControlPreset(Preset, /*bInForceDisableBehaviors*/ true);
+		}
+	}
+
+	return bSummaryUpdated;
+}
+
 #undef LOCTEXT_NAMESPACE
diff --git a/Engine/Plugins/Experimental/Avalanche/Source/AvalancheMediaEditor/Private/Rundown/AvaRundownEditorUtils.h b/Engine/Plugins/Experimental/Avalanche/Source/AvalancheMediaEditor/Private/Rundown/AvaRundownEditorUtils.h
--- a/Engine/Plugins/Experimental/Avalanche/Source/AvalancheMediaEditor/Private/Rundown/AvaRundownEditorUtils.h
+++ b/Engine/Plugins/Experimental/Avalanche/Source/AvalancheMediaEditor/Private/Rundown/AvaRundownEditorUtils.h
@@ -107,4 +107,9 @@ namespace UE::AvaRundownEditor::Utils
 	TArray<TSharedPtr<FAvaRundownManagedInstance>> GetManagedInstancesForPage(const UAvaRundown* InRundown, const FAvaRundownPage& InPage);
 	bool MergeDefaultRemoteControlValues(const TArray<TSharedPtr<FAvaRundownManagedInstance>>& InManagedInstances, FAvaPlayableRemoteControlValues& OutMergedValues);
 	EAvaPlayableRemoteControlChanges UpdateDefaultRemoteControlValues(UAvaRundown* InRundown, const TArray<int32>& InSelectedPageIds);
+	
+	/**
+     * Utility function to force update the page summary from the controller values. 
+     */
+    bool ForceUpdatePageSummary(int32 InPageId, UAvaRundown* InRundown);
 }
\ No newline at end of file
diff --git a/Engine/Plugins/Experimental/Avalanche/Source/AvalancheMediaEditor/Private/Rundown/DetailsView/RemoteControl/Controllers/AvaRundownPageControllerContextMenu.cpp b/Engine/Plugins/Experimental/Avalanche/Source/AvalancheMediaEditor/Private/Rundown/DetailsView/RemoteControl/Controllers/AvaRundownPageControllerContextMenu.cpp
--- a/Engine/Plugins/Experimental/Avalanche/Source/AvalancheMediaEditor/Private/Rundown/DetailsView/RemoteControl/Controllers/AvaRundownPageControllerContextMenu.cpp
+++ b/Engine/Plugins/Experimental/Avalanche/Source/AvalancheMediaEditor/Private/Rundown/DetailsView/RemoteControl/Controllers/AvaRundownPageControllerContextMenu.cpp
@@ -4,9 +4,11 @@
 #include "AvaRundownRCControllerItem.h"
 #include "Controller/RCController.h"
 #include "IAvaMediaEditorModule.h"
+#include "Playable/AvaPlayableRemoteControl.h"
 #include "Rundown/AvaRundownCommands.h"
 #include "Rundown/AvaRundownEditor.h"
 #include "Rundown/AvaRundownPage.h"
+#include "Rundown/DetailsView/RemoteControl/Properties/SAvaRundownPageRemoteControlProps.h"
 #include "SAvaRundownRCControllerPanel.h"
 #include "ScopedTransaction.h"
 #include "ToolMenu.h"
@@ -168,28 +170,69 @@ bool FAvaRundownPageControllerContextMenu::HasValidSelectedItems() const
 
 void FAvaRundownPageControllerContextMenu::ResetValuesToDefaults(const bool bInUseTemplateValues)
 {
-	UAvaRundown* const Rundown = GetContextRundown();
+	const TSharedPtr<FAvaRundownEditor> RundownEditor = CurrentContext.IsValid() ? CurrentContext->GetRundownEditor() : nullptr;
+	UAvaRundown* const Rundown = RundownEditor.IsValid() ? RundownEditor->GetRundown() : nullptr;
 	if (!IsValid(Rundown))
 	{
 		return;
 	}
 
+	const FAvaRundownPage& Page = Rundown->GetPage(CurrentContext->GetRundownPageId());
+	if (!Page.IsValidPage())
+	{
+		return;
+	}
+
 	FScopedTransaction Transaction(LOCTEXT("ResetControllerValuesTransaction", "Reset RC Controller Values"));
 	Rundown->Modify();
 
+	bool bControllerValueModified = false;
+	
 	TArray<TWeakPtr<FAvaRundownRCControllerItem>> SelectedControllerItems = GetSelectedControllerItems();
 	for (const TWeakPtr<FAvaRundownRCControllerItem>& ControllerItemWeak : SelectedControllerItems)
 	{
-		if (TSharedPtr<FAvaRundownRCControllerItem> ControllerItem = ControllerItemWeak.Pin())
+		if (const TSharedPtr<FAvaRundownRCControllerItem> ControllerItem = ControllerItemWeak.Pin())
 		{
-			const URCController* const Controller = ControllerItem->GetController();
+			URCController* const Controller = ControllerItem->GetController();
 			if (IsValid(Controller))
 			{
-				Rundown->ResetRemoteControlControllerValue(CurrentContext->GetRundownPageId()
-					, Controller->Id, bInUseTemplateValues, /*bInIsDefault=*/false);
+				const EAvaPlayableRemoteControlChanges Changes = Rundown->ResetRemoteControlControllerValue(
+					CurrentContext->GetRundownPageId(), Controller->Id, bInUseTemplateValues, /*bInIsDefault=*/false);
+
+				if (Changes != EAvaPlayableRemoteControlChanges::None)
+				{
+					RundownEditor->MarkAsModified();
+					
+					// Apply the reset controller values from the page to the RCP so it updates the behaviors (and controlled entities).
+					using namespace UE::AvaPlayableRemoteControl;
+					if (const FAvaPlayableRemoteControlValue* ControllerValue = Page.GetRemoteControlValues().GetControllerValue(Controller->Id))
+					{
+						// This will execute the behaviours.
+						if (SetValueOfController(Controller, ControllerValue->Value) == EAvaPlayableRemoteControlResult::Completed)
+						{
+							bControllerValueModified = true;
+							
+							// Read values of the controlled entities and save that back in the page.
+							if (const URemoteControlPreset* Preset = Controller->PresetWeakPtr.Get())
+							{
+								TSet<FGuid> EntityIds;
+								GetEntitiesControlledByController(Preset, Controller, EntityIds);
+								SAvaRundownPageRemoteControlProps::SaveRemoteControlEntitiesToPage(Preset, EntityIds, Rundown, CurrentContext->GetRundownPageId());
+							}
+						}
+					}
+				}
 			}
 		}
 	}
+
+	if (bControllerValueModified)
+	{
+		if (const TSharedPtr<SAvaRundownRCControllerPanel> ControllerListWidget = CurrentContext->GetControllerListWidget())
+		{
+			ControllerListWidget->UpdatePageSummary(true);
+		}
+	}
 }
 
 bool FAvaRundownPageControllerContextMenu::CanResetValuesToDefaults(const bool bInUseTemplateValues) const
diff --git a/Engine/Plugins/Experimental/Avalanche/Source/AvalancheMediaEditor/Private/Rundown/DetailsView/RemoteControl/Properties/SAvaRundownPageRemoteControlProps.cpp b/Engine/Plugins/Experimental/Avalanche/Source/AvalancheMediaEditor/Private/Rundown/DetailsView/RemoteControl/Properties/SAvaRundownPageRemoteControlProps.cpp
--- a/Engine/Plugins/Experimental/Avalanche/Source/AvalancheMediaEditor/Private/Rundown/DetailsView/RemoteControl/Properties/SAvaRundownPageRemoteControlProps.cpp
+++ b/Engine/Plugins/Experimental/Avalanche/Source/AvalancheMediaEditor/Private/Rundown/DetailsView/RemoteControl/Properties/SAvaRundownPageRemoteControlProps.cpp
@@ -327,42 +327,7 @@ void SAvaRundownPageRemoteControlProps::OnRemoteControlExposedPropertiesModified
 
 	if (const TSharedPtr<FAvaRundownEditor> RundownEditor = RundownEditorWeak.Pin())
 	{
-		bool bModified = false;
-
-		for (const FGuid& Id : InModifiedProperties)
-		{
-			if (const TSharedPtr<FRemoteControlEntity> Entity = InPreset->GetExposedEntity<FRemoteControlEntity>(Id).Pin())
-			{
-				FAvaPlayableRemoteControlValue EntityValue;
-
-				{
-					using namespace UE::AvaPlayableRemoteControl;
-					if (const EAvaPlayableRemoteControlResult Result = GetValueOfEntity(Entity, EntityValue.Value); Failed(Result))
-					{
-						UE_LOG(LogAvaPlayableRemoteControl, Error, TEXT("Unable to get value of entity \"%s\": %s"),
-							*Entity->GetLabel().ToString(), *EnumToString(Result));
-						continue;
-					}
-				}
-				
-				const FAvaPlayableRemoteControlValue* StoredEntityValue = GetSelectedPageEntityValue(Entity);
-
-				if (StoredEntityValue && StoredEntityValue->IsSameValueAs(EntityValue))
-				{
-					continue;	// Skip if value is identical.
-				}
-
-				if (!SetSelectedPageEntityValue(Entity, EntityValue))
-				{
-					UE_LOG(LogAvaPlayableRemoteControl, Error, TEXT("Unable to set page entity value for: \"%s\""), *Entity->GetLabel().ToString());
-					continue;
-				}
-
-				bModified = true;
-			}
-		}
-
-		if (bModified)
+		if (SaveRemoteControlEntitiesToPage(InPreset, InModifiedProperties, RundownEditor->GetRundown(), ActivePageId))
 		{
 			RundownEditor->MarkAsModified();
 		}
@@ -511,6 +476,51 @@ const TArray<FAvaRundownRCPropertyItemPtr> SAvaRundownPageRemoteControlProps::Ge
 	return PropertyContainer->GetSelectedItems();
 }
 
+bool SAvaRundownPageRemoteControlProps::SaveRemoteControlEntitiesToPage(const URemoteControlPreset* InPreset, const TSet<FGuid>& InPropertyIds, UAvaRundown* InRundown, int32 InPageId)
+{
+	using namespace UE::AvaPlayableRemoteControl;
+	
+	if (!IsValid(InPreset) || InPageId == FAvaRundownPage::InvalidPageId || !InRundown)
+	{
+		return false;
+	}
+
+	const FAvaRundownPage& Page = InRundown->GetPage(InPageId);
+
+	bool bModified = false;
+
+	for (const FGuid& Id : InPropertyIds)
+	{
+		if (const TSharedPtr<const FRemoteControlEntity> Entity = InPreset->GetExposedEntity<const FRemoteControlEntity>(Id).Pin())
+		{
+			FAvaPlayableRemoteControlValue EntityValue;
+			
+			if (const EAvaPlayableRemoteControlResult Result = GetValueOfEntity(Entity, EntityValue.Value); Failed(Result))
+			{
+				UE_LOG(LogAvaPlayableRemoteControl, Error, TEXT("Unable to get value of entity \"%s\": %s"),
+					*Entity->GetLabel().ToString(), *EnumToString(Result));
+				continue;
+			}
+			
+			const FAvaPlayableRemoteControlValue* StoredEntityValue = Page.GetRemoteControlEntityValue(Entity->GetId());
+
+			if (StoredEntityValue && StoredEntityValue->IsSameValueAs(EntityValue))
+			{
+				continue;	// Skip if value is identical.
+			}
+			
+			if (!InRundown->SetRemoteControlEntityValue(InPageId, Entity->GetId(), EntityValue))
+			{
+				UE_LOG(LogAvaPlayableRemoteControl, Error, TEXT("Unable to set page entity value for: \"%s\""), *Entity->GetLabel().ToString());
+				continue;
+			}
+
+			bModified = true;
+		}
+	}
+	return bModified;
+}
+
 TSharedPtr<FAvaRundownPageRCPropsNotifyHook> SAvaRundownPageRemoteControlProps::GetNotifyHook() const
 {
 	return NotifyHook;
diff --git a/Engine/Plugins/Experimental/Avalanche/Source/AvalancheMediaEditor/Private/Rundown/DetailsView/RemoteControl/Properties/SAvaRundownPageRemoteControlProps.h b/Engine/Plugins/Experimental/Avalanche/Source/AvalancheMediaEditor/Private/Rundown/DetailsView/RemoteControl/Properties/SAvaRundownPageRemoteControlProps.h
--- a/Engine/Plugins/Experimental/Avalanche/Source/AvalancheMediaEditor/Private/Rundown/DetailsView/RemoteControl/Properties/SAvaRundownPageRemoteControlProps.h
+++ b/Engine/Plugins/Experimental/Avalanche/Source/AvalancheMediaEditor/Private/Rundown/DetailsView/RemoteControl/Properties/SAvaRundownPageRemoteControlProps.h
@@ -21,6 +21,7 @@ class ITableRow;
 class SAvaRundownPageRemoteControlProps;
 class STableViewBase;
 class SWidget;
+class UAvaRundown;
 class URemoteControlPreset;
 struct FAvaPlayableRemoteControlValue;
 struct FAvaRundownPage;
@@ -76,6 +77,15 @@ public:
 	
 	TSharedPtr<FAvaRundownPageRCPropsNotifyHook> GetNotifyHook() const;
 
+	/**
+	 * Utility to save the given Preset's Entities to the given page's property.
+	 * @param InPreset Source Preset the entity values are read from.
+	 * @param InEntityIds Set of entity ids to read from the given preset and save to the given page.
+	 * @param InRundown Rundown.
+	 * @param InPageId Page Id to save the value to.
+	 */
+	static bool SaveRemoteControlEntitiesToPage(const URemoteControlPreset* InPreset, const TSet<FGuid>& InEntityIds, UAvaRundown* InRundown, int32 InPageId);
+
 private:
 	static FAvaRundownRCPropertyHeaderRowExtensionDelegate HeaderRowExtensionDelegate;
 	static TMap<FName, TArray<FAvaRundownRCPropertyTableRowExtensionDelegate>> TableRowExtensionDelegates;
@@ -135,4 +145,4 @@ private:
 	TSharedPtr<FAvaRundownPageRCPropsNotifyHook> NotifyHook;
 
 	friend class FAvaRundownPageRCPropsNotifyHookImpl;
-};
+};
\ No newline at end of file
diff --git a/Engine/Plugins/Experimental/Avalanche/Source/AvalancheMediaEditor/Private/Rundown/Pages/Slate/SAvaRundownInstancedPageList.cpp b/Engine/Plugins/Experimental/Avalanche/Source/AvalancheMediaEditor/Private/Rundown/Pages/Slate/SAvaRundownInstancedPageList.cpp
--- a/Engine/Plugins/Experimental/Avalanche/Source/AvalancheMediaEditor/Private/Rundown/Pages/Slate/SAvaRundownInstancedPageList.cpp
+++ b/Engine/Plugins/Experimental/Avalanche/Source/AvalancheMediaEditor/Private/Rundown/Pages/Slate/SAvaRundownInstancedPageList.cpp
@@ -1490,7 +1490,13 @@ void SAvaRundownInstancedPageList::ResetPagesToDefaults(bool bInResetToTemplate)
 
 		for (int32 PageId : SelectedPageIds)
 		{
-			Rundown->ResetRemoteControlValues(PageId, bInResetToTemplate, /*bInIsDefault=*/false);
+			const EAvaPlayableRemoteControlChanges Changes =
+				Rundown->ResetRemoteControlValues(PageId, bInResetToTemplate, /*bInIsDefault=*/false);
+
+			if (EnumHasAnyFlags(Changes, EAvaPlayableRemoteControlChanges::ControllerValues))
+			{
+				UE::AvaRundownEditor::Utils::ForceUpdatePageSummary(PageId, Rundown);
+			}
 		}
 	}
 }
