[MotionDesign] Stagger tool improvements:
 - Additional options
 - Basic stagger along curve
 - Auto-apply
 - Non modal dialog

diff --git a/Engine/Plugins/Experimental/Avalanche/Source/AvalancheSequencer/Private/AvaSequencer.cpp b/Engine/Plugins/Experimental/Avalanche/Source/AvalancheSequencer/Private/AvaSequencer.cpp
--- a/Engine/Plugins/Experimental/Avalanche/Source/AvalancheSequencer/Private/AvaSequencer.cpp
+++ b/Engine/Plugins/Experimental/Avalanche/Source/AvalancheSequencer/Private/AvaSequencer.cpp
@@ -11,7 +11,6 @@
 #include "Clipboard/AvaSequenceImporter.h"
 #include "Commands/AvaSequencerAction.h"
 #include "Commands/AvaSequencerCommands.h"
-#include "Commands/Stagger/AvaSequencerStagger.h"
 #include "CoreGlobals.h"
 #include "DetailsView/SAvaMarkDetails.h"
 #include "DetailsView/Section/AvaSequencePlaybackDetails.h"
@@ -52,6 +51,7 @@
 #include "Sidebar/SSidebar.h"
 #include "Sidebar/SSidebarContainer.h"
 #include "Sidebar/SidebarDrawerConfig.h"
+#include "StaggerTool/AvaStaggerTool.h"
 #include "ToolMenu.h"
 #include "ToolMenuEntry.h"
 #include "ToolMenuSection.h"
@@ -129,7 +129,7 @@ FAvaSequencer::FAvaSequencer(IAvaSequencerProvider& InProvider, FAvaSequencerArg
 
 	SequencerActions =
 		{
-			MakeShared<FAvaSequencerStagger>(*this),
+			MakeShared<FAvaStaggerTool>(*this),
 		};
 
 	BindCommands();
diff --git a/Engine/Plugins/Experimental/Avalanche/Source/AvalancheSequencer/Private/Commands/AvaSequencerCommands.cpp b/Engine/Plugins/Experimental/Avalanche/Source/AvalancheSequencer/Private/Commands/AvaSequencerCommands.cpp
--- a/Engine/Plugins/Experimental/Avalanche/Source/AvalancheSequencer/Private/Commands/AvaSequencerCommands.cpp
+++ b/Engine/Plugins/Experimental/Avalanche/Source/AvalancheSequencer/Private/Commands/AvaSequencerCommands.cpp
@@ -54,8 +54,8 @@ void FAvaSequencerCommands::RegisterCommands()
 		, EUserInterfaceActionType::Button
 		, FInputChord());
 
-	UI_COMMAND(StaggerLayerBars
-		, "Stagger Layer Bars"
+	UI_COMMAND(OpenStaggerTool
+		, "Stagger Tool..."
 		, "Staggers the selected layer bars"
 		, EUserInterfaceActionType::Button
 		, FInputChord());
diff --git a/Engine/Plugins/Experimental/Avalanche/Source/AvalancheSequencer/Private/Commands/AvaSequencerCommands.h b/Engine/Plugins/Experimental/Avalanche/Source/AvalancheSequencer/Private/Commands/AvaSequencerCommands.h
--- a/Engine/Plugins/Experimental/Avalanche/Source/AvalancheSequencer/Private/Commands/AvaSequencerCommands.h
+++ b/Engine/Plugins/Experimental/Avalanche/Source/AvalancheSequencer/Private/Commands/AvaSequencerCommands.h
@@ -34,5 +34,5 @@ public:
 
 	TSharedPtr<FUICommandInfo> SpawnSequencePlayer;
 
-	TSharedPtr<FUICommandInfo> StaggerLayerBars;
+	TSharedPtr<FUICommandInfo> OpenStaggerTool;
 };
diff --git a/Engine/Plugins/Experimental/Avalanche/Source/AvalancheSequencer/Private/Commands/Stagger/AvaSequencerStagger.cpp b/Engine/Plugins/Experimental/Avalanche/Source/AvalancheSequencer/Private/Commands/Stagger/AvaSequencerStagger.cpp
deleted file mode 100644
diff --git a/Engine/Plugins/Experimental/Avalanche/Source/AvalancheSequencer/Private/Commands/Stagger/AvaSequencerStagger.h b/Engine/Plugins/Experimental/Avalanche/Source/AvalancheSequencer/Private/Commands/Stagger/AvaSequencerStagger.h
deleted file mode 100644
diff --git a/Engine/Plugins/Experimental/Avalanche/Source/AvalancheSequencer/Private/Commands/Stagger/AvaSequencerStaggerSettings.h b/Engine/Plugins/Experimental/Avalanche/Source/AvalancheSequencer/Private/Commands/Stagger/AvaSequencerStaggerSettings.h
deleted file mode 100644
diff --git a/Engine/Plugins/Experimental/Avalanche/Source/AvalancheSequencer/Private/Commands/Stagger/SAvaSequencerStaggerSettings.cpp b/Engine/Plugins/Experimental/Avalanche/Source/AvalancheSequencer/Private/Commands/Stagger/SAvaSequencerStaggerSettings.cpp
deleted file mode 100644
diff --git a/Engine/Plugins/Experimental/Avalanche/Source/AvalancheSequencer/Private/Commands/Stagger/SAvaSequencerStaggerSettings.h b/Engine/Plugins/Experimental/Avalanche/Source/AvalancheSequencer/Private/Commands/Stagger/SAvaSequencerStaggerSettings.h
deleted file mode 100644
diff --git a/Engine/Plugins/Experimental/Avalanche/Source/AvalancheSequencer/Private/Customization/AvaSequenceCustomization.cpp b/Engine/Plugins/Experimental/Avalanche/Source/AvalancheSequencer/Private/Customization/AvaSequenceCustomization.cpp
--- a/Engine/Plugins/Experimental/Avalanche/Source/AvalancheSequencer/Private/Customization/AvaSequenceCustomization.cpp
+++ b/Engine/Plugins/Experimental/Avalanche/Source/AvalancheSequencer/Private/Customization/AvaSequenceCustomization.cpp
@@ -124,7 +124,7 @@ TSharedRef<SWidget> FAvaSequenceCustomization::MakePlaybackMenu() const
 		const FAvaEaseCurveToolCommands& EaseCurveToolCommands = FAvaEaseCurveToolCommands::Get();
 
 		FToolMenuSection& SequencerSection = ToolMenu->FindOrAddSection(TEXT("AvaSequencerActions"), LOCTEXT("AvaSequencerActions", "Sequence"));
-		SequencerSection.AddMenuEntry(AvaSequencerCommands.StaggerLayerBars);
+		SequencerSection.AddMenuEntry(AvaSequencerCommands.OpenStaggerTool);
 
 		FToolMenuSection& EaseCurveToolSection = ToolMenu->FindOrAddSection(TEXT("AvaEaseCurveToolActions"), LOCTEXT("AvaEaseCurveToolActions", "Ease Curve Tool"));
 		EaseCurveToolSection.AddMenuEntry(EaseCurveToolCommands.QuickEase);
diff --git a/Engine/Plugins/Experimental/Avalanche/Source/AvalancheSequencer/Private/StaggerTool/AvaStaggerTool.cpp b/Engine/Plugins/Experimental/Avalanche/Source/AvalancheSequencer/Private/StaggerTool/AvaStaggerTool.cpp
new file mode 100644
--- /dev/null
+++ b/Engine/Plugins/Experimental/Avalanche/Source/AvalancheSequencer/Private/StaggerTool/AvaStaggerTool.cpp
@@ -0,0 +1,808 @@
+// Copyright Epic Games, Inc. All Rights Reserved.
+
+#include "AvaStaggerTool.h"
+#include "AvaSequencer.h"
+#include "Channels/MovieSceneChannel.h"
+#include "Commands/AvaSequencerCommands.h"
+#include "Framework/Application/SlateApplication.h"
+#include "Framework/Commands/UICommandList.h"
+#include "IKeyArea.h"
+#include "ISequencer.h"
+#include "Kismet/KismetMathLibrary.h"
+#include "Misc/FrameTime.h"
+#include "MVVM/Selection/Selection.h"
+#include "MVVM/ViewModels/ChannelModel.h"
+#include "MVVM/ViewModels/LayerBarModel.h"
+#include "MVVM/ViewModels/SequencerEditorViewModel.h"
+#include "ScopedTransaction.h"
+#include "Widgets/SAvaStaggerTool.h"
+
+using namespace UE::Sequencer;
+
+#define LOCTEXT_NAMESPACE "AvaSequencerStagger"
+
+FAvaStaggerTool::FStaggerLayerBarElement::FStaggerLayerBarElement(TViewModelPtr<FLayerBarModel> InLayerBarModel)
+	: LayerBarModel(MoveTemp(InLayerBarModel))
+{
+	const TRange<FFrameNumber> ComputedRange = LayerBarModel->ComputeRange();
+	OriginalFrame = ComputedRange.GetLowerBoundValue();
+}
+
+FAvaStaggerTool::FStaggerKeyFrameElement::FStaggerKeyFrameElement(TViewModelPtr<FChannelModel> InKeyChannelModel
+	, FKeyHandle InKeyHandle, FFrameNumber InOriginalFrame)
+	: KeyChannelModel(MoveTemp(InKeyChannelModel))
+	, KeyHandle(MoveTemp(InKeyHandle))
+	, OriginalFrame(MoveTemp(InOriginalFrame))
+{
+}
+
+FAvaStaggerTool::FAvaStaggerTool(FAvaSequencer& InOwner)
+	: FAvaSequencerAction(InOwner)
+{
+	Settings.Reset(NewObject<UAvaSequencerStaggerSettings>(GetTransientPackage(), NAME_None, RF_Transient | RF_Transactional));
+
+	OnResetToDefaults();
+}
+
+FAvaStaggerTool::~FAvaStaggerTool()
+{
+	UnbindToSequencerSelectionChange();
+}
+
+TObjectPtr<UAvaSequencerStaggerSettings> FAvaStaggerTool::GetSettings() const
+{
+	return Settings.Get();
+}
+
+void FAvaStaggerTool::MapAction(const TSharedRef<FUICommandList>& InCommandList)
+{
+	InCommandList->MapAction(FAvaSequencerCommands::Get().OpenStaggerTool
+		, FExecuteAction::CreateSP(this, &FAvaStaggerTool::Execute)
+		, FCanExecuteAction::CreateSP(this, &FAvaStaggerTool::CanExecute));
+}
+
+bool FAvaStaggerTool::CanExecute() const
+{
+	return true;
+}
+
+void FAvaStaggerTool::Execute()
+{
+	if (ToolWindow.IsValid())
+	{
+		ToolWindow->BringToFront();
+		return;
+	}
+
+	CacheOriginalElements();
+
+	ToolWindow =
+		SNew(SWindow)
+		.Title(LOCTEXT("StaggerToolDialogTitle", "Stagger Tool"))
+		.SizingRule(ESizingRule::Autosized)
+		.SupportsMaximize(false)
+		.SupportsMinimize(false)
+		.Content()
+		[
+			SAssignNew(ToolWidget, SAvaStaggerTool, SharedThis(this))
+			.OnResetToDefaults(this, &FAvaStaggerTool::OnResetToDefaults)
+			.OnSettingChange(this, &FAvaStaggerTool::OnSettingChange)
+			.OnApply(this, &FAvaStaggerTool::OnApply)
+		];
+
+	ToolWindow->SetOnWindowClosed(FOnWindowClosed::CreateLambda([this](const TSharedRef<SWindow>&)
+		{
+			ToolWindow.Reset();
+			ToolWindow = nullptr;
+		}));
+
+	const TSharedPtr<SWindow> ParentWindow = FSlateApplication::Get().FindBestParentWindowForDialogs(Owner.GetSequencer()->GetSequencerWidget());
+	if (ParentWindow.IsValid())
+	{
+		FSlateApplication::Get().AddWindowAsNativeChild(ToolWindow.ToSharedRef(), ParentWindow.ToSharedRef());
+	}
+	else
+	{
+		FSlateApplication::Get().AddWindow(ToolWindow.ToSharedRef());
+	}
+
+	BindToSequencerSelectionChange();
+}
+
+void FAvaStaggerTool::BindToSequencerSelectionChange()
+{
+	const TSharedRef<ISequencer> Sequencer = Owner.GetSequencer();
+	if (const TSharedPtr<FSequencerEditorViewModel> ViewModel = Sequencer->GetViewModel())
+	{
+		if (const TSharedPtr<FSequencerSelection> Selection = ViewModel->GetSelection())
+		{
+			Selection->OnChanged.AddRaw(this, &FAvaStaggerTool::OnSequencerSelectionChanged);
+		}
+	}
+}
+
+void FAvaStaggerTool::UnbindToSequencerSelectionChange()
+{
+	const TSharedRef<ISequencer> Sequencer = Owner.GetSequencer();
+	if (const TSharedPtr<FSequencerEditorViewModel> ViewModel = Sequencer->GetViewModel())
+	{
+		if (const TSharedPtr<FSequencerSelection> Selection = ViewModel->GetSelection())
+		{
+			Selection->OnChanged.RemoveAll(this);
+		}
+	}
+}
+
+void FAvaStaggerTool::OnResetToDefaults()
+{
+	Settings->ResetToolOptions();
+}
+
+void FAvaStaggerTool::OnSettingChange(const FName InPropertyName)
+{
+	if (InPropertyName == GET_MEMBER_NAME_CHECKED(UAvaSequencerStaggerSettings, bAutoApply)
+		&& !Settings->bAutoApply)
+	{
+		return;
+	}
+
+	if (Settings->bAutoApply)
+	{
+		OnApply();
+	}
+}
+
+void FAvaStaggerTool::OnApply()
+{
+	Stagger();
+}
+
+void FAvaStaggerTool::OnSequencerSelectionChanged()
+{
+	CacheOriginalElements();
+
+	if (Settings->bAutoApply)
+	{
+		OnApply();
+	}
+}
+
+bool FAvaStaggerTool::HasValidSelection() const
+{
+	return IsLayerBarSelection() || IsKeyFrameSelection();
+}
+
+bool FAvaStaggerTool::IsLayerBarSelection() const
+{
+	return OriginalLayerBarElements.Num() > 1 && OriginalKeyFrameElements.Num() == 0;
+}
+
+bool FAvaStaggerTool::IsKeyFrameSelection() const
+{
+	return OriginalKeyFrameElements.Num() > 1 && OriginalLayerBarElements.Num() == 0;
+}
+
+int32 FAvaStaggerTool::GetSelectionCount() const
+{
+	if (IsLayerBarSelection())
+	{
+		return OriginalLayerBarElements.Num();
+	}
+
+	if (IsKeyFrameSelection())
+	{
+		return OriginalKeyFrameElements.Num();
+	}
+
+	return 0;
+}
+
+FFrameNumber FAvaStaggerTool::ConvertTime(FFrameTime InSourceTime) const
+{
+	const TSharedRef<ISequencer> Sequencer = Owner.GetSequencer();
+	const FFrameRate FocusedDisplayRate = Sequencer->GetFocusedDisplayRate();
+	const FFrameRate FocusedTickResolution = Sequencer->GetFocusedTickResolution();
+	return ConvertFrameTime(InSourceTime, FocusedDisplayRate, FocusedTickResolution).RoundToFrame();
+}
+
+TRange<FFrameNumber> FAvaStaggerTool::GetOperationRange(const EAvaSequencerStaggerRange InRange) const
+{
+	const TSharedRef<ISequencer> Sequencer = Owner.GetSequencer();
+
+	switch (InRange)
+	{
+	default:
+	case EAvaSequencerStaggerRange::Playback:
+		return Sequencer->GetPlaybackRange();
+
+	case EAvaSequencerStaggerRange::Selection:
+		return Sequencer->GetSelectionRange();
+
+	case EAvaSequencerStaggerRange::Custom:
+		const FFrameNumber MaxRange = ConvertTime(Settings->ToolOptions.CustomRange);
+		return TRange<FFrameNumber>(0, MaxRange);
+	}
+}
+
+FFrameNumber FAvaStaggerTool::GetInterval(const FFrameNumber InRangeSize, const int32 InElementCount, const FFrameNumber InElementSize) const
+{
+	switch (Settings->ToolOptions.Distribution)
+	{
+	default:
+	case EAvaSequencerStaggerDistribution::Increment:
+	case EAvaSequencerStaggerDistribution::Decrement:
+		{
+			FFrameNumber Interval = ConvertTime(Settings->ToolOptions.Interval);
+			if (Settings->ToolOptions.Distribution == EAvaSequencerStaggerDistribution::Decrement)
+			{
+				Interval = Interval * -1;
+			}
+			return Settings->ToolOptions.bUseCurve && InElementSize.Value != 0 ? InElementSize : Interval;
+		}
+
+	case EAvaSequencerStaggerDistribution::Range:
+		return InRangeSize.Value / InElementCount;
+
+	case EAvaSequencerStaggerDistribution::Random:
+		FRandomStream Stream(Settings->ToolOptions.RandomSeed);
+		if (InElementSize > InRangeSize.Value)
+		{
+			return Stream.RandRange(0, InRangeSize.Value);
+		}
+		return Stream.RandRange(0, (InRangeSize - InElementSize).Value);
+	}
+}
+
+FFrameNumber FAvaStaggerTool::CalculateLayerBarFrameSpan(const TArray<FStaggerLayerBarElement>& InStaggerElements, const int32 InStopIndex)
+{
+	FFrameNumber OutFrame;
+	for (int32 Index = 0; Index < InStaggerElements.Num(); ++Index)
+	{
+		if (Index == InStopIndex)
+		{
+			break;
+		}
+		OutFrame += InStaggerElements[Index].LayerBarModel->ComputeRange().Size<FFrameNumber>();
+	}
+	return MoveTemp(OutFrame);
+}
+
+FFrameNumber FAvaStaggerTool::GetLayerBarOperationOffset(const TViewModelPtr<FLayerBarModel>& InLayerBarModel) const
+{
+	const TRange<FFrameNumber> Range = InLayerBarModel->ComputeRange();
+	const FFrameNumber RangeSize = Range.Size<FFrameNumber>();
+	const float LayerBarPosition = Settings->ToolOptions.OperationPoint * RangeSize.Value;
+	return FMath::FloorToInt32(LayerBarPosition);
+}
+
+void FAvaStaggerTool::CacheOriginalElements()
+{
+	const TSharedPtr<FSequencerSelection> SequencerSelection = GetSequencerSelection();
+	if (!SequencerSelection.IsValid())
+	{
+		return;
+	}
+
+	const int32 OutlinerCount = SequencerSelection->Outliner.Num();
+	const int32 TrackAreaCount = SequencerSelection->TrackArea.Num();
+	const int32 KeyFrameCount = SequencerSelection->KeySelection.Num();
+
+	OriginalLayerBarElements.Empty();
+	OriginalKeyFrameElements.Empty();
+
+	if ((OutlinerCount > 1 && TrackAreaCount == 0 && KeyFrameCount == 0)
+		|| (OutlinerCount <= 1 && TrackAreaCount > 1 && KeyFrameCount == 0))
+	{
+		OriginalLayerBarElements = GatherSelectionLayerBarElements();
+	}
+	else if (OutlinerCount == 0 && TrackAreaCount == 0 && KeyFrameCount > 1)
+	{
+		OriginalKeyFrameElements = GatherSelectionKeyFrameElements();
+	}
+}
+
+FFrameNumber FAvaStaggerTool::FindFirstStaggerPoint(const TArrayView<FStaggerLayerBarElement>& InStaggerElements, const TRange<FFrameNumber>& InRange) const
+{
+	FFrameNumber OutFrameNumber;
+
+	switch (Settings->ToolOptions.Distribution)
+	{
+	case EAvaSequencerStaggerDistribution::Increment:
+		break;
+
+	case EAvaSequencerStaggerDistribution::Decrement:
+		break;
+
+	case EAvaSequencerStaggerDistribution::Range:
+		switch (Settings->ToolOptions.Range)
+		{
+		case EAvaSequencerStaggerRange::Playback:
+			OutFrameNumber = Owner.GetSequencer()->GetPlaybackRange().GetLowerBoundValue();
+			OutFrameNumber += ConvertTime(Settings->ToolOptions.Shift);
+			return MoveTemp(OutFrameNumber);
+
+		case EAvaSequencerStaggerRange::Selection:
+			OutFrameNumber = Owner.GetSequencer()->GetSelectionRange().GetLowerBoundValue();
+			OutFrameNumber += ConvertTime(Settings->ToolOptions.Shift);
+			return MoveTemp(OutFrameNumber);
+
+		case EAvaSequencerStaggerRange::Custom:
+			break;
+		}
+		break;
+
+	case EAvaSequencerStaggerDistribution::Random:
+		const int32 OriginalElementCount = OriginalLayerBarElements.Num();
+		if (OriginalElementCount > 1)
+		{
+			const TRange<FFrameNumber> Range = GetOperationRange(Settings->ToolOptions.Range);
+			const FFrameNumber RangeSize = Range.Size<FFrameNumber>();
+			const FFrameNumber ElementSize = OriginalLayerBarElements[0].LayerBarModel->ComputeRange().Size<FFrameNumber>();
+			OutFrameNumber = GetInterval(RangeSize, OriginalElementCount, ElementSize);
+		}
+		break;
+	}
+
+	switch (Settings->ToolOptions.StartPosition)
+	{
+	case EAvaSequencerStaggerStartPosition::FirstSelected:
+		if (!InStaggerElements.IsEmpty())
+		{
+			OutFrameNumber = InStaggerElements[0].OriginalFrame;
+		}
+		break;
+
+	case EAvaSequencerStaggerStartPosition::FirstInTimeline:
+		FindFirstFrameInTimeline(InStaggerElements, OutFrameNumber);
+		break;
+
+	case EAvaSequencerStaggerStartPosition::Playhead:
+		const TSharedRef<ISequencer> Sequencer = Owner.GetSequencer();
+		OutFrameNumber = Sequencer->GetGlobalTime().Time.FrameNumber;
+		break;
+	}
+
+	if (IsLayerBarSelection())
+	{
+		if (Settings->ToolOptions.Distribution == EAvaSequencerStaggerDistribution::Random)
+		{
+			const int32 OriginalElementCount = OriginalLayerBarElements.Num();
+			if (OriginalElementCount > 1)
+			{
+				const TRange<FFrameNumber> Range = GetOperationRange(Settings->ToolOptions.Range);
+				const FFrameNumber RangeSize = Range.Size<FFrameNumber>();
+				const FFrameNumber ElementSize = OriginalLayerBarElements[0].LayerBarModel->ComputeRange().Size<FFrameNumber>();
+				OutFrameNumber += GetInterval(RangeSize, OriginalElementCount, ElementSize);
+			}
+		}
+	}
+
+	OutFrameNumber += ConvertTime(Settings->ToolOptions.Shift);
+
+	return MoveTemp(OutFrameNumber);
+}
+
+FFrameNumber FAvaStaggerTool::FindFirstStaggerPoint(const TArrayView<FStaggerKeyFrameElement>& InStaggerElements, const TRange<FFrameNumber>& InRange) const
+{
+	FFrameNumber OutFrameNumber;
+
+	switch (Settings->ToolOptions.StartPosition)
+	{
+	case EAvaSequencerStaggerStartPosition::FirstSelected:
+		if (!InStaggerElements.IsEmpty())
+		{
+			OutFrameNumber = InStaggerElements[0].OriginalFrame;
+		}
+		break;
+
+	case EAvaSequencerStaggerStartPosition::FirstInTimeline:
+		FindFirstFrameInTimeline(InStaggerElements, OutFrameNumber);
+		break;
+
+	case EAvaSequencerStaggerStartPosition::Playhead:
+		const TSharedRef<ISequencer> Sequencer = Owner.GetSequencer();
+		OutFrameNumber = Sequencer->GetGlobalTime().Time.FrameNumber;
+		break;
+	}
+
+	return OutFrameNumber;
+}
+
+TArray<TViewModelPtr<FLayerBarModel>> FAvaStaggerTool::GatherAllSelectedLayerBars() const
+{
+	const TSharedPtr<FSequencerSelection> SequencerSelection = GetSequencerSelection();
+	if (!SequencerSelection.IsValid())
+	{
+		return {};
+	}
+
+	TArray<TViewModelPtr<FLayerBarModel>> LayerBarModels;
+
+	const int32 OutlinerCount = SequencerSelection->Outliner.Num();
+	const int32 TrackAreaCount = SequencerSelection->TrackArea.Num();
+
+	if (OutlinerCount > 1)
+	{
+		LayerBarModels.Reserve(OutlinerCount);
+
+		for (const FViewModelPtr& ViewModel : SequencerSelection->Outliner)
+		{
+			if (const TViewModelPtr<ITrackAreaExtension> TrackArea = ViewModel.ImplicitCast())
+			{
+				for (const FViewModelPtr& TrackAreaModel : TrackArea->GetTopLevelChildTrackAreaModels())
+				{
+					if (const TViewModelPtr<FLayerBarModel> LayerBarModel = TrackAreaModel.ImplicitCast())
+					{
+						LayerBarModels.Add(LayerBarModel);
+					}
+				}
+			}
+		}
+	}
+	else if (OutlinerCount <= 1 && TrackAreaCount > 0)
+	{
+		LayerBarModels.Reserve(TrackAreaCount);
+
+		for (const FViewModelPtr& ViewModel : SequencerSelection->TrackArea)
+		{
+			if (const TViewModelPtr<FLayerBarModel> LayerBarModel = ViewModel.ImplicitCast())
+			{
+				LayerBarModels.Add(LayerBarModel);
+			}
+		}
+	}
+
+	// Remove all layer bars that have a descendant that is selected
+	LayerBarModels.RemoveAll([SequencerSelection](const TViewModelPtr<FLayerBarModel>& InLayerBarModel)
+		{
+			const TViewModelPtr<IOutlinerExtension> OutlinerExtension = InLayerBarModel->GetLinkedOutlinerItem();
+			if (!OutlinerExtension.IsValid())
+			{
+				return false;
+			}
+
+			for (const TViewModelPtr<IOutlinerExtension>& ChildOutlinerExtension : OutlinerExtension.AsModel()->GetDescendantsOfType<IOutlinerExtension>())
+			{
+				if (SequencerSelection->Outliner.IsSelected(ChildOutlinerExtension))
+				{
+					return true;
+				}
+			}
+
+			return false;
+		});
+
+	return LayerBarModels;
+}
+
+TArray<FAvaStaggerTool::FStaggerLayerBarElement> FAvaStaggerTool::GatherSelectionLayerBarElements() const
+{
+	TArray<TViewModelPtr<FLayerBarModel>> LayerBarModels = GatherAllSelectedLayerBars();
+
+	// Compute range and set original frame
+	TArray<FStaggerLayerBarElement> OutElements;
+	OutElements.Append(MoveTemp(LayerBarModels));
+
+	for (FStaggerLayerBarElement& Element : OutElements)
+	{
+		Element.OriginalFrame = Element.LayerBarModel->ComputeRange().GetLowerBoundValue();
+	}
+
+	return OutElements;
+}
+
+TArray<FAvaStaggerTool::FStaggerKeyFrameElement> FAvaStaggerTool::GatherSelectionKeyFrameElements() const
+{
+	TArray<FStaggerKeyFrameElement> OutElements;
+
+	const TSharedPtr<FSequencerSelection> SequencerSelection = GetSequencerSelection();
+	if (!SequencerSelection.IsValid())
+	{
+		return OutElements;
+	}
+
+	// Gather all the selected keyframe handles
+	OutElements.Reserve(SequencerSelection->KeySelection.Num());
+
+	for (const FKeyHandle KeyHandle : SequencerSelection->KeySelection)
+	{
+		if (KeyHandle != FKeyHandle::Invalid())
+		{
+			if (const TViewModelPtr<FChannelModel> KeyModel = SequencerSelection->KeySelection.GetModelForKey(KeyHandle))
+			{
+				FFrameNumber KeyFrameNumber;
+				KeyModel->GetChannel()->GetKeyTime(KeyHandle, KeyFrameNumber);
+
+				OutElements.Add(FStaggerKeyFrameElement(KeyModel, KeyHandle, KeyFrameNumber));
+			}
+		}
+	}
+
+	return MoveTemp(OutElements);
+}
+
+FFrameNumber FAvaStaggerTool::CalculateLocalCurveOffset(const float InCurveTime, const TRange<FFrameNumber>& InRange) const
+{
+	const FRichCurve* const RichCurve = Settings->ToolOptions.Curve.GetRichCurve();
+	if (!RichCurve)
+	{
+		return FFrameNumber();
+	}
+
+	float MinTime = 0.f;
+	float MaxTime = 0.f;
+	RichCurve->GetTimeRange(MinTime, MaxTime);
+
+	float MinValue = 0.f;
+	float MaxValue = 0.f;
+	RichCurve->GetValueRange(MinValue, MaxValue);
+
+	const FFrameNumber RangeSize = InRange.Size<FFrameNumber>();
+
+	float AdjustedCurveTime = InCurveTime + Settings->ToolOptions.CurveOffset;
+	if (AdjustedCurveTime > 1.f)
+	{
+		AdjustedCurveTime -= 1.f;
+	}
+	else if (AdjustedCurveTime < -1.f)
+	{
+		AdjustedCurveTime += 1.f;
+	}
+
+	// Map the 0 - 1 curve time to the min/max of the of the range
+	// InCurveTime should be between 0 and 1
+	// MinTime should typically be 0, but can be anything the user decides
+	// MaxTime should typically be 1, but can be anything the user decides
+	const float LocalMappedCurveTime = UKismetMathLibrary::MapRangeClamped(AdjustedCurveTime
+		, 0.f, 1.f
+		, MinTime, MaxTime);
+
+	const float EvalValue = RichCurve->Eval(LocalMappedCurveTime);
+
+	// Normalize the value and map to range since the value could be anything set by the user
+	// (at least until we have a custom graph widget like the Ease Curve Tool to handle this)
+	const float LocalFrameOffset = UKismetMathLibrary::MapRangeClamped(EvalValue
+		, MinValue, MaxValue
+		, 0.f, RangeSize.Value);
+
+	return FMath::RoundToInt32(LocalFrameOffset);
+}
+
+FFrameNumber FAvaStaggerTool::FindNextStaggerLocation(const FStaggerLayerBarElement& InElement
+	, const int32 InElementIndex
+	, const FFrameNumber InFirstFrame
+	, const FFrameNumber InCurrentFrame) const
+{
+	const FRichCurve* const RichCurve = Settings->ToolOptions.Curve.GetRichCurve();
+	const FFrameNumber ElementSize = InElement.LayerBarModel->ComputeRange().Size<FFrameNumber>();
+	const int32 OriginalElementCount = OriginalLayerBarElements.Num();
+	const TRange<FFrameNumber> Range = GetOperationRange(Settings->ToolOptions.Range);
+	const FFrameNumber RangeSize = Range.Size<FFrameNumber>();
+	const FFrameNumber Interval = GetInterval(RangeSize, OriginalElementCount, ElementSize);
+
+	if (Settings->ToolOptions.bUseCurve)
+	{
+		if (!RichCurve)
+		{
+			return InFirstFrame;
+		}
+
+		const FFrameNumber FullFrameSpan = CalculateLayerBarFrameSpan(OriginalLayerBarElements);
+		const FFrameNumber ElementFrameSpan = CalculateLayerBarFrameSpan(OriginalLayerBarElements, InElementIndex + 1);
+		const float ElementStartPercentOfRange = static_cast<float>(ElementFrameSpan.Value) / static_cast<float>(FullFrameSpan.Value);
+
+		if (Settings->ToolOptions.Distribution == EAvaSequencerStaggerDistribution::Increment
+			|| Settings->ToolOptions.Distribution == EAvaSequencerStaggerDistribution::Decrement)
+		{
+			const FFrameNumber LocalCurveOffset = CalculateLocalCurveOffset(ElementStartPercentOfRange, TRange<FFrameNumber>(0, FullFrameSpan));
+			return InFirstFrame + LocalCurveOffset + Interval;
+		}
+
+		const float CurveStepTime = 1.f / OriginalElementCount;
+		const float CurveTime = CurveStepTime * InElementIndex;
+
+		const FFrameNumber LocalCurveOffset = CalculateLocalCurveOffset(CurveTime, Range);
+		return InFirstFrame + LocalCurveOffset + Interval;
+	}
+
+	if (Settings->ToolOptions.Distribution == EAvaSequencerStaggerDistribution::Random)
+	{
+		return InFirstFrame + Interval;
+	}
+
+	const FFrameNumber LayerBarOperationOffset = GetLayerBarOperationOffset(InElement.LayerBarModel);
+	return InCurrentFrame + LayerBarOperationOffset + Interval;
+}
+
+FFrameNumber FAvaStaggerTool::FindNextStaggerLocation(const FStaggerKeyFrameElement& InElement
+	, const int32 InElementIndex
+	, const FFrameNumber InFirstFrame
+	, const FFrameNumber InCurrentFrame) const
+{
+	const FRichCurve* const RichCurve = Settings->ToolOptions.Curve.GetRichCurve();
+	const int32 OriginalElementCount = OriginalLayerBarElements.Num();
+	const TRange<FFrameNumber> Range = GetOperationRange(Settings->ToolOptions.Range);
+	const FFrameNumber RangeSize = Range.Size<FFrameNumber>();
+	const FFrameNumber Interval = GetInterval(RangeSize, OriginalElementCount, FFrameNumber());
+
+	if (Settings->ToolOptions.bUseCurve)
+	{
+		if (!RichCurve)
+		{
+			return InFirstFrame;
+		}
+
+		const float CurveStepTime = 1.f / OriginalElementCount;
+		const FFrameNumber NextStaggerFrameNumber = CalculateLocalCurveOffset(CurveStepTime * InElementIndex, Range);
+
+		return InFirstFrame + NextStaggerFrameNumber + Interval;
+	}
+
+	if (Settings->ToolOptions.Distribution == EAvaSequencerStaggerDistribution::Random)
+	{
+		return InFirstFrame + Interval;
+	}
+
+	return InCurrentFrame + Interval;
+}
+
+void FAvaStaggerTool::Stagger()
+{
+	if (IsLayerBarSelection())
+	{
+		// Reset layer bar offsets that were last applied
+		for (FStaggerLayerBarElement& Element : OriginalLayerBarElements)
+		{
+			Element.LayerBarModel->Offset(-Element.LastOffset);
+			Element.LastOffset = 0;
+		}
+
+		StaggerLayers(OriginalLayerBarElements);
+	}
+	else if (IsKeyFrameSelection())
+	{
+		// Reset key frames to their last location
+		for (FStaggerKeyFrameElement& Element : OriginalKeyFrameElements)
+		{
+			SetKeyFrameElementTime(Element, Element.OriginalFrame);
+		}
+
+		StaggerKeyFrames(OriginalKeyFrameElements);
+	}
+}
+
+bool FAvaStaggerTool::IsAutoAppying() const
+{
+	return Settings->bAutoApply;
+}
+
+void FAvaStaggerTool::StaggerLayers(const TArrayView<FStaggerLayerBarElement>& InOriginalElements)
+{
+	const int32 OriginalElementCount = InOriginalElements.Num();
+	if (OriginalElementCount <= 1)
+	{
+		return;
+	}
+
+	FScopedTransaction Transaction(LOCTEXT("StaggerLayers", "Stagger Layers"));
+
+	const TRange<FFrameNumber> Range = GetOperationRange(Settings->ToolOptions.Range);
+	const FFrameNumber FirstStaggerPoint = FindFirstStaggerPoint(InOriginalElements, Range);
+
+	FFrameNumber NextStaggerPoint = FirstStaggerPoint;
+	int32 CurrentElementNum = 0;
+	int32 CurrentGroupNum = 0;
+
+	for (FStaggerLayerBarElement& Element : InOriginalElements)
+	{
+		Element.LastOffset = NextStaggerPoint - Element.OriginalFrame;
+		Element.LayerBarModel->Offset(Element.LastOffset);
+
+		++CurrentGroupNum;
+
+		// Advance stagger point if the count has been reached for the current group
+		if (CurrentGroupNum == Settings->ToolOptions.Grouping)
+		{
+			NextStaggerPoint = FindNextStaggerLocation(Element, CurrentElementNum, FirstStaggerPoint, NextStaggerPoint);
+
+			CurrentGroupNum = 0;
+		}
+
+		++CurrentElementNum;
+	}
+}
+
+void FAvaStaggerTool::StaggerKeyFrames(const TArrayView<FStaggerKeyFrameElement>& InOriginalElements)
+{
+	const int32 OriginalElementCount = InOriginalElements.Num();
+	if (OriginalElementCount <= 1)
+	{
+		return;
+	}
+
+	// Stagger all elements
+	FScopedTransaction Transaction(LOCTEXT("StaggerKeyFrames", "Stagger Key Frames"));
+
+	const TRange<FFrameNumber> Range = GetOperationRange(Settings->ToolOptions.Range);
+	const FFrameNumber FirstStaggerPoint = FindFirstStaggerPoint(InOriginalElements, Range);
+
+	FFrameNumber NextStaggerPoint = FirstStaggerPoint;
+	int32 CurrentElementNum = 0;
+	int32 CurrentGroupNum = 0;
+
+	for (FStaggerKeyFrameElement& Element : InOriginalElements)
+	{
+		SetKeyFrameElementTime(Element, NextStaggerPoint);
+
+		++CurrentElementNum;
+		++CurrentGroupNum;
+
+		// Advanced stagger point if the count has been reached for the current group
+		if (CurrentGroupNum == Settings->ToolOptions.Grouping)
+		{
+			NextStaggerPoint = FindNextStaggerLocation(Element, CurrentElementNum, FirstStaggerPoint, NextStaggerPoint);
+
+			CurrentGroupNum = 0;
+		}
+	}
+}
+
+void FAvaStaggerTool::AlignToPlayhead()
+{
+	const TSharedRef<ISequencer> Sequencer = Owner.GetSequencer();
+	const FFrameNumber PlayheadFrameNumber = Sequencer->GetGlobalTime().Time.FrameNumber;
+
+	if (IsLayerBarSelection())
+	{
+		const TArray<TViewModelPtr<FLayerBarModel>> LayerBarModels = GatherAllSelectedLayerBars();
+
+		for (const TViewModelPtr<FLayerBarModel>& LayerBarModel : LayerBarModels)
+		{
+			const TRange<FFrameNumber> Range = LayerBarModel->ComputeRange();
+			const FFrameNumber Offset = PlayheadFrameNumber - Range.GetLowerBoundValue();
+			LayerBarModel->Offset(Offset);
+		}
+	}
+	else if (IsKeyFrameSelection())
+	{
+		const TSharedPtr<FSequencerSelection> SequencerSelection = GetSequencerSelection();
+		if (!SequencerSelection.IsValid())
+		{
+			return;
+		}
+
+		for (const FKeyHandle KeyHandle : SequencerSelection->KeySelection)
+		{
+			if (KeyHandle != FKeyHandle::Invalid())
+			{
+				if (const TViewModelPtr<FChannelModel> KeyModel = SequencerSelection->KeySelection.GetModelForKey(KeyHandle))
+				{
+					const FStaggerKeyFrameElement Element(KeyModel, KeyHandle);
+					SetKeyFrameElementTime(Element, PlayheadFrameNumber);
+				}
+			}
+		}
+	}
+}
+
+void FAvaStaggerTool::SetKeyFrameElementTime(const FStaggerKeyFrameElement& InElement, const FFrameNumber InKeyTime)
+{
+	if (!InElement.KeyChannelModel.IsValid())
+	{
+		return;
+	}
+
+	UMovieSceneSection* const Section = InElement.KeyChannelModel->GetSection();
+	if (!IsValid(Section))
+	{
+		return;
+	}
+
+	if (Section->TryModify())
+	{
+		InElement.KeyChannelModel->GetKeyArea()->SetKeyTime(InElement.KeyHandle, InKeyTime);
+		Section->ExpandToFrame(InElement.OriginalFrame);
+	}
+}
+
+#undef LOCTEXT_NAMESPACE
diff --git a/Engine/Plugins/Experimental/Avalanche/Source/AvalancheSequencer/Private/StaggerTool/AvaStaggerTool.h b/Engine/Plugins/Experimental/Avalanche/Source/AvalancheSequencer/Private/StaggerTool/AvaStaggerTool.h
new file mode 100644
--- /dev/null
+++ b/Engine/Plugins/Experimental/Avalanche/Source/AvalancheSequencer/Private/StaggerTool/AvaStaggerTool.h
@@ -0,0 +1,136 @@
+// Copyright Epic Games, Inc. All Rights Reserved.
+
+#pragma once
+
+#include "AvaStaggerToolSettings.h"
+#include "Commands/AvaSequencerAction.h"
+#include "Containers/ArrayView.h"
+#include "Containers/ContainersFwd.h"
+#include "MVVM/ViewModelPtr.h"
+#include "Math/Range.h"
+#include "Misc/FrameNumber.h"
+#include "SequencerCoreFwd.h"
+
+class SAvaStaggerTool;
+class SWindow;
+class UAvaStaggerCurve;
+
+namespace UE::Sequencer
+{
+	class FChannelModel;
+	class FLayerBarModel;
+}
+
+class FAvaStaggerTool : public FAvaSequencerAction
+{
+public:
+	explicit FAvaStaggerTool(FAvaSequencer& InOwner);
+	virtual ~FAvaStaggerTool();
+
+	TObjectPtr<UAvaSequencerStaggerSettings> GetSettings() const;
+
+	bool HasValidSelection() const;
+
+	bool IsLayerBarSelection() const;
+	bool IsKeyFrameSelection() const;
+
+	int32 GetSelectionCount() const;
+
+	void Stagger();
+
+	bool IsAutoAppying() const;
+
+	void OnSequencerSelectionChanged();
+
+	void AlignToPlayhead();
+
+private:
+	struct FStaggerLayerBarElement
+	{
+		explicit FStaggerLayerBarElement(UE::Sequencer::TViewModelPtr<UE::Sequencer::FLayerBarModel> InLayerBarModel);
+
+		UE::Sequencer::TViewModelPtr<UE::Sequencer::FLayerBarModel> LayerBarModel;
+		FFrameNumber OriginalFrame;
+		FFrameNumber LastOffset;
+	};
+
+	struct FStaggerKeyFrameElement
+	{
+		explicit FStaggerKeyFrameElement(UE::Sequencer::TViewModelPtr<UE::Sequencer::FChannelModel> InKeyChannelModel
+			, FKeyHandle InKeyHandle, FFrameNumber InOriginalFrame = 0);
+
+		UE::Sequencer::TViewModelPtr<UE::Sequencer::FChannelModel> KeyChannelModel;
+		FKeyHandle KeyHandle;
+		FFrameNumber OriginalFrame;
+	};
+
+	//~ Begin FAvaSequencerAction
+	virtual void MapAction(const TSharedRef<FUICommandList>& InCommandList) override;
+	//~ End FAvaSequencerAction
+
+	bool CanExecute() const;
+	void Execute();
+
+	void BindToSequencerSelectionChange();
+	void UnbindToSequencerSelectionChange();
+
+	void OnResetToDefaults();
+	void OnSettingChange(const FName InPropertyName);
+	void OnApply();
+
+	FFrameNumber ConvertTime(FFrameTime InSourceTime) const;
+	TRange<FFrameNumber> GetOperationRange(const EAvaSequencerStaggerRange InRange) const;
+	FFrameNumber GetInterval(const FFrameNumber InRangeSize, const int32 InElementCount, const FFrameNumber InElementSize) const;
+	FFrameNumber GetLayerBarOperationOffset(const UE::Sequencer::TViewModelPtr<UE::Sequencer::FLayerBarModel>& InLayerBarModel) const;
+	static FFrameNumber CalculateLayerBarFrameSpan(const TArray<FStaggerLayerBarElement>& InStaggerElements, const int32 InStopIndex = INDEX_NONE);
+
+	void CacheOriginalElements();
+
+	FFrameNumber FindFirstStaggerPoint(const TArrayView<FStaggerLayerBarElement>& InStaggerElements, const TRange<FFrameNumber>& InRange) const;
+	FFrameNumber FindFirstStaggerPoint(const TArrayView<FStaggerKeyFrameElement>& InStaggerElements, const TRange<FFrameNumber>& InRange) const;
+
+	TArray<UE::Sequencer::TViewModelPtr<UE::Sequencer::FLayerBarModel>> GatherAllSelectedLayerBars() const;
+
+	TArray<FStaggerLayerBarElement> GatherSelectionLayerBarElements() const;
+	TArray<FStaggerKeyFrameElement> GatherSelectionKeyFrameElements() const;
+
+	FFrameNumber CalculateLocalCurveOffset(const float InCurveTime, const TRange<FFrameNumber>& InRange) const;
+
+	FFrameNumber FindNextStaggerLocation(const FStaggerLayerBarElement& InElement
+		, const int32 InElementIndex
+		, const FFrameNumber InFirstFrame
+		, const FFrameNumber InCurrentFrame) const;
+	FFrameNumber FindNextStaggerLocation(const FStaggerKeyFrameElement& InElement
+		, const int32 InElementIndex
+		, const FFrameNumber InFirstFrame
+		, const FFrameNumber InCurrentFrame) const;
+
+	void StaggerLayers(const TArrayView<FStaggerLayerBarElement>& InOriginalElements);
+    void StaggerKeyFrames(const TArrayView<FStaggerKeyFrameElement>& InOriginalElements);
+
+	static void SetKeyFrameElementTime(const FStaggerKeyFrameElement& InElement, const FFrameNumber InKeyTime);
+
+	template<typename InElementType>
+	static void FindFirstFrameInTimeline(const TArrayView<InElementType>& InElements, FFrameNumber& OutFrame)
+	{
+		bool bFound = false;
+		for (const InElementType& Element : InElements)
+		{
+			if (!bFound || Element.OriginalFrame < OutFrame)
+			{
+				OutFrame = Element.OriginalFrame;
+				bFound = true;
+			}
+		}
+	}
+
+	TSharedPtr<SWindow> ToolWindow;
+
+	TSharedPtr<SAvaStaggerTool> ToolWidget;
+
+	TStrongObjectPtr<UAvaSequencerStaggerSettings> Settings;
+
+	// Cached values at the start of auto-updating
+	TArray<FStaggerLayerBarElement> OriginalLayerBarElements;
+	TArray<FStaggerKeyFrameElement> OriginalKeyFrameElements;
+};
diff --git a/Engine/Plugins/Experimental/Avalanche/Source/AvalancheSequencer/Private/StaggerTool/AvaStaggerToolSettings.h b/Engine/Plugins/Experimental/Avalanche/Source/AvalancheSequencer/Private/StaggerTool/AvaStaggerToolSettings.h
new file mode 100644
--- /dev/null
+++ b/Engine/Plugins/Experimental/Avalanche/Source/AvalancheSequencer/Private/StaggerTool/AvaStaggerToolSettings.h
@@ -0,0 +1,148 @@
+// Copyright Epic Games, Inc. All Rights Reserved.
+
+#pragma once
+
+#include "Curves/CurveFloat.h"
+#include "Misc/FrameTime.h"
+#include "AvaStaggerToolSettings.generated.h"
+
+UENUM()
+enum class EAvaSequencerStaggerStartPosition : uint8
+{
+	/** Sequence placement will begin at the position of the first selected */
+	FirstSelected = 0,
+	/** Sequence placement will begin at the earliest sequence of the selection */
+	FirstInTimeline = 1,
+	/** Sequence placement will begin at the current position of the playhead */
+	Playhead = 2,
+
+	Default = FirstSelected UMETA(hidden)
+};
+
+UENUM()
+enum class EAvaSequencerStaggerDistribution : uint8
+{
+	/** Sequences placed one after the other */
+	Increment = 0,
+	/** Sequences placed one before the other */
+	Decrement = 1,
+	/** Sequences spaced out between a range */
+	Range = 2,
+	/** Sequences randomly distributed in range */
+	Random = 255,
+
+	Default = Increment UMETA(hidden)
+};
+
+UENUM()
+enum class EAvaSequencerStaggerRange : uint8
+{
+	/** Sequences distributed between playback range */
+	Playback = 0,
+	/** Sequences distributed between selection range */
+	Selection = 1,
+	/** Sequences distributed between custom user specified frame range */
+	Custom = 2,
+
+	Default = Playback UMETA(hidden)
+};
+
+USTRUCT()
+struct FAvaSequencerStaggerOptions
+{
+	GENERATED_BODY()
+
+	bool operator==(const FAvaSequencerStaggerOptions& InOther) const 
+	{
+		return Distribution == InOther.Distribution
+			&& RandomSeed == InOther.RandomSeed
+			&& Range == InOther.Range
+			&& CustomRange == InOther.CustomRange
+			&& StartPosition == InOther.StartPosition
+			&& OperationPoint == InOther.OperationPoint
+			&& Interval == InOther.Interval
+			&& Shift == InOther.Shift
+			&& Grouping == InOther.Grouping
+			&& bUseCurve == InOther.bUseCurve
+			&& Curve.EditorCurveData == InOther.Curve.EditorCurveData
+			&& Curve.ExternalCurve == InOther.Curve.ExternalCurve
+			&& CurveOffset == InOther.CurveOffset;
+	}
+
+	bool operator!=(const FAvaSequencerStaggerOptions& InOther) const
+	{
+		return !(*this == InOther);
+	}
+
+	UPROPERTY(EditAnywhere, Category = "Stagger Tool Options", meta = (Tooltip = "How the layers or keyframes should be distributed after each step"))
+	EAvaSequencerStaggerDistribution Distribution = EAvaSequencerStaggerDistribution::Default;
+
+	UPROPERTY(EditAnywhere, Category = "Stagger Tool Options", meta = (ClampMin = "0", UIMin = "0", Tooltip = "Seed value for randomization"))
+	int32 RandomSeed = FMath::Rand();
+
+	UPROPERTY(EditAnywhere, Category = "Stagger Tool Options", meta = (Tooltip = "Time range to perform the operation on"))
+	EAvaSequencerStaggerRange Range = EAvaSequencerStaggerRange::Default;
+
+	UPROPERTY(EditAnywhere, Category = "Stagger Tool Options", meta = (UIMin = "-500", UIMax = "500", Tooltip = "Custom time range to perform the operation with"))
+	int32 CustomRange = 30;
+
+	UPROPERTY(EditAnywhere, Category = "Stagger Tool Options", meta = (Tooltip = "Position to start the first layer or keyframe"))
+	EAvaSequencerStaggerStartPosition StartPosition = EAvaSequencerStaggerStartPosition::Default;
+
+	UPROPERTY(EditAnywhere, Category = "Stagger Tool Options", meta = (Tooltip = "Point on the current layer bar at which the next layer is positioned"))
+	float OperationPoint = 0.f;
+
+	UPROPERTY(EditAnywhere, Category = "Stagger Tool Options", meta = (UIMin = "-500", UIMax = "500", Tooltip = "Stepping interval between layers or keyframes"))
+	int32 Interval = 0;
+
+	UPROPERTY(EditAnywhere, Category = "Stagger Tool Options", meta = (UIMin = "-500", UIMax = "500", Tooltip = "Offset of operation"))
+	int32 Shift = 0;
+
+	UPROPERTY(EditAnywhere, Category = "Stagger Tool Options", meta = (ClampMin = "1", ClampMax = "255", UIMin = "1", UIMax = "10", Tooltip = "Number of layers or keyframes to group together before moving to the next stagger point"))
+	int32 Grouping = 1;
+
+	UPROPERTY(EditAnywhere, Category = "Stagger Tool Options", meta = (Tooltip = "Use a curve to layout layers or keyframes"))
+	bool bUseCurve = false;
+
+	UPROPERTY(EditAnywhere, Category = "Stagger Tool Options", meta = (Tooltip = "Curve to use for distribution between layers or keyframes"))
+	FRuntimeFloatCurve Curve;
+
+	UPROPERTY(EditAnywhere, Category = "Stagger Tool Options", meta = (ClampMin = "-1", ClampMax = "1", UIMin = "-1", UIMax = "1", Tooltip = "Offset of operation"))
+	float CurveOffset = 0.f;
+};
+
+UCLASS()
+class UAvaSequencerStaggerSettings : public UObject
+{
+	GENERATED_BODY()
+
+public:
+	UAvaSequencerStaggerSettings()
+	{
+		ResetToolOptions();
+	}
+
+	bool CanResetToolOptions() const
+	{
+		return ToolOptions != FAvaSequencerStaggerOptions();
+	}
+
+	void ResetToolOptions()
+	{
+		int32 PreviousRandomSeed = ToolOptions.RandomSeed;
+		ToolOptions = FAvaSequencerStaggerOptions();
+		ToolOptions.RandomSeed = PreviousRandomSeed;
+
+		if (FRichCurve* const RichCurve = ToolOptions.Curve.GetRichCurve())
+		{
+			RichCurve->AddKey(0.f, 0.f);
+			RichCurve->AddKey(1.f, 0.f);
+		}
+	}
+
+	UPROPERTY()
+	bool bAutoApply = false;
+
+	UPROPERTY(EditAnywhere, Category = "Stagger Tool Options", meta = (ShowOnlyInnerProperties))
+	FAvaSequencerStaggerOptions ToolOptions;
+};
diff --git a/Engine/Plugins/Experimental/Avalanche/Source/AvalancheSequencer/Private/StaggerTool/Customization/AvaStaggerToolSettingsDetailsCustomization.cpp b/Engine/Plugins/Experimental/Avalanche/Source/AvalancheSequencer/Private/StaggerTool/Customization/AvaStaggerToolSettingsDetailsCustomization.cpp
new file mode 100644
--- /dev/null
+++ b/Engine/Plugins/Experimental/Avalanche/Source/AvalancheSequencer/Private/StaggerTool/Customization/AvaStaggerToolSettingsDetailsCustomization.cpp
@@ -0,0 +1,384 @@
+﻿// Copyright Epic Games, Inc. All Rights Reserved.
+
+#include "AvaStaggerToolSettingsDetailsCustomization.h"
+#include "DetailCategoryBuilder.h"
+#include "DetailLayoutBuilder.h"
+#include "DetailWidgetRow.h"
+#include "StaggerTool/AvaStaggerTool.h"
+#include "StaggerTool/Widgets/SAvaStaggerOperationPoint.h"
+#include "StaggerTool/Widgets/SAvaStaggerSettingsRadioGroup.h"
+#include "Widgets/Input/SButton.h"
+#include "Widgets/Input/SSpinBox.h"
+
+#define LOCTEXT_NAMESPACE "AvaStaggerToolSettingsDetailsCustomization"
+
+TSharedRef<IDetailCustomization> FAvaStaggerToolSettingsDetailsCustomization::MakeInstance(const TWeakPtr<FAvaStaggerTool> InWeakTool)
+{
+	return MakeShared<FAvaStaggerToolSettingsDetailsCustomization>(InWeakTool);
+}
+
+FAvaStaggerToolSettingsDetailsCustomization::FAvaStaggerToolSettingsDetailsCustomization(const TWeakPtr<FAvaStaggerTool>& InWeakTool)
+	: WeakTool(InWeakTool)
+{
+}
+
+void FAvaStaggerToolSettingsDetailsCustomization::CustomizeDetails(IDetailLayoutBuilder& DetailBuilder)
+{
+	ToolOptionsProperty = DetailBuilder.GetProperty(GET_MEMBER_NAME_CHECKED(UAvaSequencerStaggerSettings, ToolOptions));
+
+	UseCurveProperty = ToolOptionsProperty->GetChildHandle(GET_MEMBER_NAME_CHECKED(FAvaSequencerStaggerOptions, bUseCurve));
+	CurveProperty = ToolOptionsProperty->GetChildHandle(GET_MEMBER_NAME_CHECKED(FAvaSequencerStaggerOptions, Curve));
+	CurveOffsetProperty = ToolOptionsProperty->GetChildHandle(GET_MEMBER_NAME_CHECKED(FAvaSequencerStaggerOptions, CurveOffset));
+	DistributionProperty = ToolOptionsProperty->GetChildHandle(GET_MEMBER_NAME_CHECKED(FAvaSequencerStaggerOptions, Distribution));
+	RandomSeedProperty = ToolOptionsProperty->GetChildHandle(GET_MEMBER_NAME_CHECKED(FAvaSequencerStaggerOptions, RandomSeed));
+	RangeProperty = ToolOptionsProperty->GetChildHandle(GET_MEMBER_NAME_CHECKED(FAvaSequencerStaggerOptions, Range));
+	CustomRangeProperty = ToolOptionsProperty->GetChildHandle(GET_MEMBER_NAME_CHECKED(FAvaSequencerStaggerOptions, CustomRange));
+	StartPositionProperty = ToolOptionsProperty->GetChildHandle(GET_MEMBER_NAME_CHECKED(FAvaSequencerStaggerOptions, StartPosition));
+	OperationPointProperty = ToolOptionsProperty->GetChildHandle(GET_MEMBER_NAME_CHECKED(FAvaSequencerStaggerOptions, OperationPoint));
+	IntervalProperty = ToolOptionsProperty->GetChildHandle(GET_MEMBER_NAME_CHECKED(FAvaSequencerStaggerOptions, Interval));
+	ShiftProperty = ToolOptionsProperty->GetChildHandle(GET_MEMBER_NAME_CHECKED(FAvaSequencerStaggerOptions, Shift));
+	GroupingProperty = ToolOptionsProperty->GetChildHandle(GET_MEMBER_NAME_CHECKED(FAvaSequencerStaggerOptions, Grouping));
+
+	IDetailCategoryBuilder& OptionsCategory = DetailBuilder.EditCategory(TEXT("Stagger Tool Options"));
+
+	IDetailPropertyRow& DistributionPropertyRow = OptionsCategory.AddProperty(DistributionProperty);
+	AddDistributionRow(DistributionPropertyRow, DistributionProperty.ToSharedRef());
+
+	IDetailPropertyRow& RandomSeedPropertyRow = OptionsCategory.AddProperty(RandomSeedProperty)
+		.Visibility(
+			TAttribute<EVisibility>::Create(
+				TAttribute<EVisibility>::FGetter::CreateSP(this, &FAvaStaggerToolSettingsDetailsCustomization::GetRandomSeedPropertyVisibility)
+			));
+	AddRandomSeedRow(RandomSeedPropertyRow, RandomSeedProperty.ToSharedRef());
+
+	IDetailPropertyRow& RangePropertyRow = OptionsCategory.AddProperty(RangeProperty)
+		.Visibility(
+			TAttribute<EVisibility>::Create(
+				TAttribute<EVisibility>::FGetter::CreateSP(this, &FAvaStaggerToolSettingsDetailsCustomization::GetRangePropertyVisibility)
+			));
+	AddRangeRow(RangePropertyRow, RangeProperty.ToSharedRef());
+
+	IDetailPropertyRow& CustomRangePropertyRow = OptionsCategory.AddProperty(CustomRangeProperty)
+		.Visibility(
+			TAttribute<EVisibility>::Create(
+				TAttribute<EVisibility>::FGetter::CreateSP(this, &FAvaStaggerToolSettingsDetailsCustomization::GetCustomRangePropertyVisibility)
+			));
+	AddCustomIntSpinBoxRow(CustomRangePropertyRow, CustomRangeProperty.ToSharedRef());
+
+	IDetailPropertyRow& StartPositionPropertyRow = OptionsCategory.AddProperty(StartPositionProperty)
+		.Visibility(
+			TAttribute<EVisibility>::Create(
+				TAttribute<EVisibility>::FGetter::CreateSP(this, &FAvaStaggerToolSettingsDetailsCustomization::GetStartPositionPropertyVisibility)
+			));
+	AddStartPositionRow(StartPositionPropertyRow, StartPositionProperty.ToSharedRef());
+
+	IDetailPropertyRow& OperationPointPropertyRow = OptionsCategory.AddProperty(OperationPointProperty);
+	AddOperationPointRow(OperationPointPropertyRow, OperationPointProperty.ToSharedRef());
+
+	IDetailPropertyRow& IntervalPropertyRow = OptionsCategory.AddProperty(IntervalProperty)
+		.Visibility(
+			TAttribute<EVisibility>::Create(
+				TAttribute<EVisibility>::FGetter::CreateSP(this, &FAvaStaggerToolSettingsDetailsCustomization::GetIntervalPropertyVisibility)
+			));
+	AddCustomIntSpinBoxRow(IntervalPropertyRow, IntervalProperty.ToSharedRef());
+
+	IDetailPropertyRow& ShiftPropertyRow = OptionsCategory.AddProperty(ShiftProperty);
+	AddCustomIntSpinBoxRow(ShiftPropertyRow, ShiftProperty.ToSharedRef());
+
+	IDetailPropertyRow& GroupingPropertyRow = OptionsCategory.AddProperty(GroupingProperty);
+	AddCustomIntSpinBoxRow(GroupingPropertyRow, GroupingProperty.ToSharedRef());
+
+	OptionsCategory.AddProperty(UseCurveProperty);
+
+	OptionsCategory.AddProperty(CurveProperty)
+		.Visibility(
+			TAttribute<EVisibility>::Create(
+				TAttribute<EVisibility>::FGetter::CreateSP(this, &FAvaStaggerToolSettingsDetailsCustomization::GetCurvePropertyVisibility)
+			));
+
+	IDetailPropertyRow& CurveOffsetPropertyRow = OptionsCategory.AddProperty(CurveOffsetProperty)
+		.Visibility(
+			TAttribute<EVisibility>::Create(
+				TAttribute<EVisibility>::FGetter::CreateSP(this, &FAvaStaggerToolSettingsDetailsCustomization::GetCurveOffsetPropertyVisibility)
+			));
+	AddCustomFloatSpinBoxRow(CurveOffsetPropertyRow, CurveOffsetProperty.ToSharedRef());
+}
+
+EAvaSequencerStaggerDistribution FAvaStaggerToolSettingsDetailsCustomization::GetDistributionPropertyValue() const
+{
+	uint8 Value = 0;
+	DistributionProperty->GetValue(Value);
+	return static_cast<EAvaSequencerStaggerDistribution>(Value);
+}
+
+EAvaSequencerStaggerRange FAvaStaggerToolSettingsDetailsCustomization::GetRangePropertyValue() const
+{
+	uint8 Value = 0;
+	RangeProperty->GetValue(Value);
+	return static_cast<EAvaSequencerStaggerRange>(Value);
+}
+
+EAvaSequencerStaggerStartPosition FAvaStaggerToolSettingsDetailsCustomization::GetStartPositionPropertyValue() const
+{
+	uint8 Value = 0;
+	StartPositionProperty->GetValue(Value);
+	return static_cast<EAvaSequencerStaggerStartPosition>(Value);
+}
+
+float FAvaStaggerToolSettingsDetailsCustomization::GetOperationPointPropertyValue() const
+{
+	float Value = 0;
+	OperationPointProperty->GetValue(Value);
+	return Value;
+}
+
+bool FAvaStaggerToolSettingsDetailsCustomization::GetUseCurvePropertyValue() const
+{
+	bool Value = false;
+	UseCurveProperty->GetValue(Value);
+	return Value;
+}
+
+EVisibility FAvaStaggerToolSettingsDetailsCustomization::GetCurvePropertyVisibility() const
+{
+	return GetUseCurvePropertyValue() ? EVisibility::Visible : EVisibility::Collapsed;
+}
+
+EVisibility FAvaStaggerToolSettingsDetailsCustomization::GetCurveOffsetPropertyVisibility() const
+{
+	return GetUseCurvePropertyValue() ? EVisibility::Visible : EVisibility::Collapsed;
+}
+
+EVisibility FAvaStaggerToolSettingsDetailsCustomization::GetRandomSeedPropertyVisibility() const
+{
+	const EAvaSequencerStaggerDistribution DistributionValue = GetDistributionPropertyValue();
+	return DistributionValue == EAvaSequencerStaggerDistribution::Random ? EVisibility::Visible : EVisibility::Collapsed;
+}
+
+EVisibility FAvaStaggerToolSettingsDetailsCustomization::GetIntervalPropertyVisibility() const
+{
+	return EVisibility::Visible;
+}
+
+EVisibility FAvaStaggerToolSettingsDetailsCustomization::GetShiftPropertyVisibility() const
+{
+	const EAvaSequencerStaggerDistribution DistributionValue = GetDistributionPropertyValue();
+	return (DistributionValue != EAvaSequencerStaggerDistribution::Range) ? EVisibility::Visible : EVisibility::Collapsed;
+}
+
+EVisibility FAvaStaggerToolSettingsDetailsCustomization::GetRangePropertyVisibility() const
+{
+	const EAvaSequencerStaggerDistribution DistributionValue = GetDistributionPropertyValue();
+	return (GetUseCurvePropertyValue() || DistributionValue == EAvaSequencerStaggerDistribution::Range || DistributionValue == EAvaSequencerStaggerDistribution::Random)
+		? EVisibility::Visible : EVisibility::Collapsed;
+}
+
+EVisibility FAvaStaggerToolSettingsDetailsCustomization::GetCustomRangePropertyVisibility() const
+{
+	const EAvaSequencerStaggerRange RangeValue = GetRangePropertyValue();
+	return (GetRangePropertyVisibility() == EVisibility::Visible && RangeValue == EAvaSequencerStaggerRange::Custom)
+		? EVisibility::Visible : EVisibility::Collapsed;
+}
+
+EVisibility FAvaStaggerToolSettingsDetailsCustomization::GetStartPositionPropertyVisibility() const
+{
+	const EAvaSequencerStaggerDistribution TypeValue = GetDistributionPropertyValue();
+	return (TypeValue != EAvaSequencerStaggerDistribution::Range)
+		? EVisibility::Visible : EVisibility::Collapsed;
+}
+
+EVisibility FAvaStaggerToolSettingsDetailsCustomization::GetOperationPointPropertyVisibility() const
+{
+	return (WeakTool.IsValid() && WeakTool.Pin()->IsLayerBarSelection())
+		? EVisibility::Visible : EVisibility::Collapsed;
+}
+
+template<typename InValueType>
+TSharedRef<SWidget> CreateSpinBox(const TWeakPtr<FAvaStaggerTool>& InWeakTool, IDetailPropertyRow& PropertyRow, const TSharedRef<IPropertyHandle> InProperty)
+{
+	const TWeakPtr<IPropertyHandle> WeakProperty = InProperty;
+
+	auto GetMetaDataValue = [&InProperty](const FName InKey) -> TOptional<InValueType>
+	{
+		if (std::is_integral_v<InValueType>)
+		{
+			return InProperty->HasMetaData(InKey) ? InProperty->GetIntMetaData(InKey) : TOptional<InValueType>();	
+		}
+		if (std::is_floating_point_v<InValueType>)
+		{
+			return InProperty->HasMetaData(InKey) ? InProperty->GetFloatMetaData(InKey) : TOptional<InValueType>();	
+		}
+		return TOptional<InValueType>();	
+	};
+	const TOptional<InValueType> ClampMinValue = GetMetaDataValue(TEXT("ClampMin"));
+	const TOptional<InValueType> ClampMaxValue = GetMetaDataValue(TEXT("ClampMax"));
+	const TOptional<InValueType> UIMinValue = GetMetaDataValue(TEXT("UIMin"));
+	const TOptional<InValueType> UIMaxValue = GetMetaDataValue(TEXT("UIMax"));
+
+	return SNew(SSpinBox<InValueType>)
+		.Font(IDetailLayoutBuilder::GetDetailFont())
+		.MinValue(ClampMinValue)
+		.MaxValue(ClampMaxValue)
+		.MinSliderValue(UIMinValue)
+		.MaxSliderValue(UIMaxValue)
+		.Value_Lambda([WeakProperty]()
+			{
+				InValueType OutValue = 0;
+				if (const TSharedPtr<IPropertyHandle> Property = WeakProperty.Pin())
+				{
+					Property->GetValue(OutValue);
+				}
+				return OutValue;
+			})
+		.OnValueChanged_Lambda([WeakProperty](const InValueType InNewValue)
+			{
+				if (const TSharedPtr<IPropertyHandle> Property = WeakProperty.Pin())
+				{
+					Property->SetValue(InNewValue);
+				}
+			})
+		.OnValueCommitted_Lambda([InWeakTool, WeakProperty](const InValueType InNewValue, const ETextCommit::Type InCommitType)
+			{
+				if (const TSharedPtr<IPropertyHandle> Property = WeakProperty.Pin())
+				{
+					Property->SetValue(InNewValue);
+				}
+			});
+}
+
+void FAvaStaggerToolSettingsDetailsCustomization::AddCustomIntSpinBoxRow(IDetailPropertyRow& PropertyRow, const TSharedRef<IPropertyHandle> InProperty)
+{
+	const TWeakPtr<IPropertyHandle> WeakProperty = InProperty;
+
+	auto GetMetaDataValue = [&InProperty](const FName InKey)
+	{
+		return InProperty->HasMetaData(InKey) ? InProperty->GetIntMetaData(InKey) : TOptional<int32>();	
+	};
+	const TOptional<int32> ClampMinValue = GetMetaDataValue(TEXT("ClampMin"));
+	const TOptional<int32> ClampMaxValue = GetMetaDataValue(TEXT("ClampMax"));
+	const TOptional<int32> UIMinValue = GetMetaDataValue(TEXT("UIMin"));
+	const TOptional<int32> UIMaxValue = GetMetaDataValue(TEXT("UIMax"));
+
+	PropertyRow.CustomWidget()
+		.NameContent()
+		[
+			InProperty->CreatePropertyNameWidget()
+		]
+		.ValueContent()
+		[
+			CreateSpinBox<int32>(WeakTool, PropertyRow, InProperty)
+		];
+}
+
+void FAvaStaggerToolSettingsDetailsCustomization::AddCustomFloatSpinBoxRow(IDetailPropertyRow& PropertyRow, const TSharedRef<IPropertyHandle> InProperty)
+{
+	const TWeakPtr<IPropertyHandle> WeakProperty = InProperty;
+
+	auto GetMetaDataValue = [&InProperty](const FName InKey)
+	{
+		return InProperty->HasMetaData(InKey) ? InProperty->GetFloatMetaData(InKey) : TOptional<float>();	
+	};
+	const TOptional<float> ClampMinValue = GetMetaDataValue(TEXT("ClampMin"));
+	const TOptional<float> ClampMaxValue = GetMetaDataValue(TEXT("ClampMax"));
+	const TOptional<float> UIMinValue = GetMetaDataValue(TEXT("UIMin"));
+	const TOptional<float> UIMaxValue = GetMetaDataValue(TEXT("UIMax"));
+
+	PropertyRow.CustomWidget()
+		.NameContent()
+		[
+			InProperty->CreatePropertyNameWidget()
+		]
+		.ValueContent()
+		[
+			CreateSpinBox<float>(WeakTool, PropertyRow, InProperty)
+		];
+}
+
+void FAvaStaggerToolSettingsDetailsCustomization::AddDistributionRow(IDetailPropertyRow& PropertyRow, const TSharedRef<IPropertyHandle>& InProperty)
+{
+	PropertyRow.CustomWidget()
+		.NameContent()
+		[
+			InProperty->CreatePropertyNameWidget()
+		]
+		.ValueContent()
+		[
+			SNew(SAvaStaggerSettingsRadioGroup<EAvaSequencerStaggerDistribution>, InProperty)
+		];
+}
+
+void FAvaStaggerToolSettingsDetailsCustomization::AddRandomSeedRow(IDetailPropertyRow& PropertyRow, const TSharedRef<IPropertyHandle>& InProperty)
+{
+	PropertyRow.CustomWidget()
+		.NameContent()
+		[
+			InProperty->CreatePropertyNameWidget()
+		]
+		.ValueContent()
+		[
+			SNew(SHorizontalBox)
+			+ SHorizontalBox::Slot()
+			.HAlign(HAlign_Fill)
+			.VAlign(VAlign_Center)
+			.Padding(FMargin(0.f, 0.f, 5.f, 0.f))
+			[
+				CreateSpinBox<int32>(WeakTool, PropertyRow, InProperty)
+			]
+			+ SHorizontalBox::Slot()
+			.AutoWidth()
+			.HAlign(HAlign_Right)
+			.VAlign(VAlign_Center)
+			[
+				SNew(SButton)
+				.TextStyle(FAppStyle::Get(), TEXT("SmallText"))
+				.Text(LOCTEXT("GenerateRandomSeedLabel", "Generate"))
+				.OnClicked_Lambda([this]()
+					{
+						RandomSeedProperty->SetValue(FMath::Rand());
+						return FReply::Handled();
+					})
+			]
+		];
+}
+
+void FAvaStaggerToolSettingsDetailsCustomization::AddRangeRow(IDetailPropertyRow& PropertyRow, const TSharedRef<IPropertyHandle>& InProperty)
+{
+	PropertyRow.CustomWidget()
+		.NameContent()
+		[
+			InProperty->CreatePropertyNameWidget()
+		]
+		.ValueContent()
+		[
+			SNew(SAvaStaggerSettingsRadioGroup<EAvaSequencerStaggerRange>, InProperty)
+		];
+}
+
+void FAvaStaggerToolSettingsDetailsCustomization::AddStartPositionRow(IDetailPropertyRow& PropertyRow, const TSharedRef<IPropertyHandle>& InProperty)
+{
+	PropertyRow.CustomWidget()
+		.NameContent()
+		[
+			InProperty->CreatePropertyNameWidget()
+		]
+		.ValueContent()
+		[
+			SNew(SAvaStaggerSettingsRadioGroup<EAvaSequencerStaggerStartPosition>, InProperty)
+		];
+}
+
+void FAvaStaggerToolSettingsDetailsCustomization::AddOperationPointRow(IDetailPropertyRow& PropertyRow, const TSharedRef<IPropertyHandle>& InProperty)
+{
+	PropertyRow.CustomWidget()
+		.NameContent()
+		[
+			InProperty->CreatePropertyNameWidget()
+		]
+		.ValueContent()
+		.HAlign(HAlign_Fill)
+		[
+			SNew(SAvaStaggerOperationPoint, InProperty)
+		];
+}
+
+#undef LOCTEXT_NAMESPACE
diff --git a/Engine/Plugins/Experimental/Avalanche/Source/AvalancheSequencer/Private/StaggerTool/Customization/AvaStaggerToolSettingsDetailsCustomization.h b/Engine/Plugins/Experimental/Avalanche/Source/AvalancheSequencer/Private/StaggerTool/Customization/AvaStaggerToolSettingsDetailsCustomization.h
new file mode 100644
--- /dev/null
+++ b/Engine/Plugins/Experimental/Avalanche/Source/AvalancheSequencer/Private/StaggerTool/Customization/AvaStaggerToolSettingsDetailsCustomization.h
@@ -0,0 +1,67 @@
+﻿// Copyright Epic Games, Inc. All Rights Reserved.
+
+#pragma once
+
+#include "IDetailCustomization.h"
+#include "Layout/Visibility.h"
+#include "StaggerTool/AvaStaggerToolSettings.h"
+
+class FAvaStaggerTool;
+class IDetailPropertyRow;
+class IPropertyHandle;
+
+class FAvaStaggerToolSettingsDetailsCustomization : public IDetailCustomization
+{
+public:
+	static TSharedRef<IDetailCustomization> MakeInstance(const TWeakPtr<FAvaStaggerTool> InWeakTool);
+
+	FAvaStaggerToolSettingsDetailsCustomization(const TWeakPtr<FAvaStaggerTool>& InWeakTool);
+	virtual ~FAvaStaggerToolSettingsDetailsCustomization() override {}
+
+	//~ Begin IDetailCustomization
+	virtual void CustomizeDetails(IDetailLayoutBuilder& DetailBuilder) override;
+	//~ End IDetailCustomization
+
+private:
+	EAvaSequencerStaggerDistribution GetDistributionPropertyValue() const;
+	EAvaSequencerStaggerRange GetRangePropertyValue() const;
+	EAvaSequencerStaggerStartPosition GetStartPositionPropertyValue() const;
+	float GetOperationPointPropertyValue() const;
+	bool GetUseCurvePropertyValue() const;
+
+	EVisibility GetCurvePropertyVisibility() const;
+	EVisibility GetCurveOffsetPropertyVisibility() const;
+	EVisibility GetRandomSeedPropertyVisibility() const;
+	EVisibility GetIntervalPropertyVisibility() const;
+	EVisibility GetShiftPropertyVisibility() const;
+	EVisibility GetRangePropertyVisibility() const;
+	EVisibility GetCustomRangePropertyVisibility() const;
+	EVisibility GetStartPositionPropertyVisibility() const;
+	EVisibility GetOperationPointPropertyVisibility() const;
+
+	void AddCustomIntSpinBoxRow(IDetailPropertyRow& PropertyRow, const TSharedRef<IPropertyHandle> InProperty);
+	void AddCustomFloatSpinBoxRow(IDetailPropertyRow& PropertyRow, const TSharedRef<IPropertyHandle> InProperty);
+
+	void AddDistributionRow(IDetailPropertyRow& PropertyRow, const TSharedRef<IPropertyHandle>& InProperty);
+	void AddRandomSeedRow(IDetailPropertyRow& PropertyRow, const TSharedRef<IPropertyHandle>& InProperty);
+	void AddRangeRow(IDetailPropertyRow& PropertyRow, const TSharedRef<IPropertyHandle>& InProperty);
+	void AddStartPositionRow(IDetailPropertyRow& PropertyRow, const TSharedRef<IPropertyHandle>& InProperty);
+	void AddOperationPointRow(IDetailPropertyRow& PropertyRow, const TSharedRef<IPropertyHandle>& InProperty);
+
+	TWeakPtr<FAvaStaggerTool> WeakTool;
+
+	TSharedPtr<IPropertyHandle> ToolOptionsProperty;
+
+	TSharedPtr<IPropertyHandle> UseCurveProperty;
+	TSharedPtr<IPropertyHandle> CurveProperty;
+	TSharedPtr<IPropertyHandle> CurveOffsetProperty;
+	TSharedPtr<IPropertyHandle> DistributionProperty;
+	TSharedPtr<IPropertyHandle> RandomSeedProperty;
+	TSharedPtr<IPropertyHandle> RangeProperty;
+	TSharedPtr<IPropertyHandle> CustomRangeProperty;
+	TSharedPtr<IPropertyHandle> StartPositionProperty;
+	TSharedPtr<IPropertyHandle> OperationPointProperty;
+	TSharedPtr<IPropertyHandle> IntervalProperty;
+	TSharedPtr<IPropertyHandle> ShiftProperty;
+	TSharedPtr<IPropertyHandle> GroupingProperty;
+};
diff --git a/Engine/Plugins/Experimental/Avalanche/Source/AvalancheSequencer/Private/StaggerTool/Widgets/SAvaStaggerOperationPoint.cpp b/Engine/Plugins/Experimental/Avalanche/Source/AvalancheSequencer/Private/StaggerTool/Widgets/SAvaStaggerOperationPoint.cpp
new file mode 100644
--- /dev/null
+++ b/Engine/Plugins/Experimental/Avalanche/Source/AvalancheSequencer/Private/StaggerTool/Widgets/SAvaStaggerOperationPoint.cpp
@@ -0,0 +1,124 @@
+﻿// Copyright Epic Games, Inc. All Rights Reserved.
+
+#include "SAvaStaggerOperationPoint.h"
+#include "DetailLayoutBuilder.h"
+#include "Widgets/Input/SCheckBox.h"
+#include "Widgets/Input/SSpinBox.h"
+#include "Widgets/Layout/SBox.h"
+#include "Widgets/Text/STextBlock.h"
+
+#define LOCTEXT_NAMESPACE "SAvaStaggerOperationPoint"
+
+void SAvaStaggerOperationPoint::Construct(const FArguments& InArgs, const TWeakPtr<IPropertyHandle>& InWeakProperty)
+{
+	WeakProperty = InWeakProperty;
+
+	const TSharedRef<SCheckBox> LeftButton = CreateButton(0.f, TEXT("HorizontalAlignment_Left"));
+	Buttons.Add(LeftButton);
+
+	const TSharedRef<SCheckBox> CenterButton = CreateButton(0.5f, TEXT("HorizontalAlignment_Center"));
+	Buttons.Add(CenterButton);
+
+	const TSharedRef<SCheckBox> RightButton = CreateButton(1.f, TEXT("HorizontalAlignment_Right"));
+	Buttons.Add(RightButton);
+
+	ChildSlot
+		[
+			SNew(SHorizontalBox)
+			+ SHorizontalBox::Slot()
+			.AutoWidth()
+			.VAlign(VAlign_Center)
+			.Padding(FMargin(1.f, 0.f))
+			[
+				SNew(SBox)
+				.WidthOverride(124.f)
+				[
+					SNew(SSpinBox<float>)
+					.Font(IDetailLayoutBuilder::GetDetailFont())
+					.MinValue(0.f)
+					.MaxValue(1.f)
+					.MinSliderValue(0.f)
+					.MaxSliderValue(1.f)
+					.Value_Lambda([this]()
+						{
+							float OutValue = 0;
+							if (const TSharedPtr<IPropertyHandle> Property = WeakProperty.Pin())
+							{
+								Property->GetValue(OutValue);
+							}
+							return OutValue;
+						})
+					.OnValueChanged_Lambda([this](const float InNewValue)
+						{
+							if (const TSharedPtr<IPropertyHandle> Property = WeakProperty.Pin())
+							{
+								Property->SetValue(InNewValue);
+							}
+						})
+					.OnValueCommitted_Lambda([this](const float InNewValue, const ETextCommit::Type InCommitType)
+						{
+							if (const TSharedPtr<IPropertyHandle> Property = WeakProperty.Pin())
+							{
+								Property->SetValue(InNewValue);
+							}
+						})
+				]
+			]
+			+ SHorizontalBox::Slot()
+			.AutoWidth()
+			.VAlign(VAlign_Center)
+			.Padding(FMargin(1.f, 0.f))
+			[
+				LeftButton
+			]
+			+ SHorizontalBox::Slot()
+			.AutoWidth()
+			.VAlign(VAlign_Center)
+			.Padding(FMargin(1.f, 0.f))
+			[
+				CenterButton
+			]
+			+ SHorizontalBox::Slot()
+			.AutoWidth()
+			.VAlign(VAlign_Center)
+			.Padding(FMargin(1.f, 0.f))
+			[
+				RightButton
+			]
+		];
+}
+
+TSharedRef<SCheckBox> SAvaStaggerOperationPoint::CreateButton(const float InValue, const FName InImageBrushName)
+{
+	return SNew(SCheckBox)
+		.Type(ESlateCheckBoxType::ToggleButton)
+		.Style(&FAppStyle::Get().GetWidgetStyle<FCheckBoxStyle>(TEXT("Sequencer.Outliner.ToggleButton")))
+		.HAlign(HAlign_Center)
+		.Padding(FMargin(4.f, 3.f))
+		.IsChecked_Lambda([this, InValue]() -> ECheckBoxState
+			{
+				float Value = 0.f;
+				WeakProperty.Pin()->GetValue(Value);
+				return InValue == Value ? ECheckBoxState::Checked : ECheckBoxState::Unchecked;
+			})
+		.OnCheckStateChanged_Lambda([this, InValue](const ECheckBoxState InCheckBoxState)
+			{
+				if (const TSharedPtr<IPropertyHandle> Property = WeakProperty.Pin())
+				{
+					Property->SetValue(InValue);
+				}
+			})
+		[
+			SNew(SBox)
+			.WidthOverride(14.f)
+			.HeightOverride(14.f)
+			.HAlign(HAlign_Center)
+			.VAlign(VAlign_Center)
+			[
+				SNew(SImage)
+				.Image(FAppStyle::GetBrush(InImageBrushName))
+			]
+		];
+}
+
+#undef LOCTEXT_NAMESPACE
diff --git a/Engine/Plugins/Experimental/Avalanche/Source/AvalancheSequencer/Private/StaggerTool/Widgets/SAvaStaggerOperationPoint.h b/Engine/Plugins/Experimental/Avalanche/Source/AvalancheSequencer/Private/StaggerTool/Widgets/SAvaStaggerOperationPoint.h
new file mode 100644
--- /dev/null
+++ b/Engine/Plugins/Experimental/Avalanche/Source/AvalancheSequencer/Private/StaggerTool/Widgets/SAvaStaggerOperationPoint.h
@@ -0,0 +1,25 @@
+﻿// Copyright Epic Games, Inc. All Rights Reserved.
+
+#pragma once
+
+#include "Widgets/SCompoundWidget.h"
+
+class IPropertyHandle;
+class SCheckBox;
+
+class SAvaStaggerOperationPoint : public SCompoundWidget
+{
+public:
+	SLATE_BEGIN_ARGS(SAvaStaggerOperationPoint) {}
+	SLATE_END_ARGS()
+
+	void Construct(const FArguments& InArgs, const TWeakPtr<IPropertyHandle>& InWeakProperty);
+
+protected:
+	TSharedRef<SCheckBox> CreateButton(const float InValue, const FName InImageBrushName);
+
+private:
+	TWeakPtr<IPropertyHandle> WeakProperty;
+
+	TArray<TSharedRef<SCheckBox>> Buttons;
+};
diff --git a/Engine/Plugins/Experimental/Avalanche/Source/AvalancheSequencer/Private/StaggerTool/Widgets/SAvaStaggerSettingsRadioGroup.h b/Engine/Plugins/Experimental/Avalanche/Source/AvalancheSequencer/Private/StaggerTool/Widgets/SAvaStaggerSettingsRadioGroup.h
new file mode 100644
--- /dev/null
+++ b/Engine/Plugins/Experimental/Avalanche/Source/AvalancheSequencer/Private/StaggerTool/Widgets/SAvaStaggerSettingsRadioGroup.h
@@ -0,0 +1,110 @@
+﻿// Copyright Epic Games, Inc. All Rights Reserved.
+
+#pragma once
+
+#include "DetailLayoutBuilder.h"
+#include "Widgets/Input/SCheckBox.h"
+#include "Widgets/Layout/SBox.h"
+#include "Widgets/SCompoundWidget.h"
+#include "Widgets/Text/STextBlock.h"
+
+template<typename InEnumType>
+class SAvaStaggerSettingsRadioGroup : public SCompoundWidget
+{
+public:
+	DECLARE_DELEGATE_OneParam(FEnumValueChangedDelegate, const InEnumType);
+
+	SLATE_BEGIN_ARGS(SAvaStaggerSettingsRadioGroup) {}
+		SLATE_EVENT(FEnumValueChangedDelegate, OnValueChanged)
+	SLATE_END_ARGS()
+
+	void Construct(const FArguments& InArgs, const TWeakPtr<IPropertyHandle>& InWeakProperty)
+	{
+		const UEnum* const Enum = StaticEnum<InEnumType>();
+		check(Enum);
+
+		WeakProperty = InWeakProperty;
+
+		OnValueChanged = InArgs._OnValueChanged;
+
+		const TSharedRef<SHorizontalBox> HorizontalBox = SNew(SHorizontalBox);
+
+		for (int32 Index = 0; Index < Enum->NumEnums() - 1; ++Index)
+		{
+			const bool bIsHidden = Enum->HasMetaData(TEXT("Hidden"), Index);
+			if (!bIsHidden)
+			{
+				const TSharedRef<SCheckBox> NewButton = CreateButton(static_cast<InEnumType>(Enum->GetValueByIndex(Index))
+					, Enum->GetDisplayNameTextByIndex(Index)
+					, Enum->GetToolTipTextByIndex(Index));
+				Buttons.Add(NewButton);
+
+				HorizontalBox->AddSlot()
+					.VAlign(VAlign_Center)
+					.Padding(FMargin(1.f, 0.f))
+					[
+						NewButton
+					];
+			}
+		}
+
+		ChildSlot
+			[
+				HorizontalBox
+			];
+	}
+
+protected:
+	TSharedRef<SCheckBox> CreateButton(const InEnumType InValue, const FText& InText, const FText& InToolTipText)
+	{
+		return SNew(SCheckBox)
+			.Type(ESlateCheckBoxType::ToggleButton)
+			.Style(&FAppStyle::Get().GetWidgetStyle<FCheckBoxStyle>(TEXT("Sequencer.Outliner.ToggleButton")))
+			.HAlign(HAlign_Center)
+			.Padding(FMargin(4.f, 3.f))
+			.ToolTipText(InToolTipText)
+			.IsChecked_Lambda([this, InValue]() -> ECheckBoxState
+				{
+					uint8 Value = 0;
+					WeakProperty.Pin()->GetValue(Value);
+					return InValue == static_cast<InEnumType>(Value) ? ECheckBoxState::Checked : ECheckBoxState::Unchecked;
+				})
+			.OnCheckStateChanged_Lambda([this, InValue](const ECheckBoxState InCheckBoxState)
+				{
+					if (const TSharedPtr<IPropertyHandle> Property = WeakProperty.Pin())
+					{
+						Property->SetValue(static_cast<uint8>(InValue));
+					}
+					OnValueChanged.ExecuteIfBound(InValue);
+				})
+			[
+				SNew(SBox)
+				.WidthOverride(80.f)
+				.HAlign(HAlign_Center)
+				.VAlign(VAlign_Center)
+				[
+					SNew(STextBlock)
+					.Font(IDetailLayoutBuilder::GetDetailFont())
+					.Text(InText)
+				]
+			];
+	}
+
+	InEnumType GetEnumValue() const
+	{
+		if (const TSharedPtr<IPropertyHandle> Property = WeakProperty.Pin())
+		{
+			uint8 Value = 0; 
+			Property->GetValue(Value);
+			return static_cast<InEnumType>(Value);
+		}
+		return InEnumType::Default;
+	}
+
+private:
+	TWeakPtr<IPropertyHandle> WeakProperty;
+
+	FEnumValueChangedDelegate OnValueChanged;
+
+	TArray<TSharedRef<SCheckBox>> Buttons;
+};
diff --git a/Engine/Plugins/Experimental/Avalanche/Source/AvalancheSequencer/Private/StaggerTool/Widgets/SAvaStaggerTool.cpp b/Engine/Plugins/Experimental/Avalanche/Source/AvalancheSequencer/Private/StaggerTool/Widgets/SAvaStaggerTool.cpp
new file mode 100644
--- /dev/null
+++ b/Engine/Plugins/Experimental/Avalanche/Source/AvalancheSequencer/Private/StaggerTool/Widgets/SAvaStaggerTool.cpp
@@ -0,0 +1,294 @@
+// Copyright Epic Games, Inc. All Rights Reserved.
+
+#include "SAvaStaggerTool.h"
+#include "Framework/Application/SlateApplication.h"
+#include "IDetailsView.h"
+#include "Input/Reply.h"
+#include "Misc/Optional.h"
+#include "Modules/ModuleManager.h"
+#include "PropertyEditorModule.h"
+#include "StaggerTool/AvaStaggerTool.h"
+#include "StaggerTool/Customization/AvaStaggerToolSettingsDetailsCustomization.h"
+#include "Styling/StyleColors.h"
+#include "Widgets/Images/SImage.h"
+#include "Widgets/Input/SButton.h"
+#include "Widgets/Layout/SBox.h"
+#include "Widgets/Layout/SSeparator.h"
+#include "Widgets/SBoxPanel.h"
+#include "Widgets/Text/STextBlock.h"
+
+#define LOCTEXT_NAMESPACE "SAvaStaggerTool"
+
+void SAvaStaggerTool::Construct(const FArguments& InArgs, const TWeakPtr<FAvaStaggerTool>& InWeakTool)
+{
+	check(InWeakTool.IsValid());
+	WeakTool = InWeakTool;
+
+	OnResetToDefaults = InArgs._OnResetToDefaults;
+	OnSettingChange = InArgs._OnSettingChange;
+	OnApply = InArgs._OnApply;
+
+	ChildSlot
+	[
+		SNew(SVerticalBox)
+		+ SVerticalBox::Slot()
+		.AutoHeight()
+		.Padding(5.f, 1.f, 5.f, 2.f)
+		[
+			ConstructAlignTools()
+		]
+		+ SVerticalBox::Slot()
+		.AutoHeight()
+		[
+			SNew(SSeparator)
+			.ColorAndOpacity(FStyleColors::Dropdown)
+			.SeparatorImage(FAppStyle::Get().GetBrush("WhiteBrush"))
+		]
+		+ SVerticalBox::Slot()
+		.AutoHeight()
+		.Padding(5.f, 5.f, 5.f, 2.f)
+		[
+			ConstructApplyRow()
+		]
+		+ SVerticalBox::Slot()
+		.AutoHeight()
+		.Padding(5.f, 0.f, 5.f, 5.f)
+		[
+			ConstructDetails()
+		]
+	];
+}
+
+TSharedRef<SWidget> SAvaStaggerTool::ConstructDetails()
+{
+	if (!WeakTool.IsValid())
+	{
+		return SNullWidget::NullWidget;
+	}
+
+	FDetailsViewArgs DetailsViewArgs;
+	DetailsViewArgs.bAllowSearch = false;
+	DetailsViewArgs.bShowOptions = false;
+	DetailsViewArgs.bShowScrollBar = true;
+	DetailsViewArgs.bHideSelectionTip = true;
+	DetailsViewArgs.NotifyHook = this;
+	DetailsViewArgs.ColumnWidth = 0.75f;
+
+	FPropertyEditorModule& PropertyEditorModule = FModuleManager::Get().GetModuleChecked<FPropertyEditorModule>(TEXT("PropertyEditor"));
+	DetailsView = PropertyEditorModule.CreateDetailView(DetailsViewArgs);
+	DetailsView->RegisterInstancedCustomPropertyLayout(UAvaSequencerStaggerSettings::StaticClass()
+		, FOnGetDetailCustomizationInstance::CreateStatic(&FAvaStaggerToolSettingsDetailsCustomization::MakeInstance, WeakTool));
+
+	DetailsView->SetObject(WeakTool.Pin()->GetSettings());
+
+	return DetailsView.ToSharedRef();
+}
+
+TSharedRef<SWidget> SAvaStaggerTool::ConstructApplyRow()
+{
+	constexpr float ButtonSize = 22.f;
+	constexpr float ImageSize = 18.f;
+	constexpr float Spacing = 5.f;
+
+	return 
+		SNew(SHorizontalBox)
+		+ SHorizontalBox::Slot()
+		.VAlign(VAlign_Center)
+		.Padding(Spacing, 0.f, 0.f, 0.f)
+		[
+			SNew(STextBlock)
+			.TextStyle(FAppStyle::Get(), TEXT("HintText"))
+			.Text(this, &SAvaStaggerTool::GetSelectionText)
+			.ToolTipText(this, &SAvaStaggerTool::GetSelectionToolTipText)
+			.ColorAndOpacity(this, &SAvaStaggerTool::GetSelectionTextColor)
+		]
+		+ SHorizontalBox::Slot()
+		.AutoWidth()
+		.VAlign(VAlign_Bottom)
+		.Padding(Spacing, 0.f, 0.f, 0.f)
+		[
+			SNew(SButton)
+			.Text(LOCTEXT("ApplyButtonText", "Apply"))
+			.OnClicked(this, &SAvaStaggerTool::OnApplyButtonClick)
+			.IsEnabled_Lambda([this]()
+				{
+					return WeakTool.IsValid() && !WeakTool.Pin()->IsAutoAppying();
+				})
+		]
+		+ SHorizontalBox::Slot()
+		.AutoWidth()
+		.VAlign(VAlign_Bottom)
+		.Padding(Spacing, 0.f, 0.f, 0.f)
+		[
+			SNew(SBox)
+			.WidthOverride(ButtonSize)
+			.HeightOverride(ButtonSize)
+			[
+				SNew(SCheckBox)
+				.Style(FAppStyle::Get(), TEXT("Sequencer.Outliner.ToggleButton"))
+				.ToolTipText(LOCTEXT("AutoApplyToolTip", "Auto apply options on change"))
+				.IsChecked_Lambda([this]()
+					{
+						return (WeakTool.IsValid() && WeakTool.Pin()->IsAutoAppying())
+							? ECheckBoxState::Checked : ECheckBoxState::Unchecked;
+					})
+				.OnCheckStateChanged(this, &SAvaStaggerTool::OnToggleAutoUpdateClick)
+				[
+					SNew(SBox)
+					.Padding(0.f, 4.f)
+					[
+						SNew(SImage)
+						.Image_Lambda([this]()
+							{
+								const FName BrushName = (WeakTool.IsValid() && WeakTool.Pin()->IsAutoAppying())
+									? TEXT("Icons.Lock") : TEXT("Icons.Unlock");
+								return FAppStyle::GetBrush(BrushName);
+							})
+					]
+				]
+			]
+		]
+		+ SHorizontalBox::Slot()
+		.AutoWidth()
+		.VAlign(VAlign_Bottom)
+		.Padding(Spacing, 0.f, 0.f, 0.f)
+		[
+			SNew(SBox)
+			.WidthOverride(ButtonSize)
+			.HeightOverride(ButtonSize)
+			[
+				SNew(SButton)
+				.ContentPadding(2.f)
+				.ButtonStyle(FAppStyle::Get(), TEXT("SimpleButton"))
+				.ToolTipText(LOCTEXT("ResetToDefaultsToolTip", "Reset All to Defaults"))
+				.OnClicked(this, &SAvaStaggerTool::OnResetToDefaultsClick)
+				.IsEnabled_Lambda([this]()
+					{
+						return WeakTool.IsValid() ? WeakTool.Pin()->GetSettings()->CanResetToolOptions() : false;
+					})
+				[
+					SNew(SImage)
+					.DesiredSizeOverride(FVector2D(ImageSize))
+					.Image(FAppStyle::GetBrush(TEXT("PropertyWindow.DiffersFromDefault")))
+				]
+			]
+		];
+}
+
+TSharedRef<SWidget> SAvaStaggerTool::ConstructAlignTools()
+{
+	return SNew(SBox)
+		.Padding(5.f)
+		[
+			SNew(SHorizontalBox)
+			+ SHorizontalBox::Slot()
+			.AutoWidth()
+			.HAlign(HAlign_Left)
+			.VAlign(VAlign_Center)
+			.Padding(FMargin(0.f, 0.f, 10.f, 0.f))
+			[
+				SNew(STextBlock)
+				.TextStyle(FAppStyle::Get(), TEXT("SmallText"))
+				.Text(LOCTEXT("AlignToText", "Align to: "))
+			]
+			+ SHorizontalBox::Slot()
+			.AutoWidth()
+			.HAlign(HAlign_Left)
+			.VAlign(VAlign_Center)
+			[
+				SNew(SButton)
+				.TextStyle(FAppStyle::Get(), TEXT("SmallText"))
+				.Text(LOCTEXT("AlignToPlayheadText", "Playhead"))
+				.IsEnabled_Lambda([this]()
+					{
+						const TSharedPtr<FAvaStaggerTool> StaggerTool = WeakTool.Pin();
+						return StaggerTool.IsValid() ? (StaggerTool->IsLayerBarSelection() || StaggerTool->IsKeyFrameSelection()) : false;
+					})
+				.OnClicked(this, &SAvaStaggerTool::OnAlignToPlayhead)
+			]
+		];
+}
+
+FText SAvaStaggerTool::GetSelectionText() const
+{
+	const TSharedPtr<FAvaStaggerTool> StaggerTool = WeakTool.Pin();
+	if (!StaggerTool.IsValid())
+	{
+		return FText::GetEmpty();
+	}
+
+	if (StaggerTool->IsLayerBarSelection())
+	{
+		const FText SelectionCountText = FText::FromString(FString::FromInt(StaggerTool->GetSelectionCount()));
+		return FText::Format(LOCTEXT("LayerBarSelectionDetails", "{0} Layer Bars Selected"), SelectionCountText);
+	}
+
+	if (StaggerTool->IsKeyFrameSelection())
+	{
+		const FText SelectionCountText = FText::FromString(FString::FromInt(StaggerTool->GetSelectionCount()));
+		return FText::Format(LOCTEXT("KeyFrameSelectionDetails", "{0} Key Frames Selected"), SelectionCountText);
+	}
+
+	return LOCTEXT("InvalidSelection", "Invalid Selection");
+}
+
+FText SAvaStaggerTool::GetSelectionToolTipText() const
+{
+	return LOCTEXT("InvalidSelectionTooltip",
+		"Select at least two Sequencer outliner tracks containing a layer bar or at least two layer bars directly\n\n"
+		" - OR -\n\n"
+		"Select at least two Sequencer key frames");
+}
+
+FSlateColor SAvaStaggerTool::GetSelectionTextColor() const
+{
+	return (WeakTool.IsValid() && WeakTool.Pin()->HasValidSelection()) ? GetForegroundColor() : FStyleColors::Error;
+}
+
+FReply SAvaStaggerTool::OnApplyButtonClick()
+{
+	OnApply.ExecuteIfBound();
+
+	return FReply::Handled();
+}
+
+FReply SAvaStaggerTool::OnResetToDefaultsClick()
+{
+	OnResetToDefaults.ExecuteIfBound();
+
+	return FReply::Handled();
+}
+
+void SAvaStaggerTool::OnToggleAutoUpdateClick(const ECheckBoxState InNewState)
+{
+	if (WeakTool.IsValid())
+	{
+		return;
+	}
+
+	UAvaSequencerStaggerSettings* const Settings = WeakTool.Pin()->GetSettings();
+
+	Settings->bAutoApply = InNewState == ECheckBoxState::Checked;
+
+	if (Settings->bAutoApply)
+	{
+		OnApply.ExecuteIfBound();
+	}
+}
+
+FReply SAvaStaggerTool::OnAlignToPlayhead()
+{
+	if (const TSharedPtr<FAvaStaggerTool> Tool = WeakTool.Pin())
+	{
+		Tool->AlignToPlayhead();
+	}
+
+	return FReply::Handled();
+}
+
+void SAvaStaggerTool::NotifyPostChange(const FPropertyChangedEvent& InPropertyChangedEvent, FProperty* const InPropertyChanged)
+{
+	OnSettingChange.ExecuteIfBound(InPropertyChanged->GetFName());
+}
+
+#undef LOCTEXT_NAMESPACE
diff --git a/Engine/Plugins/Experimental/Avalanche/Source/AvalancheSequencer/Private/StaggerTool/Widgets/SAvaStaggerTool.h b/Engine/Plugins/Experimental/Avalanche/Source/AvalancheSequencer/Private/StaggerTool/Widgets/SAvaStaggerTool.h
new file mode 100644
--- /dev/null
+++ b/Engine/Plugins/Experimental/Avalanche/Source/AvalancheSequencer/Private/StaggerTool/Widgets/SAvaStaggerTool.h
@@ -0,0 +1,56 @@
+// Copyright Epic Games, Inc. All Rights Reserved.
+
+#pragma once
+
+#include "GenericPlatform/GenericPlatformMisc.h"
+#include "Misc/Attribute.h"
+#include "Misc/NotifyHook.h"
+#include "Misc/OptionalFwd.h"
+#include "Widgets/Input/SCheckBox.h"
+#include "Widgets/SCompoundWidget.h"
+
+class FAvaStaggerTool;
+class FReply;
+class IDetailsView;
+class SWidget;
+
+DECLARE_DELEGATE_OneParam(FAvaStaggerToolSettingChange, const FName /*InPropertyName*/)
+
+class SAvaStaggerTool : public SCompoundWidget, public FNotifyHook
+{
+public:
+	SLATE_BEGIN_ARGS(SAvaStaggerTool) {}
+		SLATE_EVENT(FSimpleDelegate, OnResetToDefaults)
+		SLATE_EVENT(FAvaStaggerToolSettingChange, OnSettingChange)
+		SLATE_EVENT(FSimpleDelegate, OnApply)
+	SLATE_END_ARGS()
+
+	void Construct(const FArguments& InArgs, const TWeakPtr<FAvaStaggerTool>& InWeakTool);
+
+protected:
+	TSharedRef<SWidget> ConstructDetails();
+	TSharedRef<SWidget> ConstructApplyRow();
+	TSharedRef<SWidget> ConstructAlignTools();
+
+	FText GetSelectionText() const;
+	FText GetSelectionToolTipText() const;
+	FSlateColor GetSelectionTextColor() const;
+
+	FReply OnApplyButtonClick();
+	FReply OnResetToDefaultsClick();
+	void OnToggleAutoUpdateClick(const ECheckBoxState InNewState);
+
+	FReply OnAlignToPlayhead();
+
+	//~ Begin FNotifyHook
+	virtual void NotifyPostChange(const FPropertyChangedEvent& InPropertyChangedEvent, FProperty* const InPropertyChanged) override;
+	//~ End FNotifyHook
+
+	TWeakPtr<FAvaStaggerTool> WeakTool;
+
+	FSimpleDelegate OnResetToDefaults;
+	FAvaStaggerToolSettingChange OnSettingChange;
+	FSimpleDelegate OnApply;
+
+	TSharedPtr<IDetailsView> DetailsView;
+};
