[MotionDesign] Stagger Tool: Remove unbinding to Sequencer changes since stagger tool exists for the lifetime of the Sequencer

diff --git a/Engine/Plugins/Experimental/Avalanche/Source/AvalancheSequencer/Private/StaggerTool/AvaStaggerTool.cpp b/Engine/Plugins/Experimental/Avalanche/Source/AvalancheSequencer/Private/StaggerTool/AvaStaggerTool.cpp
--- a/Engine/Plugins/Experimental/Avalanche/Source/AvalancheSequencer/Private/StaggerTool/AvaStaggerTool.cpp
+++ b/Engine/Plugins/Experimental/Avalanche/Source/AvalancheSequencer/Private/StaggerTool/AvaStaggerTool.cpp
@@ -46,7 +46,6 @@ FAvaStaggerTool::FAvaStaggerTool(FAvaSequencer& InOwner)
 
 FAvaStaggerTool::~FAvaStaggerTool()
 {
-	UnbindToSequencerSelectionChange();
 }
 
 TObjectPtr<UAvaSequencerStaggerSettings> FAvaStaggerTool::GetSettings() const
