[Motion Design] Rundown Server: Fixing Get Page Details in packaged games.

In cooked games, the original .umap file is not present. An alternative method is to use the asset registry to check the package flags. This seems the most reliable method in any build configuration.

Note: FPackageName::DoesPackageExist strips the original file extension in cooked builds.

diff --git a/Engine/Plugins/Experimental/Avalanche/Source/AvalancheMedia/Private/Playback/AvaPlaybackUtils.cpp b/Engine/Plugins/Experimental/Avalanche/Source/AvalancheMedia/Private/Playback/AvaPlaybackUtils.cpp
--- a/Engine/Plugins/Experimental/Avalanche/Source/AvalancheMedia/Private/Playback/AvaPlaybackUtils.cpp
+++ b/Engine/Plugins/Experimental/Avalanche/Source/AvalancheMedia/Private/Playback/AvaPlaybackUtils.cpp
@@ -3,6 +3,7 @@
 #include "Playback/AvaPlaybackUtils.h"
 
 #include "AssetRegistry/AssetData.h"
+#include "AssetRegistry/IAssetRegistry.h"
 #include "AvaAssetTags.h"
 #include "Framework/AvaSoftAssetPtr.h"
 #include "IAvaMediaModule.h"
@@ -152,12 +153,28 @@ bool FAvaPlaybackUtils::ReloadPackages(const TArray<UPackage*>& InExistingPackag
 
 bool FAvaPlaybackUtils::IsMapAsset(const FString& InPackageName)
 {
-	FString FullPath;
-	if (FPackageName::TryConvertLongPackageNameToFilename(InPackageName, FullPath, FPackageName::GetMapPackageExtension()))
+	if (const IAssetRegistry* AssetRegistry = IAssetRegistry::Get())
 	{
-		if (FPlatformFileManager::Get().GetPlatformFile().FileExists(*FullPath))
+		IAssetRegistry::FLoadPackageRegistryData RegistryData;
+		if (AssetRegistry->GetAssetsByPackageName(FName(*InPackageName), RegistryData.Data))
 		{
-			return true;
+			// In case Asset Data is not yet in registry's cache, need to load it directly.
+			if (RegistryData.Data.Num() == 0)
+			{
+				FString PackageFilename;
+				if (FPackageName::DoesPackageExist(InPackageName, &PackageFilename))
+				{
+					AssetRegistry->LoadPackageRegistryData(PackageFilename, RegistryData);
+				}
+			}
+
+			for (const FAssetData& AssetData : RegistryData.Data)
+			{
+				if (AssetData.IsValid() && AssetData.HasAnyPackageFlags(PKG_ContainsMap))
+				{
+					return true;
+				}
+			}
 		}
 	}
 	return false;
