MotionDesign ; Text3D
- Fixed custom material option would set same material on all slots after load, each slots can have its own material

diff --git a/Engine/Plugins/Experimental/Avalanche/Source/AvalancheText/Private/AvaText3DComponent.cpp b/Engine/Plugins/Experimental/Avalanche/Source/AvalancheText/Private/AvaText3DComponent.cpp
--- a/Engine/Plugins/Experimental/Avalanche/Source/AvalancheText/Private/AvaText3DComponent.cpp
+++ b/Engine/Plugins/Experimental/Avalanche/Source/AvalancheText/Private/AvaText3DComponent.cpp
@@ -196,15 +196,6 @@ void UAvaText3DComponent::SetTiling(const FVector2D InTiling)
 	}
 }
 
-void UAvaText3DComponent::SetCustomMaterial(UMaterialInterface* InCustomMaterial)
-{
-	if (CustomMaterial != InCustomMaterial)
-	{
-		CustomMaterial = InCustomMaterial;
-		RefreshMaterialInstances();
-	}
-}
-
 void UAvaText3DComponent::SetTranslucencyStyle(EAvaTextTranslucency InTranslucencyStyle)
 {
 	if (TranslucencyStyle != InTranslucencyStyle)
@@ -645,49 +636,14 @@ void UAvaText3DComponent::SetupCustomMaterial()
 	TexturedTextMaterialInstance = nullptr;
 	GradientTextMaterialInstance = nullptr;
 
-	if (CustomMaterial && IsMarkedForMaterialRefresh())
+	if (IsMarkedForMaterialRefresh())
 	{
 		RemoveRefreshReason(EAvaTextRefreshReason::MaterialInstancesChange);
 
-		UE::AvaText::Private::PatchMaterial(
-			this
-			, &UText3DComponent::GetFrontMaterial
-			, &UText3DComponent::SetFrontMaterial
-			, GradientTextMaterialInstance);
-
-		if (UMaterialInstanceDynamic* CustomMaterialAsInstance = Cast<UMaterialInstanceDynamic>(CustomMaterial))
-		{
-			UE::AvaText::Private::PatchMaterial(
-				this
-				, &UText3DComponent::GetFrontMaterial
-				, &UText3DComponent::SetFrontMaterial
-				, CustomMaterialAsInstance);
-
-			UE::AvaText::Private::PatchMaterial(
-				this
-				, &UText3DComponent::GetExtrudeMaterial
-				, &UText3DComponent::SetExtrudeMaterial
-				, CustomMaterialAsInstance);
-
-			UE::AvaText::Private::PatchMaterial(
-				this
-				, &UText3DComponent::GetBevelMaterial
-				, &UText3DComponent::SetBevelMaterial
-				, CustomMaterialAsInstance);
-
-			UE::AvaText::Private::PatchMaterial(
-				this
-				, &UText3DComponent::GetBackMaterial
-				, &UText3DComponent::SetBackMaterial
-				, CustomMaterialAsInstance);
-		}
-		else
-		{
-			SetFrontMaterial(CustomMaterial);
-			SetExtrudeMaterial(CustomMaterial);
-			SetBevelMaterial(CustomMaterial);
-			SetBackMaterial(CustomMaterial);
-		}
+		SetFrontMaterial(FrontMaterial);
+		SetExtrudeMaterial(ExtrudeMaterial);
+		SetBevelMaterial(BevelMaterial);
+		SetBackMaterial(BackMaterial);
 	}
 }
 
@@ -936,7 +892,6 @@ void UAvaText3DComponent::RetrieveCurrentMaterials()
 			break;
 
 		case EAvaTextColoringStyle::CustomMaterial: 
-			RetrieveCustomMaterialReference();
 			break;
 	}
 }
@@ -969,7 +924,6 @@ void UAvaText3DComponent::RetrieveCurrentColoringStyle()
 		else // if there is no match, best guess is custom material
 		{
 			ColoringStyle = EAvaTextColoringStyle::CustomMaterial;
-			CustomMaterial = CurrentFrontMaterial;
 			TranslucencyStyle = EAvaTextTranslucency::None;
 		}
 	}
@@ -1076,16 +1030,6 @@ void UAvaText3DComponent::RetrieveTexturedMaterialInstance()
 	}
 }
 
-void UAvaText3DComponent::RetrieveCustomMaterialReference()
-{
-	if (!FrontMaterial)
-	{
-		return;
-	}
-
-	CustomMaterial = FrontMaterial;
-}
-
 void UAvaText3DComponent::RetrieveTranslucentMaterialParameters()
 {
 	switch (TranslucencyStyle)
@@ -1209,7 +1153,6 @@ void UAvaText3DComponent::RegisterOnPropertyChangeFunctions()
 		PropertyChangeFunctions.Add(GET_MEMBER_NAME_CHECKED(UAvaText3DComponent, Text), &UAvaText3DComponent::RefreshGeometry);
 		PropertyChangeFunctions.Add(GET_MEMBER_NAME_CHECKED(UAvaText3DComponent, bEnforceUpperCase), &UAvaText3DComponent::RefreshFormatting);
 
-		PropertyChangeFunctions.Add(GET_MEMBER_NAME_CHECKED(UAvaText3DComponent, CustomMaterial), &UAvaText3DComponent::RefreshMaterialInstances);
 		PropertyChangeFunctions.Add(GET_MEMBER_NAME_CHECKED(UAvaText3DComponent, ColoringStyle), &UAvaText3DComponent::RefreshMaterialInstances);
 		PropertyChangeFunctions.Add(GET_MEMBER_NAME_CHECKED(UAvaText3DComponent, MainTexture), &UAvaText3DComponent::RefreshMaterialInstances);
 		PropertyChangeFunctions.Add(GET_MEMBER_NAME_CHECKED(UAvaText3DComponent, Tiling), &UAvaText3DComponent::RefreshMaterialInstances);
diff --git a/Engine/Plugins/Experimental/Avalanche/Source/AvalancheText/Public/AvaText3DComponent.h b/Engine/Plugins/Experimental/Avalanche/Source/AvalancheText/Public/AvaText3DComponent.h
--- a/Engine/Plugins/Experimental/Avalanche/Source/AvalancheText/Public/AvaText3DComponent.h
+++ b/Engine/Plugins/Experimental/Avalanche/Source/AvalancheText/Public/AvaText3DComponent.h
@@ -77,10 +77,6 @@ public:
 	AVALANCHETEXT_API void SetTiling(const FVector2D InTiling);
 	const FVector2D& GetTiling() const { return Tiling; }
 
-	UFUNCTION(BlueprintSetter)
-	AVALANCHETEXT_API void SetCustomMaterial(UMaterialInterface* InCustomMaterial);
-	UMaterialInterface* GetCustomMaterial() const { return CustomMaterial; }
-
 	UFUNCTION(BlueprintSetter)
 	AVALANCHETEXT_API void SetTranslucencyStyle(EAvaTextTranslucency InTranslucencyStyle);
 	EAvaTextTranslucency GetTranslucencyStyle() const { return TranslucencyStyle; }
@@ -200,7 +196,6 @@ protected:
 	void RetrieveSolidMaterialInstances();
 	void RetrieveGradientMaterialInstance();
 	void RetrieveTexturedMaterialInstance();
-	void RetrieveCustomMaterialReference();
 
 	void RetrieveTranslucentMaterialParameters();
 	void RetrieveMaskedMaterialParameters(const UMaterialInstanceDynamic* InMaterialInstance);
@@ -325,9 +320,6 @@ private:
 	UPROPERTY(EditAnywhere, Getter="GetTiling", Setter="SetTiling", Category="Style", meta=(Delta=0.1, EditCondition="ColoringStyle == EAvaTextColoringStyle::FromTexture", EditConditionHides, DisplayAfter="Text", AllowPrivateAccess="true"))
 	FVector2D Tiling;
 
-	UPROPERTY(EditAnywhere, Getter="GetCustomMaterial", Setter="SetCustomMaterial", Category="Style", meta=(EditCondition="ColoringStyle == EAvaTextColoringStyle::CustomMaterial", EditConditionHides, DisplayAfter="Text", AllowPrivateAccess="true"))
-	TObjectPtr<UMaterialInterface> CustomMaterial = nullptr;
-
 	UPROPERTY(EditAnywhere, Getter="GetTranslucencyStyle", Setter="SetTranslucencyStyle", Category="Style", meta=(EditCondition="ColoringStyle != EAvaTextColoringStyle::CustomMaterial", EditConditionHides,DisplayAfter="Text", AllowPrivateAccess="true"))
 	EAvaTextTranslucency TranslucencyStyle = EAvaTextTranslucency::None;
 
@@ -346,7 +338,7 @@ private:
 	UPROPERTY(EditAnywhere, Getter="GetMaskRotation", Setter="SetMaskRotation", Category="Style", meta=(ClampMin=0.0, ClampMax=1.0, EditCondition="MaskOrientation == EAvaMaterialMaskOrientation::Custom && TranslucencyStyle == EAvaTextTranslucency::GradientMask", EditConditionHides, DisplayAfter="Text", AllowPrivateAccess="true"))
 	float MaskRotation = 0.0f;
 
-	UPROPERTY(EditAnywhere, Getter="GetIsUnlit", Setter="SetIsUnlit", Category="Style", meta=(EditCondition="ColoringStyle != EAvaTextColoringStyle::Custom", EditConditionHides, DisplayAfter="Text", AllowPrivateAccess="true"))
+	UPROPERTY(EditAnywhere, Getter="GetIsUnlit", Setter="SetIsUnlit", Category="Style", meta=(EditCondition="ColoringStyle != EAvaTextColoringStyle::CustomMaterial", EditConditionHides, DisplayAfter="Text", AllowPrivateAccess="true"))
 	bool bIsUnlit = true;
 
 #if WITH_EDITOR
