MotionDesign : Modifiers
- Fixed boolean modifier tools intersection causing additional faces on target mesh

diff --git a/Engine/Plugins/Experimental/Avalanche/Source/AvalancheModifiers/Private/Modifiers/AvaBooleanModifier.cpp b/Engine/Plugins/Experimental/Avalanche/Source/AvalancheModifiers/Private/Modifiers/AvaBooleanModifier.cpp
--- a/Engine/Plugins/Experimental/Avalanche/Source/AvalancheModifiers/Private/Modifiers/AvaBooleanModifier.cpp
+++ b/Engine/Plugins/Experimental/Avalanche/Source/AvalancheModifiers/Private/Modifiers/AvaBooleanModifier.cpp
@@ -12,6 +12,7 @@
 #include "Materials/MaterialInstanceDynamic.h"
 #include "Modifiers/ActorModifierCoreStack.h"
 #include "Operations/MeshBoolean.h"
+#include "Operations/MeshSelfUnion.h"
 #include "Operations/OffsetMeshRegion.h"
 #include "Profilers/AvaBooleanModifierProfiler.h"
 #include "Shared/AvaBooleanModifierShared.h"
@@ -391,6 +392,17 @@ void UAvaBooleanModifier::ApplyInternal()
 				});
 			}
 
+			// Remove tools intersection
+			for (TPair<EAvaBooleanMode, FDynamicMesh3>& ModeToolPair : ModeTools)
+			{
+				FMeshSelfUnion Union(&ModeToolPair.Value);
+				Union.WindingThreshold = 0.5f;
+				Union.bTrimFlaps = true;
+				Union.bSimplifyAlongNewEdges = true;
+				Union.SimplificationAngleTolerance = 0.01;
+				Union.Compute();
+			}
+
 			FDynamicMesh3 OutputMesh;
 			GetMeshComponent()->ProcessMesh([&OutputMesh](const FDynamicMesh3& InEditMesh)
 			{
