Motion Design: Cleanup unused Startup Camera and the interfaces built up for it

diff --git a/Engine/Plugins/Experimental/Avalanche/Source/Avalanche/Private/AvaScene.cpp b/Engine/Plugins/Experimental/Avalanche/Source/Avalanche/Private/AvaScene.cpp
--- a/Engine/Plugins/Experimental/Avalanche/Source/Avalanche/Private/AvaScene.cpp
+++ b/Engine/Plugins/Experimental/Avalanche/Source/Avalanche/Private/AvaScene.cpp
@@ -87,8 +87,6 @@ AAvaScene::AAvaScene()
 
 	RemoteControlPreset = CreateDefaultSubobject<URemoteControlPreset>(TEXT("RemoteControlPreset"));
 
-	StartupCameraName = NAME_None;
-
 #if WITH_EDITOR
 	if (!HasAnyFlags(EObjectFlags::RF_ClassDefaultObject))
 	{
@@ -353,13 +351,6 @@ void AAvaScene::BeginDestroy()
 #endif
 }
 
-#if WITH_EDITOR
-void AAvaScene::SetStartupCameraName(FName InName)
-{
-	StartupCameraName = InName;
-}
-#endif
-
 void AAvaScene::RegisterObjects()
 {
 	if (!HasAnyFlags(EObjectFlags::RF_ClassDefaultObject))
diff --git a/Engine/Plugins/Experimental/Avalanche/Source/Avalanche/Public/AvaScene.h b/Engine/Plugins/Experimental/Avalanche/Source/Avalanche/Public/AvaScene.h
--- a/Engine/Plugins/Experimental/Avalanche/Source/Avalanche/Public/AvaScene.h
+++ b/Engine/Plugins/Experimental/Avalanche/Source/Avalanche/Public/AvaScene.h
@@ -10,7 +10,6 @@
 #include "IAvaSequenceProvider.h"
 #include "UObject/ObjectPtr.h"
 #include "UObject/WeakObjectPtr.h"
-#include "Viewport/Interaction/IAvaViewportDataProvider.h"
 #include "AvaScene.generated.h"
 
 class IAvaSequencePlaybackObject;
@@ -24,7 +23,7 @@ class ISequencer;
 #endif
 
 UCLASS(MinimalAPI, NotPlaceable, Hidden, NotBlueprintable, NotBlueprintType, DisplayName = "Motion Design Scene")
-class AAvaScene : public AActor, public IAvaSequenceProvider, public IAvaSceneInterface, public IAvaViewportDataProvider, public IAvaRemoteControlInterface
+class AAvaScene : public AActor, public IAvaSequenceProvider, public IAvaSceneInterface, public IAvaRemoteControlInterface
 {
 	GENERATED_BODY()
 
@@ -109,13 +108,6 @@ public:
 	virtual void BeginDestroy() override;
 	//~ End UObject
 
-	//~ Begin IAvaViewportDataProvider
-	virtual FName GetStartupCameraName() const override { return StartupCameraName; }
-#if WITH_EDITOR
-	virtual void SetStartupCameraName(FName InName) override;
-#endif
-	//~ End IAvaViewportDataProvider
-
 	void RegisterObjects();
 	void UnregisterObjects();
 
@@ -132,9 +124,6 @@ protected:
 	UPROPERTY()
 	TObjectPtr<URemoteControlPreset> RemoteControlPreset;
 
-	UPROPERTY()
-	FName StartupCameraName = NAME_None;
-
 	/** A List of All Animations, including those that are nested in other Animations */
 	UPROPERTY()
 	TArray<TObjectPtr<UAvaSequence>> Animations;
diff --git a/Engine/Plugins/Experimental/Avalanche/Source/Avalanche/Public/Viewport/Interaction/IAvaViewportDataProvider.h b/Engine/Plugins/Experimental/Avalanche/Source/Avalanche/Public/Viewport/Interaction/IAvaViewportDataProvider.h
deleted file mode 100644
diff --git a/Engine/Plugins/Experimental/Avalanche/Source/Avalanche/Public/Viewport/Interaction/IAvaViewportDataProxy.h b/Engine/Plugins/Experimental/Avalanche/Source/Avalanche/Public/Viewport/Interaction/IAvaViewportDataProxy.h
deleted file mode 100644
diff --git a/Engine/Plugins/Experimental/Avalanche/Source/AvalancheEditor/Private/Viewport/AvaLevelViewportExtension.cpp b/Engine/Plugins/Experimental/Avalanche/Source/AvalancheEditor/Private/Viewport/AvaLevelViewportExtension.cpp
--- a/Engine/Plugins/Experimental/Avalanche/Source/AvalancheEditor/Private/Viewport/AvaLevelViewportExtension.cpp
+++ b/Engine/Plugins/Experimental/Avalanche/Source/AvalancheEditor/Private/Viewport/AvaLevelViewportExtension.cpp
@@ -19,7 +19,6 @@
 #include "Sequencer/AvaSequencerExtension.h"
 #include "SLevelViewport.h"
 #include "UnrealEdGlobals.h"
-#include "Viewport/Interaction/IAvaViewportDataProvider.h"
 #include "ViewportClient/IAvaViewportClient.h"
 
 #define LOCTEXT_NAMESPACE "AvaLevelViewportExtension"
@@ -339,9 +338,9 @@ void FAvaLevelViewportExtension::SetViewportType(const TSharedRef<FUICommandInfo
 FViewportTypeDefinition FAvaLevelViewportExtension::MakeViewportTypeDefinition()
 {
 	return FViewportTypeDefinition(
-		[ThisWeak = SharedThis(this).ToWeakPtr()](const FAssetEditorViewportConstructionArgs& InArgs, TSharedPtr<ILevelEditor> InLevelEditor)
+		[](const FAssetEditorViewportConstructionArgs& InArgs, TSharedPtr<ILevelEditor> InLevelEditor)
 		{
-			return MakeShared<FAvaLevelViewportLayoutEntity>(InArgs, InLevelEditor, ThisWeak.Pin());
+			return MakeShared<FAvaLevelViewportLayoutEntity>(InArgs, InLevelEditor);
 		}
 		, FAvaEditorCommands::Get().SetMotionDesignViewportType);
 }
@@ -551,11 +550,6 @@ void FAvaLevelViewportExtension::SetActiveCamera(AActor* InActiveCameraActor, bo
 	}
 }
 
-IAvaViewportDataProvider* FAvaLevelViewportExtension::GetViewportDataProvider() const
-{
-	return GetSceneObject<IAvaViewportDataProvider>();
-}
-
 void FAvaLevelViewportExtension::ExecuteResetLocation()
 {
 	TSharedPtr<IAvaEditor> Editor = GetEditor();
diff --git a/Engine/Plugins/Experimental/Avalanche/Source/AvalancheEditor/Private/Viewport/AvaLevelViewportExtension.h b/Engine/Plugins/Experimental/Avalanche/Source/AvalancheEditor/Private/Viewport/AvaLevelViewportExtension.h
--- a/Engine/Plugins/Experimental/Avalanche/Source/AvalancheEditor/Private/Viewport/AvaLevelViewportExtension.h
+++ b/Engine/Plugins/Experimental/Avalanche/Source/AvalancheEditor/Private/Viewport/AvaLevelViewportExtension.h
@@ -6,21 +6,19 @@
 #include "Templates/SharedPointer.h"
 #include "UObject/WeakObjectPtr.h"
 #include "UObject/WeakObjectPtrTemplates.h"
-#include "Viewport/Interaction/IAvaViewportDataProxy.h"
 
 class AActor;
 class FUICommandInfo;
 class IAvaEditor;
-class IAvaViewportDataProxy;
 class UObject;
 class UWorld;
 enum class EMapChangeType : uint8;
 struct FViewportTypeDefinition;
 
-class FAvaLevelViewportExtension : public FAvaViewportExtension, public IAvaViewportDataProxy
+class FAvaLevelViewportExtension : public FAvaViewportExtension
 {
 public:
-	UE_AVA_INHERITS(FAvaLevelViewportExtension, FAvaViewportExtension, IAvaViewportDataProxy);
+	UE_AVA_INHERITS(FAvaLevelViewportExtension, FAvaViewportExtension);
 
 	static TArray<TSharedPtr<IAvaViewportClient>> GetLevelEditorViewportClients();
 
@@ -48,10 +46,6 @@ public:
 	virtual bool IsDroppingPreviewActor() const override;
 	//~ End IAvaViewportProvider
 
-	//~ Begin IAvaViewportDataProxy
-	virtual IAvaViewportDataProvider* GetViewportDataProvider() const override;
-	//~ End IAvaViewportDataProxy
-
 private:
 	FDelegateHandle OnMapChangedHandle;
 
diff --git a/Engine/Plugins/Experimental/Avalanche/Source/AvalancheLevelViewport/Private/AvaLevelViewportLayoutEntity.cpp b/Engine/Plugins/Experimental/Avalanche/Source/AvalancheLevelViewport/Private/AvaLevelViewportLayoutEntity.cpp
--- a/Engine/Plugins/Experimental/Avalanche/Source/AvalancheLevelViewport/Private/AvaLevelViewportLayoutEntity.cpp
+++ b/Engine/Plugins/Experimental/Avalanche/Source/AvalancheLevelViewport/Private/AvaLevelViewportLayoutEntity.cpp
@@ -8,18 +8,9 @@
 #include "SAvaLevelViewportFrame.h"
 #include "ViewportClient/AvaLevelViewportClient.h"
 
-FAvaLevelViewportLayoutEntity::FAvaLevelViewportLayoutEntity(const FAssetEditorViewportConstructionArgs& InArgs, 
-	TSharedPtr<ILevelEditor> InLevelEditor, TSharedPtr<IAvaViewportDataProxy> InDataProxy)
+FAvaLevelViewportLayoutEntity::FAvaLevelViewportLayoutEntity(const FAssetEditorViewportConstructionArgs& InArgs, TSharedPtr<ILevelEditor> InLevelEditor)
 {
 	Viewport = SNew(SAvaLevelViewportFrame, InArgs, InLevelEditor);
-
-	if (InDataProxy.IsValid())
-	{
-		if (TSharedPtr<FAvaLevelViewportClient> AvaViewportClient = Viewport->GetViewportClient())
-		{
-			AvaViewportClient->SetViewportDataProxy(InDataProxy);
-		}
-	}
 }
 
 FName FAvaLevelViewportLayoutEntity::GetStaticType()
diff --git a/Engine/Plugins/Experimental/Avalanche/Source/AvalancheLevelViewport/Private/SAvaLevelViewport.cpp b/Engine/Plugins/Experimental/Avalanche/Source/AvalancheLevelViewport/Private/SAvaLevelViewport.cpp
--- a/Engine/Plugins/Experimental/Avalanche/Source/AvalancheLevelViewport/Private/SAvaLevelViewport.cpp
+++ b/Engine/Plugins/Experimental/Avalanche/Source/AvalancheLevelViewport/Private/SAvaLevelViewport.cpp
@@ -19,7 +19,6 @@
 #include "ScopedTransaction.h"
 #include "Subsystems/PanelExtensionSubsystem.h"
 #include "Types/ISlateMetaData.h"
-#include "Viewport/Interaction/IAvaViewportDataProvider.h"
 #include "ViewportClient/AvaLevelViewportClient.h"
 #include "Widgets/Input/SCheckBox.h"
 #include "Widgets/Layout/SBorder.h"
@@ -479,13 +478,6 @@ void SAvaLevelViewport::LoadGuides()
 		return;
 	}
 
-	IAvaViewportDataProvider* DataProvider = ViewportClient->GetViewportDataProvider();
-
-	if (!DataProvider)
-	{
-		return;
-	}
-
 	UAvaViewportDataSubsystem* DataSubsystem = UAvaViewportDataSubsystem::Get(ViewportClient->GetViewportWorld());
 
 	if (!DataSubsystem)
@@ -575,13 +567,6 @@ void SAvaLevelViewport::SaveGuides()
 	Data->GuideData = GuideInfos;
 }
 
-void SAvaLevelViewport::OnViewportDataProxyChanged()
-{
-	ReloadGuides();
-	LoadPostProcessInfo();
-	LoadVirtualSizeSettings();
-}
-
 bool SAvaLevelViewport::MatchesContext(const FTransactionContext& InContext,
 	const TArray<TPair<UObject*, FTransactionObjectEvent>>& InTransactionObjectContexts) const
 {
diff --git a/Engine/Plugins/Experimental/Avalanche/Source/AvalancheLevelViewport/Private/SAvaLevelViewport.h b/Engine/Plugins/Experimental/Avalanche/Source/AvalancheLevelViewport/Private/SAvaLevelViewport.h
--- a/Engine/Plugins/Experimental/Avalanche/Source/AvalancheLevelViewport/Private/SAvaLevelViewport.h
+++ b/Engine/Plugins/Experimental/Avalanche/Source/AvalancheLevelViewport/Private/SAvaLevelViewport.h
@@ -76,8 +76,6 @@ public:
 	bool RemoveGuide(TSharedPtr<SAvaLevelViewportGuide> InGuidetoRemove);
 	void SaveGuides();
 
-	void OnViewportDataProxyChanged();
-
 	bool CanToggleOverlay() const;
 	void ExecuteToggleOverlay();
 
diff --git a/Engine/Plugins/Experimental/Avalanche/Source/AvalancheLevelViewport/Private/SAvaLevelViewport_Menus.cpp b/Engine/Plugins/Experimental/Avalanche/Source/AvalancheLevelViewport/Private/SAvaLevelViewport_Menus.cpp
--- a/Engine/Plugins/Experimental/Avalanche/Source/AvalancheLevelViewport/Private/SAvaLevelViewport_Menus.cpp
+++ b/Engine/Plugins/Experimental/Avalanche/Source/AvalancheLevelViewport/Private/SAvaLevelViewport_Menus.cpp
@@ -21,7 +21,6 @@
 #include "SortHelper.h"
 #include "SViewportToolBar.h"
 #include "ToolMenus.h"
-#include "Viewport/Interaction/IAvaViewportDataProvider.h"
 #include "ViewportClient/AvaLevelViewportClient.h"
 #include "Widgets/DataTypes/AvaUserInputDialogDataTypeText.h"
 #include "Widgets/Images/SImage.h"
@@ -968,25 +967,9 @@ void SAvaLevelViewport::OnVirtualSizeSliderBegin()
 		return;
 	}
 
-	IAvaViewportDataProvider* DataProvider = ViewportClient->GetViewportDataProvider();
-
-	if (!DataProvider)
-	{
-		return;
-	}
-
-	UObject* DataProviderObject = DataProvider->ToUObject();
-
-	if (!DataProviderObject)
-	{
-		return;
-	}
-
 	VirtualSizeComponentSliderTransaction = MakeShared<FScopedTransaction>(
 		LOCTEXT("SetCustomVirtualSize", "Set Custom Virtual Size")
 	);
-
-	DataProviderObject->Modify();
 }
 
 void SAvaLevelViewport::OnVirtualSizeComponentSliderEnd(int32 InNewSize)
diff --git a/Engine/Plugins/Experimental/Avalanche/Source/AvalancheLevelViewport/Private/SAvaLevelViewport_StatusBar.cpp b/Engine/Plugins/Experimental/Avalanche/Source/AvalancheLevelViewport/Private/SAvaLevelViewport_StatusBar.cpp
--- a/Engine/Plugins/Experimental/Avalanche/Source/AvalancheLevelViewport/Private/SAvaLevelViewport_StatusBar.cpp
+++ b/Engine/Plugins/Experimental/Avalanche/Source/AvalancheLevelViewport/Private/SAvaLevelViewport_StatusBar.cpp
@@ -6,7 +6,6 @@
 #include "AvaViewportSettings.h"
 #include "Engine/Texture.h"
 #include "ScopedTransaction.h"
-#include "Viewport/Interaction/IAvaViewportDataProvider.h"
 #include "ViewportClient/AvaLevelViewportClient.h"
 #include "Visualizers/IAvaViewportBoundingBoxVisualizer.h"
 #include "Visualizers/IAvaViewportPostProcessVisualizer.h"
@@ -494,26 +493,10 @@ void SAvaLevelViewport::BeginPostProcessInfoTransaction()
 		return;
 	}
 
-	IAvaViewportDataProvider* DataProvider = ViewportClient->GetViewportDataProvider();
-
-	if (!DataProvider)
-	{
-		return;
-	}
-
-	UObject* DataObject = DataProvider->ToUObject();
-
-	if (!DataObject)
-	{
-		return;
-	}
-
 	if (!PostProcessInfoTransaction.IsValid())
 	{
 		PostProcessInfoTransaction = MakeShared<FScopedTransaction>(LOCTEXT("PostProcessSettingsChange", "Post Process Settings Change"));
 	}
-
-	DataObject->Modify();
 }
 
 void SAvaLevelViewport::EndPostProcessInfoTransaction()
diff --git a/Engine/Plugins/Experimental/Avalanche/Source/AvalancheLevelViewport/Private/ViewportClient/AvaLevelViewportClient.cpp b/Engine/Plugins/Experimental/Avalanche/Source/AvalancheLevelViewport/Private/ViewportClient/AvaLevelViewportClient.cpp
--- a/Engine/Plugins/Experimental/Avalanche/Source/AvalancheLevelViewport/Private/ViewportClient/AvaLevelViewportClient.cpp
+++ b/Engine/Plugins/Experimental/Avalanche/Source/AvalancheLevelViewport/Private/ViewportClient/AvaLevelViewportClient.cpp
@@ -24,8 +24,6 @@
 #include "Selection/AvaSelectionProviderSubsystem.h"
 #include "SLevelViewport.h"
 #include "UnrealEdGlobals.h"
-#include "Viewport/Interaction/IAvaViewportDataProvider.h"
-#include "Viewport/Interaction/IAvaViewportDataProxy.h"
 #include "Visualizers/IAvaViewportBoundingBoxVisualizer.h"
 
 #define LOCTEXT_NAMESPACE "AvaLevelViewportClient"
@@ -167,26 +165,7 @@ void FAvaLevelViewportClient::SetViewTarget(TWeakObjectPtr<AActor> InViewTarget)
 
 	UCameraComponent* CameraComponent = ActiveCameraComponentWeak.Get();
 	AActor* ViewTarget = CameraComponent ? CameraComponent->GetOwner() : nullptr;
-
 	SetActorLock(ViewTarget);
-
-	if (ViewTarget)
-	{
-		if (IAvaViewportDataProvider* DataProvider = GetViewportDataProvider())
-		{
-			FName ViewTargetName = ViewTarget->GetFName();
-
-			if (DataProvider->GetStartupCameraName() != ViewTargetName)
-			{
-				if (UObject* Object = DataProvider->ToUObject())
-				{
-					Object->Modify();
-				}
-
-				DataProvider->SetStartupCameraName(ViewTargetName);
-			}
-		}
-	}
 }
 
 void FAvaLevelViewportClient::SetCinematicViewTarget(AActor* InCinematicViewTarget)
@@ -742,34 +721,4 @@ void FAvaLevelViewportClient::SetActiveCameraComponent(UCameraComponent* InCamer
 	}
 }
 
-void FAvaLevelViewportClient::SetViewportDataProxy(const TSharedPtr<IAvaViewportDataProxy>& InDataProxy)
-{
-	ViewportDataProxyWeak = InDataProxy;
-
-	if (TSharedPtr<SAvaLevelViewport> AvaLevelViewport = AvaLevelViewportWeak.Pin())
-	{
-		AvaLevelViewport->OnViewportDataProxyChanged();
-	}
-
-	if (IAvaViewportDataProvider* Provider = InDataProxy->GetViewportDataProvider())
-	{
-		FName StartupCameraName = Provider->GetStartupCameraName();
-
-		if (StartupCameraName != NAME_None)
-		{
-			if (UWorld* WorldLocal = GetWorld())
-			{
-				for (AActor* Actor : TActorRange<AActor>(WorldLocal))
-				{
-					if (Actor->GetFName() == StartupCameraName)
-					{
-						SetViewTarget(Actor);
-						break;
-					}					
-				}
-			}
-		}
-	}
-}
-
 #undef LOCTEXT_NAMESPACE
diff --git a/Engine/Plugins/Experimental/Avalanche/Source/AvalancheLevelViewport/Private/ViewportClient/AvaLevelViewportClient.h b/Engine/Plugins/Experimental/Avalanche/Source/AvalancheLevelViewport/Private/ViewportClient/AvaLevelViewportClient.h
--- a/Engine/Plugins/Experimental/Avalanche/Source/AvalancheLevelViewport/Private/ViewportClient/AvaLevelViewportClient.h
+++ b/Engine/Plugins/Experimental/Avalanche/Source/AvalancheLevelViewport/Private/ViewportClient/AvaLevelViewportClient.h
@@ -19,8 +19,6 @@ class FAvaViewportPostProcessManager;
 class FEditorViewportClient;
 class IAvaBoundsProviderInterface;
 class IAvaViewportBoundingBoxVisualizer;
-class IAvaViewportDataProvider;
-class IAvaViewportDataProxy;
 class SAvaLevelViewport;
 class UCameraComponent;
 struct FAvaSnapPoint;
@@ -99,8 +97,6 @@ public:
 	virtual FVector2f GetConstrainedViewportMousePosition() const override;
 	virtual FVector2f GetUnconstrainedZoomedViewportMousePosition() const override;
 	virtual FVector2f GetConstrainedZoomedViewportMousePosition() const override;
-	virtual TSharedPtr<IAvaViewportDataProxy> GetViewportDataProxy() const override { return ViewportDataProxyWeak.Pin(); }
-	virtual void SetViewportDataProxy(const TSharedPtr<IAvaViewportDataProxy>& InDataProxy) override;
 	virtual TSharedPtr<FAvaSnapOperation> GetSnapOperation() const override { return SnapOperationWeak.Pin(); };
 	virtual TSharedPtr<FAvaSnapOperation> StartSnapOperation() override;
 	virtual bool EndSnapOperation(FAvaSnapOperation* InSnapOperation) override;
@@ -127,8 +123,6 @@ protected:
 
 	TWeakPtr<FAvaSnapOperation> SnapOperationWeak;
 
-	TWeakPtr<IAvaViewportDataProxy> ViewportDataProxyWeak;
-
 	TSharedPtr<FAvaDragOperation> DragOperation;
 
 	TSharedPtr<FAvaChildTransformLockOperation> ChildTransformLockOperation;
diff --git a/Engine/Plugins/Experimental/Avalanche/Source/AvalancheLevelViewport/Public/AvaLevelViewportLayoutEntity.h b/Engine/Plugins/Experimental/Avalanche/Source/AvalancheLevelViewport/Public/AvaLevelViewportLayoutEntity.h
--- a/Engine/Plugins/Experimental/Avalanche/Source/AvalancheLevelViewport/Public/AvaLevelViewportLayoutEntity.h
+++ b/Engine/Plugins/Experimental/Avalanche/Source/AvalancheLevelViewport/Public/AvaLevelViewportLayoutEntity.h
@@ -5,7 +5,6 @@
 #include "LevelViewportLayout.h"
 #include "Templates/SharedPointer.h"
 
-class IAvaViewportDataProxy;
 class ILevelEditor;
 class SAvaLevelViewportFrame;
 struct FAssetEditorViewportConstructionArgs;
@@ -13,8 +12,7 @@ struct FAssetEditorViewportConstructionArgs;
 class AVALANCHELEVELVIEWPORT_API FAvaLevelViewportLayoutEntity : public ILevelViewportLayoutEntity
 {
 public:
-	FAvaLevelViewportLayoutEntity(const FAssetEditorViewportConstructionArgs& InArgs, TSharedPtr<ILevelEditor> InLevelEditor,
-		TSharedPtr<IAvaViewportDataProxy> InDataProxy);
+	FAvaLevelViewportLayoutEntity(const FAssetEditorViewportConstructionArgs& InArgs, TSharedPtr<ILevelEditor> InLevelEditor);
 
 	static FName GetStaticType();
 
diff --git a/Engine/Plugins/Experimental/Avalanche/Source/AvalancheMRQ/Private/AvaMRQRundownPageSetting.cpp b/Engine/Plugins/Experimental/Avalanche/Source/AvalancheMRQ/Private/AvaMRQRundownPageSetting.cpp
--- a/Engine/Plugins/Experimental/Avalanche/Source/AvalancheMRQ/Private/AvaMRQRundownPageSetting.cpp
+++ b/Engine/Plugins/Experimental/Avalanche/Source/AvalancheMRQ/Private/AvaMRQRundownPageSetting.cpp
@@ -118,23 +118,4 @@ void UAvaMRQRundownPageSetting::SetupForPipelineImpl(UMoviePipeline* InPipeline)
 
 	FAvaRemoteControlRebind::RebindUnboundEntities(RemoteControlPreset, Level);
 	Page.GetRemoteControlValues().ApplyEntityValuesToRemoteControlPreset(RemoteControlPreset);
-
-	// todo: ideally, _getUObject should not be used here, and instead the Viewport Data Provider gotten directly through an interface call
-	if (IAvaViewportDataProvider* ViewportDataProvider = Cast<IAvaViewportDataProvider>(Scene->_getUObject()))
-	{
-		if (APlayerController* PlayerController = World->GetFirstPlayerController())
-		{
-			FName StartupCameraName = ViewportDataProvider->GetStartupCameraName();
-
-			TObjectPtr<AActor>* ViewTarget = World->PersistentLevel->Actors.FindByPredicate([StartupCameraName](TObjectPtr<AActor> InActor)
-			{
-				return InActor && InActor->GetFName() == StartupCameraName;
-			});
-
-			if (ViewTarget)
-			{
-				PlayerController->SetViewTarget(*ViewTarget);
-			}
-		}
-	}
 }
diff --git a/Engine/Plugins/Experimental/Avalanche/Source/AvalancheMedia/Private/Playable/Playables/AvaPlayableLevelStreaming.cpp b/Engine/Plugins/Experimental/Avalanche/Source/AvalancheMedia/Private/Playable/Playables/AvaPlayableLevelStreaming.cpp
--- a/Engine/Plugins/Experimental/Avalanche/Source/AvalancheMedia/Private/Playable/Playables/AvaPlayableLevelStreaming.cpp
+++ b/Engine/Plugins/Experimental/Avalanche/Source/AvalancheMedia/Private/Playable/Playables/AvaPlayableLevelStreaming.cpp
@@ -51,76 +51,6 @@ namespace UE::AvaMedia::LevelStreamingPlayable::Private
 		return AvaScene;
 	}
 
-	AActor* FindActorByName(const ULevel* InLevel, const FName& InActorName)
-	{
-		for (AActor* Actor : InLevel->Actors)
-		{
-			if (Actor && Actor->GetFName() == InActorName)
-			{
-				return Actor;
-			}
-		}
-		return nullptr;
-	}
-
-	bool IsCameraActor(const AActor* InActor)
-	{
-		return InActor->IsA<ACameraActor>();
-	}
-	
-	AActor* FindFirstCameraActor(const ULevel* InLevel)
-	{
-		for (AActor* Actor : InLevel->Actors)
-		{
-			if (Actor && IsCameraActor(Actor))
-			{
-				return Actor;
-			}
-		}
-		return nullptr;
-	}
-
-	FName GetStartupCameraName(const AAvaScene* InAvaScene)
-	{
-		return InAvaScene ? InAvaScene->GetStartupCameraName() : NAME_None;
-	}
-
-	AActor* FindCameraActor(const ULevel* InLevel, FName InStartupCameraName)
-	{
-		if (InStartupCameraName != NAME_None)
-		{
-			AActor* CameraActor = FindActorByName(InLevel, InStartupCameraName);
-			if (CameraActor && IsCameraActor(CameraActor))
-			{
-				return CameraActor;
-			}
-		}
-		return nullptr;
-	}
-
-	AActor* FindStartupCameraActor(const ULevel* InLevel, FName InStartupCameraName, const AAvaScene* InAvaScene)
-	{
-		if (!IsValid(InLevel))
-		{
-			return nullptr;
-		}
-
-		// 1- Find by name from the parent scene.
-		if (AActor* FoundCameraActor = FindCameraActor(InLevel, InStartupCameraName))
-		{
-			return FoundCameraActor;
-		}
-
-		// 2- Find by name from dependent scene (if any). Dependent levels may not have an ava scene.
-		if (AActor* FoundCameraActor = FindCameraActor(InLevel, GetStartupCameraName(InAvaScene)))
-		{
-			return FoundCameraActor;
-		}
-
-		// 3- For now we fallback to looking for the first camera actor in the level.
-		return FindFirstCameraActor(InLevel);
-	}
-
 	EAvaPlayableStatus GetPlayableStatusFromLevelStreamingState(ELevelStreamingState InState, bool bInShouldBeLoaded)
 	{
 		switch (InState)
@@ -605,37 +535,6 @@ void UAvaPlayableLevelStreaming::ResolveScene(const ULevel* InLevel)
 	}
 }
 
-AActor* UAvaPlayableLevelStreaming::FindStartupCameraActor(FName InStartupCameraName, UAvaPlayable** OutParentPlayable)
-{
-	using namespace UE::AvaMedia::LevelStreamingPlayable;
-
-	// TODO: Configurable priorities for camera selection.
-	
-	if (InStartupCameraName == NAME_None && Scene)
-	{
-		InStartupCameraName = Scene->GetStartupCameraName();
-	}
-	
-	if (AActor* FoundCameraActor = Private::FindStartupCameraActor(GetLoadedLevel(), InStartupCameraName, Scene))
-	{
-		if (OutParentPlayable)
-		{
-			*OutParentPlayable = this;
-		}
-		return FoundCameraActor;
-	}
-
-	for (UAvaPlayableLevelStreaming* SubPlayable : SubPlayables)
-	{
-		if (AActor* FoundCameraActor = SubPlayable->FindStartupCameraActor(InStartupCameraName, OutParentPlayable))
-		{
-			return FoundCameraActor;
-		}
-	}
-	
-	return nullptr;
-}
-
 bool UAvaPlayableLevelStreaming::InitPlayable(const FPlayableCreationInfo& InPlayableInfo)
 {
 	// For now, we share all the levels in the same instance group. We may do sub-grouping later.
diff --git a/Engine/Plugins/Experimental/Avalanche/Source/AvalancheMedia/Private/Playable/Playables/AvaPlayableLevelStreaming.h b/Engine/Plugins/Experimental/Avalanche/Source/AvalancheMedia/Private/Playable/Playables/AvaPlayableLevelStreaming.h
--- a/Engine/Plugins/Experimental/Avalanche/Source/AvalancheMedia/Private/Playable/Playables/AvaPlayableLevelStreaming.h
+++ b/Engine/Plugins/Experimental/Avalanche/Source/AvalancheMedia/Private/Playable/Playables/AvaPlayableLevelStreaming.h
@@ -65,7 +65,6 @@ protected:
 	void UnbindDelegates();
 	ULevel* GetLoadedLevel() const;
 	void ResolveScene(const ULevel* InLevel);
-	AActor* FindStartupCameraActor(FName InStartupCameraName, UAvaPlayable** OutParentPlayable);
 
 	void LoadSubPlayables(const UWorld* InLevelInstance);
 	void UnloadSubPlayables();
diff --git a/Engine/Plugins/Experimental/Avalanche/Source/AvalancheViewport/Private/Interaction/AvaSnapOperation.cpp b/Engine/Plugins/Experimental/Avalanche/Source/AvalancheViewport/Private/Interaction/AvaSnapOperation.cpp
--- a/Engine/Plugins/Experimental/Avalanche/Source/AvalancheViewport/Private/Interaction/AvaSnapOperation.cpp
+++ b/Engine/Plugins/Experimental/Avalanche/Source/AvalancheViewport/Private/Interaction/AvaSnapOperation.cpp
@@ -24,7 +24,6 @@
 #include "Selection/AvaSelectionProviderSubsystem.h"
 #include "Viewport/Interaction/AvaSnapPoint.h"
 #include "Viewport/Interaction/IAvaSnapPointGenerator.h"
-#include "Viewport/Interaction/IAvaViewportDataProvider.h"
 #include "ViewportClient/IAvaViewportClient.h"
 
 namespace UE::AvaViewport::Private
diff --git a/Engine/Plugins/Experimental/Avalanche/Source/AvalancheViewport/Private/ViewportClient/AvaViewportClientUtilityProvider.cpp b/Engine/Plugins/Experimental/Avalanche/Source/AvalancheViewport/Private/ViewportClient/AvaViewportClientUtilityProvider.cpp
--- a/Engine/Plugins/Experimental/Avalanche/Source/AvalancheViewport/Private/ViewportClient/AvaViewportClientUtilityProvider.cpp
+++ b/Engine/Plugins/Experimental/Avalanche/Source/AvalancheViewport/Private/ViewportClient/AvaViewportClientUtilityProvider.cpp
@@ -9,7 +9,6 @@
 #include "Math/Transform.h"
 #include "Math/Vector.h"
 #include "Math/Vector2D.h"
-#include "Viewport/Interaction/IAvaViewportDataProxy.h"
 
 FVector2f FAvaViewportClientUtilityProvider::GetViewportWidgetSize() const
 {
@@ -186,16 +185,6 @@ void FAvaViewportClientUtilityProvider::WorldPositionToViewportPosition(const FV
 	}
 }
 
-IAvaViewportDataProvider* FAvaViewportClientUtilityProvider::GetViewportDataProvider() const
-{
-	if (TSharedPtr<IAvaViewportDataProxy> DataProxy = GetViewportDataProxy())
-	{
-		return DataProxy->GetViewportDataProvider();
-	}
-
-	return nullptr;
-}
-
 void FAvaViewportClientUtilityProvider::OnCameraCut(AActor* InTarget, bool bInJumpCut)
 {
 	if (!IsValid(InTarget))
diff --git a/Engine/Plugins/Experimental/Avalanche/Source/AvalancheViewport/Private/ViewportClient/EditorViewportClientUtilityWrapper.h b/Engine/Plugins/Experimental/Avalanche/Source/AvalancheViewport/Private/ViewportClient/EditorViewportClientUtilityWrapper.h
--- a/Engine/Plugins/Experimental/Avalanche/Source/AvalancheViewport/Private/ViewportClient/EditorViewportClientUtilityWrapper.h
+++ b/Engine/Plugins/Experimental/Avalanche/Source/AvalancheViewport/Private/ViewportClient/EditorViewportClientUtilityWrapper.h
@@ -43,8 +43,6 @@ public:
 	virtual FVector2f GetConstrainedViewportMousePosition() const override;
 	virtual FVector2f GetUnconstrainedZoomedViewportMousePosition() const override;
 	virtual FVector2f GetConstrainedZoomedViewportMousePosition() const override;
-	virtual TSharedPtr<IAvaViewportDataProxy> GetViewportDataProxy() const override { return nullptr; } // No support
-	virtual void SetViewportDataProxy(const TSharedPtr<IAvaViewportDataProxy>& InDataProxy) override {} // No support
 	virtual TSharedPtr<FAvaSnapOperation> GetSnapOperation() const override { return nullptr; } // No support
 	virtual TSharedPtr<FAvaSnapOperation> StartSnapOperation() override { return nullptr; } // No support
 	virtual bool EndSnapOperation(FAvaSnapOperation* InSnapOperation) override { return false; } // No support
diff --git a/Engine/Plugins/Experimental/Avalanche/Source/AvalancheViewport/Public/ViewportClient/AvaViewportClientUtilityProvider.h b/Engine/Plugins/Experimental/Avalanche/Source/AvalancheViewport/Public/ViewportClient/AvaViewportClientUtilityProvider.h
--- a/Engine/Plugins/Experimental/Avalanche/Source/AvalancheViewport/Public/ViewportClient/AvaViewportClientUtilityProvider.h
+++ b/Engine/Plugins/Experimental/Avalanche/Source/AvalancheViewport/Public/ViewportClient/AvaViewportClientUtilityProvider.h
@@ -24,7 +24,6 @@ public:
 	virtual FVector2f GetVirtualViewportScale() const override;
 	virtual float GetAverageVirtualViewportScale() const override;
 	virtual FVector2D GetZoomedFrustumSizeAtDistance(double InDistance) const override;
-	virtual IAvaViewportDataProvider* GetViewportDataProvider() const override;
 	virtual void OnCameraCut(AActor* InTarget, bool bInJumpCut) override;
 	//~ End IAvaViewportClient
 
diff --git a/Engine/Plugins/Experimental/Avalanche/Source/AvalancheViewport/Public/ViewportClient/IAvaViewportClient.h b/Engine/Plugins/Experimental/Avalanche/Source/AvalancheViewport/Public/ViewportClient/IAvaViewportClient.h
--- a/Engine/Plugins/Experimental/Avalanche/Source/AvalancheViewport/Public/ViewportClient/IAvaViewportClient.h
+++ b/Engine/Plugins/Experimental/Avalanche/Source/AvalancheViewport/Public/ViewportClient/IAvaViewportClient.h
@@ -15,8 +15,6 @@ class FAvaSnapOperation;
 class FAvaViewportPostProcessManager;
 class FEditorViewportClient;
 class IAvaBoundsProviderInterface;
-class IAvaViewportDataProvider;
-class IAvaViewportDataProxy;
 class ILevelEditor;
 class UCameraComponent;
 struct FAvaVisibleArea;
@@ -107,15 +105,6 @@ public:
 	 */
 	virtual FVector2f GetConstrainedZoomedViewportMousePosition() const = 0;
 
-	/** Provides access to any viewport-specific data that is available through a third party. */
-	virtual TSharedPtr<IAvaViewportDataProxy> GetViewportDataProxy() const = 0;
-
-	/** Sets the third party viewport data provider. */
-	virtual void SetViewportDataProxy(const TSharedPtr<IAvaViewportDataProxy>& InDataProxy) = 0;
-
-	/** Uses the data proxy to fetch the data provider, if available. */
-	virtual IAvaViewportDataProvider* GetViewportDataProvider() const = 0;
-
 	/** Gets an active snap operation on this viewport. */
 	virtual TSharedPtr<FAvaSnapOperation> GetSnapOperation() const = 0;
 
