Motion Design: Viewport Transactions
- Fixed some viewport transaction issues related to post process and rulers.

diff --git a/Engine/Plugins/Experimental/Avalanche/Source/AvalancheLevelViewport/Private/SAvaLevelViewport.cpp b/Engine/Plugins/Experimental/Avalanche/Source/AvalancheLevelViewport/Private/SAvaLevelViewport.cpp
--- a/Engine/Plugins/Experimental/Avalanche/Source/AvalancheLevelViewport/Private/SAvaLevelViewport.cpp
+++ b/Engine/Plugins/Experimental/Avalanche/Source/AvalancheLevelViewport/Private/SAvaLevelViewport.cpp
@@ -567,12 +567,6 @@ void SAvaLevelViewport::SaveGuides()
 	Data->GuideData = GuideInfos;
 }
 
-bool SAvaLevelViewport::MatchesContext(const FTransactionContext& InContext,
-	const TArray<TPair<UObject*, FTransactionObjectEvent>>& InTransactionObjectContexts) const
-{
-	return InContext.Title.EqualTo(UE::AvaLevelViewport::Private::GuideTransactionTitle);
-}
-
 void SAvaLevelViewport::PostUndo(bool bInSuccess)
 {
 	ReloadGuides();
diff --git a/Engine/Plugins/Experimental/Avalanche/Source/AvalancheLevelViewport/Private/SAvaLevelViewport.h b/Engine/Plugins/Experimental/Avalanche/Source/AvalancheLevelViewport/Private/SAvaLevelViewport.h
--- a/Engine/Plugins/Experimental/Avalanche/Source/AvalancheLevelViewport/Private/SAvaLevelViewport.h
+++ b/Engine/Plugins/Experimental/Avalanche/Source/AvalancheLevelViewport/Private/SAvaLevelViewport.h
@@ -1,4 +1,4 @@
-﻿// Copyright Epic Games, Inc. All Rights Reserved.
+// Copyright Epic Games, Inc. All Rights Reserved.
 
 #pragma once
 
@@ -154,8 +154,6 @@ public:
 	void ResetPilotedCameraTransform();
 
 	//~ Begin FSelfRegisteringEditorUndoClient
-	virtual bool MatchesContext(const FTransactionContext& InContext, const TArray<TPair<UObject*, 
-		FTransactionObjectEvent>>& InTransactionObjectContexts) const override;
 	virtual void PostUndo(bool bInSuccess);
 	virtual void PostRedo(bool bInSuccess);
 	//~ End FSelfRegisteringEditorUndoClient
diff --git a/Engine/Plugins/Experimental/Avalanche/Source/AvalancheLevelViewport/Private/SAvaLevelViewport_Menus.cpp b/Engine/Plugins/Experimental/Avalanche/Source/AvalancheLevelViewport/Private/SAvaLevelViewport_Menus.cpp
--- a/Engine/Plugins/Experimental/Avalanche/Source/AvalancheLevelViewport/Private/SAvaLevelViewport_Menus.cpp
+++ b/Engine/Plugins/Experimental/Avalanche/Source/AvalancheLevelViewport/Private/SAvaLevelViewport_Menus.cpp
@@ -967,9 +967,24 @@ void SAvaLevelViewport::OnVirtualSizeSliderBegin()
 		return;
 	}
 
+	UAvaViewportDataSubsystem* DataSubsystem = UAvaViewportDataSubsystem::Get(ViewportClient->GetViewportWorld());
+
+	if (!DataSubsystem)
+	{
+		return;
+	}
+
+	FAvaViewportData* Data = DataSubsystem->GetData();
+
+	if (!Data)
+	{
+		return;
+	}
+
 	VirtualSizeComponentSliderTransaction = MakeShared<FScopedTransaction>(
 		LOCTEXT("SetCustomVirtualSize", "Set Custom Virtual Size")
 	);
+	DataSubsystem->ModifyDataSource();
 }
 
 void SAvaLevelViewport::OnVirtualSizeComponentSliderEnd(int32 InNewSize)
diff --git a/Engine/Plugins/Experimental/Avalanche/Source/AvalancheLevelViewport/Private/SAvaLevelViewport_StatusBar.cpp b/Engine/Plugins/Experimental/Avalanche/Source/AvalancheLevelViewport/Private/SAvaLevelViewport_StatusBar.cpp
--- a/Engine/Plugins/Experimental/Avalanche/Source/AvalancheLevelViewport/Private/SAvaLevelViewport_StatusBar.cpp
+++ b/Engine/Plugins/Experimental/Avalanche/Source/AvalancheLevelViewport/Private/SAvaLevelViewport_StatusBar.cpp
@@ -1,6 +1,7 @@
 // Copyright Epic Games, Inc. All Rights Reserved.
 
 #include "AvaLevelViewportModule.h"
+#include "AvaViewportDataSubsystem.h"
 #include "SAvaLevelViewport.h"
 #include "AvaViewportPostProcessManager.h"
 #include "AvaViewportSettings.h"
@@ -486,6 +487,11 @@ float SAvaLevelViewport::GetBackgroundOpacity() const
 
 void SAvaLevelViewport::BeginPostProcessInfoTransaction()
 {
+	if (PostProcessInfoTransaction.IsValid())
+	{
+		return;
+	}
+
 	TSharedPtr<FAvaLevelViewportClient> ViewportClient = GetAvaLevelViewportClient();
 
 	if (!ViewportClient.IsValid())
@@ -493,10 +499,22 @@ void SAvaLevelViewport::BeginPostProcessInfoTransaction()
 		return;
 	}
 
-	if (!PostProcessInfoTransaction.IsValid())
+	UAvaViewportDataSubsystem* DataSubsystem = UAvaViewportDataSubsystem::Get(ViewportClient->GetViewportWorld());
+
+	if (!DataSubsystem)
+	{
+		return;
+	}
+
+	FAvaViewportData* Data = DataSubsystem->GetData();
+
+	if (!Data)
 	{
-		PostProcessInfoTransaction = MakeShared<FScopedTransaction>(LOCTEXT("PostProcessSettingsChange", "Post Process Settings Change"));
+		return;
 	}
+
+	PostProcessInfoTransaction = MakeShared<FScopedTransaction>(LOCTEXT("PostProcessSettingsChange", "Post Process Settings Change"));
+	DataSubsystem->ModifyDataSource();
 }
 
 void SAvaLevelViewport::EndPostProcessInfoTransaction()
diff --git a/Engine/Plugins/Experimental/Avalanche/Source/AvalancheLevelViewport/Private/Widgets/SAvaLevelViewportStatusBarButtons.cpp b/Engine/Plugins/Experimental/Avalanche/Source/AvalancheLevelViewport/Private/Widgets/SAvaLevelViewportStatusBarButtons.cpp
--- a/Engine/Plugins/Experimental/Avalanche/Source/AvalancheLevelViewport/Private/Widgets/SAvaLevelViewportStatusBarButtons.cpp
+++ b/Engine/Plugins/Experimental/Avalanche/Source/AvalancheLevelViewport/Private/Widgets/SAvaLevelViewportStatusBarButtons.cpp
@@ -16,6 +16,7 @@
 #include "SAvaLevelViewport.h"
 #include "SAvaLevelViewportFrame.h"
 #include "SAvaViewportInfo.h"
+#include "ScopedTransaction.h"
 #include "Selection.h"
 #include "Styling/AppStyle.h"
 #include "Styling/SlateIconFinder.h"
@@ -732,6 +733,7 @@ ECheckBoxState SAvaLevelViewportStatusBarButtons::GetPostProcessActiveMenu(const
 
 void SAvaLevelViewportStatusBarButtons::TogglePostProcessMenu(const FToolMenuContext& InContext, EAvaViewportPostProcessType InPostProcessType)
 {
+	FScopedTransaction Transaction(LOCTEXT("ChangeViewportPostProcess", "Change Viewport Post Process"));
 	UE::Ava::LevelViewportStatusBarButtons::Private::TogglePostProcess(ViewportFrameWeak, InPostProcessType);
 }
 
diff --git a/Engine/Plugins/Experimental/Avalanche/Source/AvalancheViewport/Private/AvaViewportPostProcessManager.cpp b/Engine/Plugins/Experimental/Avalanche/Source/AvalancheViewport/Private/AvaViewportPostProcessManager.cpp
--- a/Engine/Plugins/Experimental/Avalanche/Source/AvalancheViewport/Private/AvaViewportPostProcessManager.cpp
+++ b/Engine/Plugins/Experimental/Avalanche/Source/AvalancheViewport/Private/AvaViewportPostProcessManager.cpp
@@ -132,7 +132,8 @@ void FAvaViewportPostProcessManager::SetType(EAvaViewportPostProcessType InType)
 	{
 		CurrentVisualizer->OnDeactivate();
 	}
-	
+
+	ModifyDataSource();
 	PostProcessInfo->Type = InType;
 
 	if (NewVisualizer.IsValid())
@@ -156,6 +157,7 @@ void FAvaViewportPostProcessManager::SetOpacity(float InOpacity)
 {
 	if (FAvaViewportPostProcessInfo* PostProcessInfo = GetPostProcessInfo())
 	{
+		ModifyDataSource();
 		PostProcessInfo->Opacity = InOpacity;
 		LoadPostProcessInfo();
 	}
@@ -165,6 +167,27 @@ void FAvaViewportPostProcessManager::SetOpacity(float InOpacity)
 	}
 }
 
+void FAvaViewportPostProcessManager::ModifyDataSource()
+{
+	TSharedPtr<IAvaViewportClient> AvaViewportClient = AvaViewportClientWeak.Pin();
+
+	if (!AvaViewportClient.IsValid())
+	{
+		UE_LOG(AvaViewportLog, Warning, TEXT("FAvaViewportPostProcessManager::ModifyDataSource: Invalid viewport client."));
+		return;
+	}
+
+	UAvaViewportDataSubsystem* DataSubsystem = UAvaViewportDataSubsystem::Get(AvaViewportClient->GetViewportWorld());
+
+	if (!DataSubsystem)
+	{
+		UE_LOG(AvaViewportLog, Warning, TEXT("FAvaViewportPostProcessManager::ModifyDataSource: Failed to find data subsystem."));
+		return;
+	}
+
+	DataSubsystem->ModifyDataSource();
+}
+
 TSharedPtr<IAvaViewportPostProcessVisualizer> FAvaViewportPostProcessManager::GetVisualizer(EAvaViewportPostProcessType InType) const
 {
 	if (const TSharedPtr<IAvaViewportPostProcessVisualizer>* VisualizerPtr = Visualizers.Find(InType))
diff --git a/Engine/Plugins/Experimental/Avalanche/Source/AvalancheViewport/Public/AvaViewportPostProcessManager.h b/Engine/Plugins/Experimental/Avalanche/Source/AvalancheViewport/Public/AvaViewportPostProcessManager.h
--- a/Engine/Plugins/Experimental/Avalanche/Source/AvalancheViewport/Public/AvaViewportPostProcessManager.h
+++ b/Engine/Plugins/Experimental/Avalanche/Source/AvalancheViewport/Public/AvaViewportPostProcessManager.h
@@ -70,4 +70,6 @@ protected:
 	TWeakPtr<IAvaViewportClient> AvaViewportClientWeak;
 
 	TMap<EAvaViewportPostProcessType, TSharedPtr<IAvaViewportPostProcessVisualizer>> Visualizers;
+
+	void ModifyDataSource();
 };
