[MotionDesign] Ease Curve Tool: support for multiple tracks

diff --git a/Engine/Plugins/Experimental/Avalanche/Source/AvalancheSequencer/Private/EaseCurveTool/AvaEaseCurveKeySelection.cpp b/Engine/Plugins/Experimental/Avalanche/Source/AvalancheSequencer/Private/EaseCurveTool/AvaEaseCurveKeySelection.cpp
--- a/Engine/Plugins/Experimental/Avalanche/Source/AvalancheSequencer/Private/EaseCurveTool/AvaEaseCurveKeySelection.cpp
+++ b/Engine/Plugins/Experimental/Avalanche/Source/AvalancheSequencer/Private/EaseCurveTool/AvaEaseCurveKeySelection.cpp
@@ -83,7 +83,11 @@ FAvaEaseCurveKeySelection::FAvaEaseCurveKeySelection(const TSharedPtr<FAvaSequen
 
 		const FMovieSceneChannelHandle& ChannelHandle = KeyArea->GetChannel();
 
-		FChannelData& Entry = ChannelKeyData.FindOrAdd(ChannelModel->GetChannelName());
+		const FString PathName = IOutlinerExtension::GetPathName(ChannelModel);
+		const FString ChannelName = ChannelModel->GetChannelName().ToString();
+		const FName FullPathMapKey = FName(PathName + ChannelName);
+
+		FChannelData& Entry = ChannelKeyData.FindOrAdd(FullPathMapKey);
 		Entry.ChannelModel = ChannelModel;
 		Entry.Section = ChannelModel->GetSection();
 		Entry.Channel = ChannelHandle;
