MotionDesign : Modifiers
- Fixed boolean modifier overlapping check to support Text3D and its various disconnected glyphs not being detected due to tool bounding box being in between glyphs and not intersecting any geometry

diff --git a/Engine/Plugins/Experimental/Avalanche/Source/AvalancheModifiers/Private/Shared/AvaBooleanModifierShared.cpp b/Engine/Plugins/Experimental/Avalanche/Source/AvalancheModifiers/Private/Shared/AvaBooleanModifierShared.cpp
--- a/Engine/Plugins/Experimental/Avalanche/Source/AvalancheModifiers/Private/Shared/AvaBooleanModifierShared.cpp
+++ b/Engine/Plugins/Experimental/Avalanche/Source/AvalancheModifiers/Private/Shared/AvaBooleanModifierShared.cpp
@@ -141,45 +141,95 @@ bool UAvaBooleanModifierShared::TestIntersection(const UE::Geometry::FDynamicMes
 		return false;
 	}
 
-	FDynamicMeshAABBTree3 Spatials[2];
-	ParallelFor(2, [&Spatials, &InToolMesh, &InTargetMesh](int32 k)
+	FDynamicMeshAABBTree3 MeshSpatialTrees[2];
+	ParallelFor(2, [&MeshSpatialTrees, &InToolMesh, &InTargetMesh](int32 k)
 	{
-		Spatials[k].SetMesh(k == 0 ? &InToolMesh : &InTargetMesh, true);
+		MeshSpatialTrees[k].SetMesh(k == 0 ? &InToolMesh : &InTargetMesh, true);
 	}, EParallelForFlags::Unbalanced);
 
-	const bool bIsIdentity1 = InToolTransform.Equals(FTransform::Identity, 0);
-	FTransformSRT3d Transform1(InToolTransform);
+	const bool bIsIdentityTool = InToolTransform.Equals(FTransform::Identity, 0);
+	FTransformSRT3d ToolTransform(InToolTransform);
 
-	const bool bIsIdentity2 = InTargetTransform.Equals(FTransform::Identity, 0);
-	FTransformSRT3d Transform2(InTargetTransform);
+	const bool bIsIdentityTarget = InTargetTransform.Equals(FTransform::Identity, 0);
+	FTransformSRT3d TargetTransform(InTargetTransform);
 
 	bool bIsIntersecting = false;
-	if (bIsIdentity1 && bIsIdentity2)
+	if (bIsIdentityTool && bIsIdentityTarget)
 	{
-		bIsIntersecting = Spatials[0].TestIntersection(Spatials[1]);
+		bIsIntersecting = MeshSpatialTrees[0].TestIntersection(MeshSpatialTrees[1]);
 	}
-	else if (bIsIdentity1 || bIsIdentity2)
+	else if (bIsIdentityTool || bIsIdentityTarget)
 	{
-		const FIndex2i Indices = (bIsIdentity1) ? FIndex2i(0,1) : FIndex2i(1,0);
-		FTransformSRT3d UseTransform = (bIsIdentity2) ? Transform1 : Transform2;
-		bIsIntersecting = Spatials[Indices.A].TestIntersection(Spatials[Indices.B],
+		const FIndex2i Indices = (bIsIdentityTool) ? FIndex2i(0,1) : FIndex2i(1,0);
+		FTransformSRT3d UseTransform = (bIsIdentityTarget) ? ToolTransform : TargetTransform;
+		bIsIntersecting = MeshSpatialTrees[Indices.A].TestIntersection(MeshSpatialTrees[Indices.B],
 			[&UseTransform](const FVector3d& Pos) { return UseTransform.TransformPosition(Pos); });
 	}
 	else
 	{
-		bIsIntersecting = Spatials[0].TestIntersection(Spatials[1],
-			[&Transform1, &Transform2](const FVector3d& Pos) { return Transform1.InverseTransformPosition(Transform2.TransformPosition(Pos)); });
+		bIsIntersecting = MeshSpatialTrees[0].TestIntersection(MeshSpatialTrees[1],
+			[&ToolTransform, &TargetTransform](const FVector3d& Pos) { return ToolTransform.InverseTransformPosition(TargetTransform.TransformPosition(Pos)); });
 	}
 
-	// We are not intersecting but target is perhaps contained inside tool
+	// We are not intersecting but target is perhaps contained fully inside tool or partially inside tool (for mesh with disconnected parts)
 	if (!bIsIntersecting)
 	{
-		const FBox ToolBounds = static_cast<FBox>(InToolMesh.GetBounds(true));
-		const FVector ActorOffset = InToolTransform.GetLocation() - InTargetTransform.GetLocation();
+		const FBox ToolBounds = static_cast<FBox>(MeshSpatialTrees[0].GetBoundingBox()).TransformBy(ToolTransform);
+		const FBox TargetBounds = static_cast<FBox>(MeshSpatialTrees[1].GetBoundingBox()).TransformBy(TargetTransform);
 
-		if (ToolBounds.IsInside(ActorOffset))
+		// Compute tool centroid
+		FVector3d ToolCentroid(0.0, 0.0, 0.0);
 		{
-			bIsIntersecting = true;
+			double TotalTotalArea = 0.0;
+		
+			for (int32 TId : InToolMesh.TriangleIndicesItr())
+			{
+				FVector3d TriCentroid = InToolMesh.GetTriCentroid(TId);
+				double TriArea = InToolMesh.GetTriArea(TId);
+
+				ToolCentroid += TriCentroid * TriArea;
+				TotalTotalArea += TriArea;
+			}
+
+			if (TotalTotalArea != 0)
+			{
+				ToolCentroid /= TotalTotalArea;
+			}
+		}
+
+		// Ensure target bounding box intersect or is inside tool
+		if (ToolBounds.Intersect(TargetBounds) || ToolBounds.IsInside(TargetBounds))
+		{
+			const FVector3d ToolCenter = ToolTransform.TransformPosition(ToolCentroid);
+
+			// Find target nearest triangle from tool center
+			double DistSquaredTarget;
+			const int32 TargetMeshTId = MeshSpatialTrees[1].FindNearestTriangle(TargetTransform.InverseTransformPosition(ToolCenter), DistSquaredTarget);
+
+			if (TargetMeshTId != INDEX_NONE)
+			{
+				const FVector3d ClosestTargetCentroid = TargetTransform.TransformPosition(InTargetMesh.GetTriCentroid(TargetMeshTId));
+
+				// Ensure target centroid is inside tool
+				if (ToolBounds.IsInside(ClosestTargetCentroid))
+				{
+					FRay3d ToolRay(ToolCentroid, ToolTransform.InverseTransformPosition(ClosestTargetCentroid).GetSafeNormal());
+					double ToolNearestT = 0.f;
+					int32 ToolNearestTId = INDEX_NONE;
+					FVector3d ToolBaryCoords = FVector3d::ZeroVector;
+
+					// Find tool nearest triangle based on ray projected in nearest target triangle direction
+					if (MeshSpatialTrees[0].FindNearestHitTriangle(ToolRay, ToolNearestT, ToolNearestTId, ToolBaryCoords))
+					{
+						const double DistanceToolCenterToTarget = FVector3d::Distance(ToolCenter, ClosestTargetCentroid);
+						
+						if (DistanceToolCenterToTarget < ToolNearestT)
+						{
+							bIsIntersecting = true;
+						}
+					}
+				}
+			}
 		}
 	}
 
