MotionDesign : Modifiers
- Fix translucent sort priority modifier, first two actors would have same sort priority

diff --git a/Engine/Plugins/Experimental/Avalanche/Source/AvalancheModifiers/Private/Modifiers/AvaTranslucentPriorityModifier.cpp b/Engine/Plugins/Experimental/Avalanche/Source/AvalancheModifiers/Private/Modifiers/AvaTranslucentPriorityModifier.cpp
--- a/Engine/Plugins/Experimental/Avalanche/Source/AvalancheModifiers/Private/Modifiers/AvaTranslucentPriorityModifier.cpp
+++ b/Engine/Plugins/Experimental/Avalanche/Source/AvalancheModifiers/Private/Modifiers/AvaTranslucentPriorityModifier.cpp
@@ -226,12 +226,7 @@ void UAvaTranslucentPriorityModifier::Apply()
 
 				ComponentModifier->MarkModifierDirty();
 			}
-
-			if (!PreviousComponentOwner)
-			{
-				PreviousComponentOwner = Component->GetOwner();
-			}
-
+			
 			// Only increase when new actor detected or mode is camera
 			if (Mode == EAvaTranslucentPriorityModifierMode::AutoCameraDistance
 				|| PreviousComponentOwner != Component->GetOwner())
