[Motion Design] Rundown Editor's Page list: hiding the checkbox in the header row.

The checkbox does nothing and it is reported as a bug.

diff --git a/Engine/Plugins/Experimental/Avalanche/Source/AvalancheMediaEditor/Private/Rundown/Pages/Columns/AvaRundownPageEnabledColumn.cpp b/Engine/Plugins/Experimental/Avalanche/Source/AvalancheMediaEditor/Private/Rundown/Pages/Columns/AvaRundownPageEnabledColumn.cpp
--- a/Engine/Plugins/Experimental/Avalanche/Source/AvalancheMediaEditor/Private/Rundown/Pages/Columns/AvaRundownPageEnabledColumn.cpp
+++ b/Engine/Plugins/Experimental/Avalanche/Source/AvalancheMediaEditor/Private/Rundown/Pages/Columns/AvaRundownPageEnabledColumn.cpp
@@ -30,6 +30,7 @@ SHeaderRow::FColumn::FArguments FAvaRundownPageEnabledColumn::ConstructHeaderRow
 		[
 			SNew(SCheckBox)
 			.IsChecked(true)
+			.Visibility(EVisibility::Hidden)
 		];
 }
 
