Motion Design: Fixed issue with Motion Design Sequencer opening causing Analytic Trace to generate incorrect data on Open Asset Editor

diff --git a/Engine/Plugins/Experimental/Avalanche/Source/AvalancheEditor/Private/Sequencer/AvaLevelSequencerExtension.cpp b/Engine/Plugins/Experimental/Avalanche/Source/AvalancheEditor/Private/Sequencer/AvaLevelSequencerExtension.cpp
--- a/Engine/Plugins/Experimental/Avalanche/Source/AvalancheEditor/Private/Sequencer/AvaLevelSequencerExtension.cpp
+++ b/Engine/Plugins/Experimental/Avalanche/Source/AvalancheEditor/Private/Sequencer/AvaLevelSequencerExtension.cpp
@@ -38,11 +38,18 @@ namespace UE::AvaEditor::Private
 		UAssetEditorSubsystem* AssetEditorSubsystem = GEditor->GetEditorSubsystem<UAssetEditorSubsystem>();
 		check(AssetEditorSubsystem);
 
-		AssetEditorSubsystem->OpenEditorForAsset(InSequence);
-
 		constexpr bool bFocusIfOpen = false;
-		IAssetEditorInstance* const AssetEditor = AssetEditorSubsystem->FindEditorForAsset(InSequence, bFocusIfOpen);
+		IAssetEditorInstance* AssetEditor = AssetEditorSubsystem->FindEditorForAsset(InSequence, bFocusIfOpen);
+
+		// Return early if it's already opened
+		if (AssetEditor)
+		{
+			return static_cast<ILevelSequenceEditorToolkit*>(AssetEditor);
+		}
+
+		AssetEditorSubsystem->OpenEditorForAsset(InSequence);
 
+		AssetEditor = AssetEditorSubsystem->FindEditorForAsset(InSequence, bFocusIfOpen);
 		if (ensureAlwaysMsgf(AssetEditor, TEXT("Unable to find Editor for Sequence: %s"), *InSequence->GetName()))
 		{
 			return static_cast<ILevelSequenceEditorToolkit*>(AssetEditor);
