[Motion Design] Fix Rundown Editor Crash when adding a page in a page view.

Regression introduced by CL 36655633.

diff --git a/Engine/Plugins/Experimental/Avalanche/Source/AvalancheMediaEditor/Private/Rundown/Pages/Slate/SAvaRundownInstancedPageList.cpp b/Engine/Plugins/Experimental/Avalanche/Source/AvalancheMediaEditor/Private/Rundown/Pages/Slate/SAvaRundownInstancedPageList.cpp
--- a/Engine/Plugins/Experimental/Avalanche/Source/AvalancheMediaEditor/Private/Rundown/Pages/Slate/SAvaRundownInstancedPageList.cpp
+++ b/Engine/Plugins/Experimental/Avalanche/Source/AvalancheMediaEditor/Private/Rundown/Pages/Slate/SAvaRundownInstancedPageList.cpp
@@ -176,7 +176,7 @@ void SAvaRundownInstancedPageList::Refresh()
 			const FAvaRundownSubList& SubList = Rundown->GetSubList(PageListReference);
 			const int32 VisiblePageCount = VisiblePageIds.IsEmpty() ? SubList.PageIds.Num() : VisiblePageIds.Num();
 
-			if (SubList.PageIds.Num() != VisiblePageCount)
+			if (PageViews.Num() != VisiblePageCount)
 			{
 				PageViews.Reset(SubList.PageIds.Num());
 			
