MotionDesign : Modifiers
- Fix base arrange modifier will not mark itself dirty when modified actor is moved, this will avoid performance impact when moving actors with sequencer that have arrange modifiers on them, to avoid modifying headers, virtual functions override will be removed in 5.6

diff --git a/Engine/Plugins/Experimental/Avalanche/Source/AvalancheModifiers/Private/Modifiers/AvaArrangeBaseModifier.cpp b/Engine/Plugins/Experimental/Avalanche/Source/AvalancheModifiers/Private/Modifiers/AvaArrangeBaseModifier.cpp
--- a/Engine/Plugins/Experimental/Avalanche/Source/AvalancheModifiers/Private/Modifiers/AvaArrangeBaseModifier.cpp
+++ b/Engine/Plugins/Experimental/Avalanche/Source/AvalancheModifiers/Private/Modifiers/AvaArrangeBaseModifier.cpp
@@ -43,8 +43,6 @@ void UAvaArrangeBaseModifier::OnModifierDisabled(EActorModifierCoreDisableReason
 void UAvaArrangeBaseModifier::OnModifiedActorTransformed()
 {
 	Super::OnModifiedActorTransformed();
-
-	MarkModifierDirty();
 }
 
 void UAvaArrangeBaseModifier::OnSceneTreeTrackedActorChildrenChanged(int32 InIdx,
