Motion Design: Fix issue with FindOrAdd Component in Mask Modifier creating a component on save.

diff --git a/Engine/Plugins/Experimental/Avalanche/Source/AvalancheMask/Private/AvaMaskUtilities.cpp b/Engine/Plugins/Experimental/Avalanche/Source/AvalancheMask/Private/AvaMaskUtilities.cpp
--- a/Engine/Plugins/Experimental/Avalanche/Source/AvalancheMask/Private/AvaMaskUtilities.cpp
+++ b/Engine/Plugins/Experimental/Avalanche/Source/AvalancheMask/Private/AvaMaskUtilities.cpp
@@ -73,7 +73,9 @@ namespace UE::AvaMask::Internal
 		}
 
 		UActorComponent* Component = InActor->FindComponentByClass(InComponentClass);
-		if (!Component)
+
+        // If no existing component is found, create it only if the package the actor is in isn't being saved
+		if (!Component && !UE::IsSavingPackage(InActor))
 		{
 #if WITH_EDITOR
 			InActor->Modify(/*bAlwaysMarkDirty*/false);
