MotionDesign : Modifiers
- Fix dynamic mesh converter with Text3D static meshes
- Improving cloner mesh builder to filter/exclude components to convert
- Detect material changes on primitive components to mark modifier dirty

diff --git a/Engine/Plugins/Experimental/Avalanche/Config/DefaultAvalanche.ini b/Engine/Plugins/Experimental/Avalanche/Config/DefaultAvalanche.ini
--- a/Engine/Plugins/Experimental/Avalanche/Config/DefaultAvalanche.ini
+++ b/Engine/Plugins/Experimental/Avalanche/Config/DefaultAvalanche.ini
@@ -301,5 +301,7 @@
 +StructRedirects=(OldName="/Script/AvalancheMediaEditor.AvaRundownGetChannelImage",NewName="/Script/AvalancheMedia.AvaRundownGetChannelImage")
 +StructRedirects=(OldName="/Script/AvalancheMediaEditor.AvaRundownGetDevices",NewName="/Script/AvalancheMedia.AvaRundownGetDevices")
 
++PropertyRedirects=(OldName="/Script/AvalancheModifiers.AvaDynamicMeshConverterModifier.TransformUpdateInterval",NewName="/Script/AvalancheModifiers.AvaDynamicMeshConverterModifier.UpdateInterval")
+
 [/Script/Avalanche.AvaAnimation]
 DefaultCompletionMode=KeepState
diff --git a/Engine/Plugins/Experimental/Avalanche/Source/AvalancheModifiers/Private/Modifiers/AvaDynamicMeshConverterModifier.cpp b/Engine/Plugins/Experimental/Avalanche/Source/AvalancheModifiers/Private/Modifiers/AvaDynamicMeshConverterModifier.cpp
--- a/Engine/Plugins/Experimental/Avalanche/Source/AvalancheModifiers/Private/Modifiers/AvaDynamicMeshConverterModifier.cpp
+++ b/Engine/Plugins/Experimental/Avalanche/Source/AvalancheModifiers/Private/Modifiers/AvaDynamicMeshConverterModifier.cpp
@@ -3,16 +3,12 @@
 #include "Modifiers/AvaDynamicMeshConverterModifier.h"
 
 #include "Async/Async.h"
-#include "Components/BrushComponent.h"
 #include "Components/DynamicMeshComponent.h"
-#include "Components/SkeletalMeshComponent.h"
-#include "Components/StaticMeshComponent.h"
 #include "DynamicMesh/DynamicMesh3.h"
 #include "Engine/StaticMesh.h"
 #include "Extensions/AvaSceneTreeUpdateModifierExtension.h"
 #include "GeometryScript/MeshAssetFunctions.h"
 #include "GeometryScript/MeshBasicEditFunctions.h"
-#include "ProceduralMeshComponent.h"
 
 #if WITH_EDITOR
 #include "AssetRegistry/AssetRegistryModule.h"
@@ -27,26 +23,24 @@ FAvaDynamicMeshConverterModifierComponentState::FAvaDynamicMeshConverterModifier
 {
 	if (Component.IsValid())
 	{
-		if (const AActor* ComponentOwner = InPrimitiveComponent->GetOwner())
-		{
-			bActorHiddenInGame = ComponentOwner->IsHidden();
-#if WITH_EDITOR
-			bActorHiddenInEditor = ComponentOwner->IsTemporarilyHiddenInEditor();
-#endif
-			if (const USceneComponent* RootComponent = ComponentOwner->GetRootComponent())
-			{
-				bComponentVisible = RootComponent->IsVisible();
-				bComponentHiddenInGame = RootComponent->bHiddenInGame;
-			}
+		bComponentVisible = InPrimitiveComponent->GetVisibleFlag();
+		bComponentHiddenInGame = InPrimitiveComponent->bHiddenInGame;
 
-			if (const AActor* ParentActor = ComponentOwner->GetAttachParentActor())
-			{
-				ActorRelativeTransform = ComponentOwner->GetActorTransform().GetRelativeTransform(ParentActor->GetActorTransform());
-			}
+		for (int32 Index = 0; Index < InPrimitiveComponent->GetNumMaterials(); Index++)
+		{
+			ComponentMaterialsWeak.Add(InPrimitiveComponent->GetMaterial(Index));
 		}
 	}
 }
 
+void FAvaDynamicMeshConverterModifierComponentState::UpdateRelativeTransform(const FTransform& InParentTransform)
+{
+	if (const UPrimitiveComponent* PrimitiveComponent = Component.Get())
+	{
+		ActorRelativeTransform = PrimitiveComponent->GetComponentTransform().GetRelativeTransform(InParentTransform);
+	}
+}
+
 void UAvaDynamicMeshConverterModifier::OnModifierCDOSetup(FActorModifierCoreMetadata& InMetadata)
 {
 	Super::OnModifierCDOSetup(InMetadata);
@@ -83,6 +77,29 @@ void UAvaDynamicMeshConverterModifier::OnModifierAdded(EActorModifierCoreEnableR
 void UAvaDynamicMeshConverterModifier::OnModifierEnabled(EActorModifierCoreEnableReason InReason)
 {
 	Super::OnModifierEnabled(InReason);
+
+	if (bComponentCreated)
+	{
+		if (UDynamicMeshComponent* PrimitiveComponent = GetMeshComponent())
+		{
+			PrimitiveComponent->SetVisibleFlag(true);
+			PrimitiveComponent->SetHiddenInGame(false);
+		}
+	}
+}
+
+void UAvaDynamicMeshConverterModifier::OnModifierDisabled(EActorModifierCoreDisableReason InReason)
+{
+	Super::OnModifierDisabled(InReason);
+
+	if (bComponentCreated)
+	{
+		if (UDynamicMeshComponent* PrimitiveComponent = GetMeshComponent())
+		{
+			PrimitiveComponent->SetVisibleFlag(false);
+			PrimitiveComponent->SetHiddenInGame(true);
+		}
+	}
 }
 
 void UAvaDynamicMeshConverterModifier::RestorePreState()
@@ -91,24 +108,10 @@ void UAvaDynamicMeshConverterModifier::RestorePreState()
 
 	for (FAvaDynamicMeshConverterModifierComponentState& ConvertedComponent : ConvertedComponents)
 	{
-		if (ConvertedComponent.Component.IsValid())
+		if (UPrimitiveComponent* PrimitiveComponent = ConvertedComponent.Component.Get())
 		{
-			AActor* ComponentActor = ConvertedComponent.Component->GetOwner();
-
-			if (ComponentActor != GetModifiedActor())
-			{
-				// Hide actor but do not hide ourselves
-				ComponentActor->SetHidden(ConvertedComponent.bActorHiddenInGame);
-#if WITH_EDITOR
-				ComponentActor->SetIsTemporarilyHiddenInEditor(ConvertedComponent.bActorHiddenInEditor);
-#endif
-			}
-			else if (USceneComponent* RootComponent = ComponentActor->GetRootComponent())
-			{
-				// In the meantime hide root component but later hide component itself
-				RootComponent->SetHiddenInGame(ConvertedComponent.bComponentHiddenInGame);
-				RootComponent->SetVisibility(ConvertedComponent.bComponentVisible);
-			}
+			PrimitiveComponent->SetHiddenInGame(ConvertedComponent.bComponentHiddenInGame);
+			PrimitiveComponent->SetVisibility(ConvertedComponent.bComponentVisible);
 		}
 	}
 }
@@ -127,14 +130,14 @@ bool UAvaDynamicMeshConverterModifier::IsModifierDirtyable() const
 {
 	const double CurrentTime = FPlatformTime::Seconds();
 
-	if (TransformUpdateInterval > 0
-		&& CurrentTime - LastTransformUpdateTime > TransformUpdateInterval)
+	if (UpdateInterval > 0
+		&& CurrentTime - LastTransformUpdateTime > UpdateInterval)
 	{
 		const_cast<UAvaDynamicMeshConverterModifier*>(this)->LastTransformUpdateTime = CurrentTime;
 
 		for (const FAvaDynamicMeshConverterModifierComponentState& ConvertedComponent : ConvertedComponents)
 		{
-			const UPrimitiveComponent* PrimitiveComponent = ConvertedComponent.Component.Get();
+			UPrimitiveComponent* PrimitiveComponent = ConvertedComponent.Component.Get();
 			if (!PrimitiveComponent)
 			{
 				continue;
@@ -146,18 +149,29 @@ bool UAvaDynamicMeshConverterModifier::IsModifierDirtyable() const
 				continue;
 			}
 
-			const AActor* ParentActor = ChildActor->GetAttachParentActor();
-			FTransform ExpectedTransform = FTransform::Identity;
-
-			if (ParentActor)
+			const UDynamicMeshComponent* DynamicMeshComponent = GetMeshComponent();
+			if (!DynamicMeshComponent)
 			{
-				ExpectedTransform = ChildActor->GetActorTransform().GetRelativeTransform(ParentActor->GetActorTransform());
+				continue;
 			}
 
+			// Check Transform
+			const FTransform ExpectedTransform = PrimitiveComponent->GetComponentTransform().GetRelativeTransform(DynamicMeshComponent->GetComponentTransform());
 			if (!ConvertedComponent.ActorRelativeTransform.Equals(ExpectedTransform, 0.01))
 			{
 				return true;
 			}
+
+			// Check Materials
+			FAvaDynamicMeshConverterModifierComponentState NewState(PrimitiveComponent);
+			for (int32 Index = 0; Index < ConvertedComponent.ComponentMaterialsWeak.Num(); Index++)
+			{
+				if (!NewState.ComponentMaterialsWeak.IsValidIndex(Index)
+					|| ConvertedComponent.ComponentMaterialsWeak[Index] != NewState.ComponentMaterialsWeak[Index])
+				{
+					return true;
+				}
+			}
 		}
 	}
 
@@ -174,8 +188,6 @@ void UAvaDynamicMeshConverterModifier::OnSceneTreeTrackedActorChildrenChanged(in
 
 void UAvaDynamicMeshConverterModifier::Apply()
 {
-	const AActor* ActorModified = GetModifiedActor();
-
 	if (!IsMeshValid())
 	{
 		Fail(LOCTEXT("InvalidDynamicMeshComponent", "Invalid dynamic mesh component on modified actor"));
@@ -203,27 +215,15 @@ void UAvaDynamicMeshConverterModifier::Apply()
 		DynMeshComponent->SetMaterial(MatIndex, Material);
 	}
 
-	for (const FAvaDynamicMeshConverterModifierComponentState& OutConvert : ConvertedComponents)
+	// Hide converted component
+	if (bHideConvertedMesh)
 	{
-		if (OutConvert.Component.IsValid())
+		for (const FAvaDynamicMeshConverterModifierComponentState& ConvertedComponent : ConvertedComponents)
 		{
-			AActor* ComponentActor = OutConvert.Component->GetOwner();
-
-			// Hide converted component
-			if (bHideConvertedMesh)
+			if (UPrimitiveComponent* PrimitiveComponent = ConvertedComponent.Component.Get())
 			{
-				if (ComponentActor != ActorModified)
-				{
-					ComponentActor->SetHidden(true);
-#if WITH_EDITOR
-					ComponentActor->SetIsTemporarilyHiddenInEditor(true);
-#endif
-				}
-				else if (USceneComponent* RootComponent = ComponentActor->GetRootComponent())
-				{
-					RootComponent->SetHiddenInGame(true);
-					RootComponent->SetVisibility(false);
-				}
+				PrimitiveComponent->SetVisibility(false);
+				PrimitiveComponent->SetHiddenInGame(true);
 			}
 		}
 	}
@@ -375,32 +375,69 @@ TSet<EAvaDynamicMeshConverterModifierType> UAvaDynamicMeshConverterModifier::Get
 
 void UAvaDynamicMeshConverterModifier::SetComponentType(int32 InComponentType)
 {
-	if (InComponentType == ComponentType)
+	if (ComponentType == InComponentType)
 	{
 		return;
 	}
 
 	ComponentType = InComponentType;
+	MarkModifierDirty();
 }
 
 void UAvaDynamicMeshConverterModifier::SetFilterActorMode(EAvaDynamicMeshConverterModifierFilter InFilter)
 {
+	if (FilterActorMode == InFilter)
+	{
+		return;
+	}
+
 	FilterActorMode = InFilter;
+	MarkModifierDirty();
 }
 
 void UAvaDynamicMeshConverterModifier::SetFilterActorClasses(const TSet<TSubclassOf<AActor>>& InClasses)
 {
+	if (FilterActorClasses.Includes(InClasses) && FilterActorClasses.Num() == InClasses.Num())
+	{
+		return;
+	}
+
 	FilterActorClasses = InClasses;
+	MarkModifierDirty();
 }
 
 void UAvaDynamicMeshConverterModifier::SetIncludeAttachedActors(bool bInInclude)
 {
+	if (bIncludeAttachedActors == bInInclude)
+	{
+		return;
+	}
+
 	bIncludeAttachedActors = bInInclude;
+	MarkModifierDirty();
 }
 
 void UAvaDynamicMeshConverterModifier::SetHideConvertedMesh(bool bInHide)
 {
+	if (bHideConvertedMesh == bInHide)
+	{
+		return;
+	}
+
 	bHideConvertedMesh = bInHide;
+	MarkModifierDirty();
+}
+
+void UAvaDynamicMeshConverterModifier::SetUpdateInterval(float InInterval)
+{
+	InInterval = FMath::Max(0.f, InInterval);
+	if (FMath::IsNearlyEqual(UpdateInterval, InInterval))
+	{
+		return;
+	}
+
+	UpdateInterval = InInterval;
+	MarkModifierDirty();
 }
 
 void UAvaDynamicMeshConverterModifier::OnRenderStateUpdated(AActor* InActor, UActorComponent* InComponent)
@@ -464,7 +501,9 @@ bool UAvaDynamicMeshConverterModifier::ConvertComponents(TArray<TWeakObjectPtr<U
 		return false;
 	}
 
-	ConvertedComponents.Empty();
+	TSet<FAvaDynamicMeshConverterModifierComponentState> NewConvertedComponents;
+	NewConvertedComponents.Reserve(ConvertedComponents.Num());
+
 	MeshBuilder.Reset();
 
 	UDynamicMeshComponent* DynamicMeshComponent = GetMeshComponent();
@@ -474,81 +513,62 @@ bool UAvaDynamicMeshConverterModifier::ConvertComponents(TArray<TWeakObjectPtr<U
 	TArray<AActor*> FilteredActors;
 	GetFilteredActors(FilteredActors);
 
-	if (HasFlag(EAvaDynamicMeshConverterModifierType::StaticMeshComponent))
+	auto FindOrAddComponentState = [this, &SourceTransform, &NewConvertedComponents](UPrimitiveComponent* InComponent)
 	{
-		TArray<UStaticMeshComponent*> Components;
-		GetStaticMeshComponents(FilteredActors, Components);
-
-		for (UStaticMeshComponent* Component : Components)
+		FAvaDynamicMeshConverterModifierComponentState NewState(InComponent);
+		if (FAvaDynamicMeshConverterModifierComponentState* OldState = ConvertedComponents.Find(NewState))
 		{
-			if (MeshBuilder.AppendComponent(Component, SourceTransform))
-			{
-				FAvaDynamicMeshConverterModifierComponentState State(Component);
-				ConvertedComponents.Emplace(State);
-			}
+			OldState->UpdateRelativeTransform(SourceTransform);
+			NewConvertedComponents.Emplace(*OldState);
 		}
+		else
+		{
+			NewState.UpdateRelativeTransform(SourceTransform);
+			NewConvertedComponents.Emplace(NewState);
+		}
+	};
+
+	ECEMeshBuilderComponentType ComponentTypes = ECEMeshBuilderComponentType::None;
+
+	if (HasFlag(EAvaDynamicMeshConverterModifierType::StaticMeshComponent))
+	{
+		EnumAddFlags(ComponentTypes, ECEMeshBuilderComponentType::StaticMeshComponent);
 	}
 
 	if (HasFlag(EAvaDynamicMeshConverterModifierType::DynamicMeshComponent))
 	{
-		TArray<UDynamicMeshComponent*> Components;
-		GetDynamicMeshComponents(FilteredActors, Components);
-
-		for (UDynamicMeshComponent* Component : Components)
-		{
-			if (MeshBuilder.AppendComponent(Component, SourceTransform))
-			{
-				FAvaDynamicMeshConverterModifierComponentState State(Component);
-				ConvertedComponents.Emplace(State);
-			}
-		}
+		EnumAddFlags(ComponentTypes, ECEMeshBuilderComponentType::DynamicMeshComponent);
 	}
 
 	if (HasFlag(EAvaDynamicMeshConverterModifierType::SkeletalMeshComponent))
 	{
-		TArray<USkeletalMeshComponent*> Components;
-		GetSkeletalMeshComponents(FilteredActors, Components);
-
-		for (USkeletalMeshComponent* Component : Components)
-		{
-			if (MeshBuilder.AppendComponent(Component, SourceTransform))
-			{
-				FAvaDynamicMeshConverterModifierComponentState State(Component);
-				ConvertedComponents.Emplace(State);
-			}
-		}
+		EnumAddFlags(ComponentTypes, ECEMeshBuilderComponentType::SkeletalMeshComponent);
 	}
 
 	if (HasFlag(EAvaDynamicMeshConverterModifierType::BrushComponent))
 	{
-		TArray<UBrushComponent*> Components;
-		GetBrushComponents(FilteredActors, Components);
-
-		for (UBrushComponent* Component : Components)
-		{
-			if (MeshBuilder.AppendComponent(Component, SourceTransform))
-			{
-				FAvaDynamicMeshConverterModifierComponentState State(Component);
-				ConvertedComponents.Emplace(State);
-			}
-		}
+		EnumAddFlags(ComponentTypes, ECEMeshBuilderComponentType::BrushComponent);
 	}
 
 	if (HasFlag(EAvaDynamicMeshConverterModifierType::ProceduralMeshComponent))
 	{
-		TArray<UProceduralMeshComponent*> Components;
-		GetProceduralMeshComponents(FilteredActors, Components);
+		EnumAddFlags(ComponentTypes, ECEMeshBuilderComponentType::ProceduralMeshComponent);
+	}
 
-		for (UProceduralMeshComponent* Component : Components)
+	FCEMeshBuilder::FCEMeshBuilderAppendParams AppendParams;
+	AppendParams.ComponentTypes = ComponentTypes;
+	AppendParams.ExcludeComponents.Add(GetMeshComponent());
+
+	for (const AActor* FilteredActor : FilteredActors)
+	{
+		for (UPrimitiveComponent* PrimitiveComponent : MeshBuilder.AppendActor(FilteredActor, SourceTransform, AppendParams))
 		{
-			if (MeshBuilder.AppendComponent(Component, SourceTransform))
-			{
-				FAvaDynamicMeshConverterModifierComponentState State(Component);
-				ConvertedComponents.Emplace(State);
-			}
+			FindOrAddComponentState(PrimitiveComponent);
 		}
 	}
 
+	ConvertedComponents = MoveTemp(NewConvertedComponents);
+
 	return MeshBuilder.BuildDynamicMesh(DynamicMeshComponent->GetDynamicMesh(), OutMaterialsWeak);
 }
 
@@ -687,99 +707,4 @@ void UAvaDynamicMeshConverterModifier::GetFilteredActors(TArray<AActor*>& OutAct
 	}
 }
 
-void UAvaDynamicMeshConverterModifier::GetStaticMeshComponents(const TArray<AActor*>& InActors, TArray<UStaticMeshComponent*>& OutComponents) const
-{
-	for (const AActor* Actor : InActors)
-	{
-		TArray<UStaticMeshComponent*> OutMeshComponents;
-		Actor->GetComponents(OutMeshComponents, false);
-		OutComponents.Append(OutMeshComponents);
-	}
-	// remove all invalid components
-	OutComponents.RemoveAll([](const UStaticMeshComponent* InComponent)->bool
-	{
-#if WITH_EDITOR
-		return !IsValid(InComponent) || InComponent->IsVisualizationComponent();
-#else
-		return !IsValid(InComponent);
-#endif
-	});
-}
-
-void UAvaDynamicMeshConverterModifier::GetDynamicMeshComponents(const TArray<AActor*>& InActors, TArray<UDynamicMeshComponent*>& OutComponents) const
-{
-	for (const AActor* Actor : InActors)
-	{
-		TArray<UDynamicMeshComponent*> OutMeshComponents;
-		Actor->GetComponents(OutMeshComponents, false);
-		OutComponents.Append(OutMeshComponents);
-	}
-	// remove all invalid components & and the modifier created component
-	OutComponents.RemoveAll([this](const UDynamicMeshComponent* InComponent)->bool
-	{
-#if WITH_EDITOR
-		return !IsValid(InComponent) || InComponent == GetMeshComponent() || InComponent->IsVisualizationComponent();
-#else
-		return !IsValid(InComponent) || InComponent == GetMeshComponent();
-#endif
-	});
-}
-
-void UAvaDynamicMeshConverterModifier::GetSkeletalMeshComponents(const TArray<AActor*>& InActors, TArray<USkeletalMeshComponent*>& OutComponents) const
-{
-	for (const AActor* Actor : InActors)
-	{
-		TArray<USkeletalMeshComponent*> OutMeshComponents;
-		Actor->GetComponents(OutMeshComponents, false);
-		OutComponents.Append(OutMeshComponents);
-	}
-	// remove all invalid components
-	OutComponents.RemoveAll([this](const USkeletalMeshComponent* InComponent)->bool
-	{
-#if WITH_EDITOR
-		return !IsValid(InComponent) || InComponent->IsVisualizationComponent();
-#else
-		return !IsValid(InComponent);
-#endif
-	});
-}
-
-void UAvaDynamicMeshConverterModifier::GetBrushComponents(const TArray<AActor*>& InActors, TArray<UBrushComponent*>& OutComponents) const
-{
-	for (const AActor* Actor : InActors)
-	{
-		TArray<UBrushComponent*> OutMeshComponents;
-		Actor->GetComponents(OutMeshComponents, false);
-		OutComponents.Append(OutMeshComponents);
-	}
-	// remove all invalid components
-	OutComponents.RemoveAll([this](const UBrushComponent* InComponent)->bool
-	{
-#if WITH_EDITOR
-		return !IsValid(InComponent) || InComponent->IsVisualizationComponent();
-#else
-		return !IsValid(InComponent);
-#endif
-	});
-}
-
-void UAvaDynamicMeshConverterModifier::GetProceduralMeshComponents(const TArray<AActor*>& InActors, TArray<UProceduralMeshComponent*>& OutComponents) const
-{
-	for (const AActor* Actor : InActors)
-	{
-		TArray<UProceduralMeshComponent*> OutMeshComponents;
-		Actor->GetComponents(OutMeshComponents, false);
-		OutComponents.Append(OutMeshComponents);
-	}
-	// remove all invalid components
-	OutComponents.RemoveAll([this](const UProceduralMeshComponent* InComponent)->bool
-	{
-#if WITH_EDITOR
-		return !IsValid(InComponent) || InComponent->IsVisualizationComponent();
-#else
-		return !IsValid(InComponent);
-#endif
-	});
-}
-
 #undef LOCTEXT_NAMESPACE
diff --git a/Engine/Plugins/Experimental/Avalanche/Source/AvalancheModifiers/Public/Modifiers/AvaDynamicMeshConverterModifier.h b/Engine/Plugins/Experimental/Avalanche/Source/AvalancheModifiers/Public/Modifiers/AvaDynamicMeshConverterModifier.h
--- a/Engine/Plugins/Experimental/Avalanche/Source/AvalancheModifiers/Public/Modifiers/AvaDynamicMeshConverterModifier.h
+++ b/Engine/Plugins/Experimental/Avalanche/Source/AvalancheModifiers/Public/Modifiers/AvaDynamicMeshConverterModifier.h
@@ -11,6 +11,7 @@
 
 class UBrushComponent;
 class UDynamicMeshComponent;
+class UMaterialInterface;
 class UPrimitiveComponent;
 class UProceduralMeshComponent;
 class USkeletalMeshComponent;
@@ -45,6 +46,8 @@ struct FAvaDynamicMeshConverterModifierComponentState
 	explicit FAvaDynamicMeshConverterModifierComponentState() {}
 	explicit FAvaDynamicMeshConverterModifierComponentState(UPrimitiveComponent* InPrimitiveComponent);
 
+	void UpdateRelativeTransform(const FTransform& InParentTransform);
+	
 	friend uint32 GetTypeHash(const FAvaDynamicMeshConverterModifierComponentState& InItem)
 	{
 		return GetTypeHash(InItem.Component);
@@ -59,14 +62,6 @@ struct FAvaDynamicMeshConverterModifierComponentState
 	UPROPERTY()
 	TWeakObjectPtr<UPrimitiveComponent> Component;
 
-	/** The default visibility of the actor converted component in game */
-	UPROPERTY()
-	bool bActorHiddenInGame = true;
-
-	/** The default visibility of the actor converted component in editor */
-	UPROPERTY()
-	bool bActorHiddenInEditor = true;
-
 	/** The default visibility of the actor converted component in game */
 	UPROPERTY()
 	bool bComponentHiddenInGame = false;
@@ -78,6 +73,10 @@ struct FAvaDynamicMeshConverterModifierComponentState
 	/** Transform saved before mesh is converted */
 	UPROPERTY()
 	FTransform ActorRelativeTransform = FTransform::Identity;
+
+	/** Used for diffs */
+	UPROPERTY(Transient)
+	TArray<TWeakObjectPtr<UMaterialInterface>> ComponentMaterialsWeak;
 };
 
 UCLASS(MinimalAPI, BlueprintType)
@@ -155,6 +154,15 @@ public:
 	{
 		return bHideConvertedMesh;
 	}
+	
+	UFUNCTION(BlueprintCallable, Category="Motion Design|Modifiers|DynamicMeshConverter")
+	AVALANCHEMODIFIERS_API void SetUpdateInterval(float InInterval);
+
+	UFUNCTION(BlueprintPure, Category="Motion Design|Modifiers|DynamicMeshConverter")
+	float GetUpdateInterval() const
+	{
+		return UpdateInterval;
+	}
 
 protected:
 	//~ Begin UObject
@@ -170,6 +178,7 @@ protected:
 	virtual void OnModifierCDOSetup(FActorModifierCoreMetadata& InMetadata) override;
 	virtual void OnModifierAdded(EActorModifierCoreEnableReason InReason) override;
 	virtual void OnModifierEnabled(EActorModifierCoreEnableReason InReason) override;
+	virtual void OnModifierDisabled(EActorModifierCoreDisableReason InReason) override;
 	virtual void RestorePreState() override;
 	virtual void Apply() override;
 	virtual void OnModifierRemoved(EActorModifierCoreDisableReason InReason) override;
@@ -199,12 +208,6 @@ protected:
 
 	void GetFilteredActors(TArray<AActor*>& OutActors) const;
 
-	void GetStaticMeshComponents(const TArray<AActor*>& InActors, TArray<UStaticMeshComponent*>& OutComponents) const;
-	void GetDynamicMeshComponents(const TArray<AActor*>& InActors, TArray<UDynamicMeshComponent*>& OutComponents) const;
-	void GetSkeletalMeshComponents(const TArray<AActor*>& InActors, TArray<USkeletalMeshComponent*>& OutComponents) const;
-	void GetBrushComponents(const TArray<AActor*>& InActors, TArray<UBrushComponent*>& OutComponents) const;
-	void GetProceduralMeshComponents(const TArray<AActor*>& InActors, TArray<UProceduralMeshComponent*>& OutComponents) const;
-
 	/** What actor should we copy from, by default is self */
 	UPROPERTY(EditInstanceOnly, Setter="SetSourceActorWeak", Getter="GetSourceActorWeak", Category="DynamicMeshConverter", meta=(DisplayName="SourceActor", AllowPrivateAccess="true"))
 	TWeakObjectPtr<AActor> SourceActorWeak = GetModifiedActor();
@@ -234,14 +237,14 @@ protected:
 	UPROPERTY(EditInstanceOnly, Setter="SetHideConvertedMesh", Getter="GetHideConvertedMesh", Category="DynamicMeshConverter", meta=(AllowPrivateAccess="true"))
 	bool bHideConvertedMesh = true;
 
+	/** Update interval to compare if a transform/material has changed in converted components, when value <= 0 then skipped */
+	UPROPERTY(EditInstanceOnly, Setter, Getter, Category="DynamicMeshConverter", meta=(ClampMin="0.0"))
+	float UpdateInterval = 1.f;
+
 	/** Did we create the dynamic mesh component from this modifier or retrieved it */
 	UPROPERTY()
 	bool bComponentCreated = false;
 
-	/** Update interval to compare if a transform has changed in converted components, when value <= 0 then skipped */
-	UPROPERTY(EditInstanceOnly, Category="DynamicMeshConverter")
-	float TransformUpdateInterval = 0;
-
 	/** Components converted to dynamic mesh */
 	UPROPERTY()
 	TSet<FAvaDynamicMeshConverterModifierComponentState> ConvertedComponents;
diff --git a/Engine/Plugins/Experimental/ClonerEffector/Source/ClonerEffector/Private/Cloner/CEClonerComponent.cpp b/Engine/Plugins/Experimental/ClonerEffector/Source/ClonerEffector/Private/Cloner/CEClonerComponent.cpp
--- a/Engine/Plugins/Experimental/ClonerEffector/Source/ClonerEffector/Private/Cloner/CEClonerComponent.cpp
+++ b/Engine/Plugins/Experimental/ClonerEffector/Source/ClonerEffector/Private/Cloner/CEClonerComponent.cpp
@@ -827,7 +827,7 @@ void UCEClonerComponent::UpdateActorBakedDynamicMesh(AActor* InActor)
 	UDynamicMesh* Mesh = NewObject<UDynamicMesh>();
 	TArray<TWeakObjectPtr<UMaterialInterface>> MeshMaterials;
 
-	MeshBuilder.AppendActor(InActor);
+	MeshBuilder.AppendActor(InActor, InActor->GetActorTransform());
 	MeshBuilder.BuildDynamicMesh(Mesh, MeshMaterials);
 	MeshBuilder.Reset();
 
diff --git a/Engine/Plugins/Experimental/ClonerEffector/Source/ClonerEffectorMeshBuilder/Private/CEMeshBuilder.cpp b/Engine/Plugins/Experimental/ClonerEffector/Source/ClonerEffectorMeshBuilder/Private/CEMeshBuilder.cpp
--- a/Engine/Plugins/Experimental/ClonerEffector/Source/ClonerEffectorMeshBuilder/Private/CEMeshBuilder.cpp
+++ b/Engine/Plugins/Experimental/ClonerEffector/Source/ClonerEffectorMeshBuilder/Private/CEMeshBuilder.cpp
@@ -28,7 +28,8 @@
 #include "StaticMeshOperations.h"
 #include "UDynamicMesh.h"
 
-const FCEMeshBuilder::FCEMeshBuilderParams FCEMeshBuilder::DefaultParams;
+const FCEMeshBuilder::FCEMeshBuilderParams FCEMeshBuilder::DefaultBuildParams;
+const FCEMeshBuilder::FCEMeshBuilderAppendParams FCEMeshBuilder::DefaultAppendParams;
 
 bool FCEMeshBuilder::HasAnyGeometry(UActorComponent* InComponent)
 {
@@ -161,86 +162,131 @@ void FCEMeshBuilder::Reset()
 	MeshInstances.Empty();
 }
 
-int32 FCEMeshBuilder::AppendActor(const AActor* InActor)
+TArray<UPrimitiveComponent*> FCEMeshBuilder::AppendActor(const AActor* InActor, const FTransform& InSourceTransform, const FCEMeshBuilderAppendParams& InParams)
 {
-	int32 ComponentConverted = 0;
+	TArray<UPrimitiveComponent*> AppendPrimitiveComponents;
 
-	if (!IsValid(InActor))
+	if (!IsValid(InActor) || InParams.ComponentTypes == ECEMeshBuilderComponentType::None)
 	{
-		return ComponentConverted;
+		return AppendPrimitiveComponents;
 	}
 
-	const FTransform& SourceTransform = InActor->GetActorTransform();
-
 	TArray<UPrimitiveComponent*> PrimitiveComponents;
 	InActor->GetComponents(PrimitiveComponents, /** IncludeChildrenActors */ false);
 
+	AppendPrimitiveComponents.Reserve(PrimitiveComponents.Num());
+
 	for (UPrimitiveComponent* PrimitiveComponent : PrimitiveComponents)
 	{
-		if (!IsValid(PrimitiveComponent))
+		if (!IsValid(PrimitiveComponent)
+#if WITH_EDITOR
+			|| PrimitiveComponent->IsVisualizationComponent()
+#endif
+			|| InParams.ExcludeComponents.Contains(PrimitiveComponent)
+			)
 		{
 			continue;
 		}
 
 		if (UDynamicMeshComponent* DynamicMeshComponent = Cast<UDynamicMeshComponent>(PrimitiveComponent))
 		{
-			if (AppendComponent(DynamicMeshComponent, SourceTransform))
+			if (!EnumHasAnyFlags(InParams.ComponentTypes, ECEMeshBuilderComponentType::DynamicMeshComponent))
 			{
-				ComponentConverted++;
+				continue;
+			}
+			
+			if (AppendComponent(DynamicMeshComponent, InSourceTransform))
+			{
+				AppendPrimitiveComponents.Add(PrimitiveComponent);
 			}
 		}
 		else if (USkeletalMeshComponent* SkeletalMeshComponent = Cast<USkeletalMeshComponent>(PrimitiveComponent))
 		{
-			if (AppendComponent(SkeletalMeshComponent, SourceTransform))
+			if (!EnumHasAnyFlags(InParams.ComponentTypes, ECEMeshBuilderComponentType::SkeletalMeshComponent))
+			{
+				continue;
+			}
+			
+			if (AppendComponent(SkeletalMeshComponent, InSourceTransform))
 			{
-				ComponentConverted++;
+				AppendPrimitiveComponents.Add(PrimitiveComponent);
 			}
 		}
 		else if (UBrushComponent* BrushComponent = Cast<UBrushComponent>(PrimitiveComponent))
 		{
-			if (AppendComponent(BrushComponent, SourceTransform))
+			if (!EnumHasAnyFlags(InParams.ComponentTypes, ECEMeshBuilderComponentType::BrushComponent))
 			{
-				ComponentConverted++;
+				continue;
+			}
+			
+			if (AppendComponent(BrushComponent, InSourceTransform))
+			{
+				AppendPrimitiveComponents.Add(PrimitiveComponent);
 			}
 		}
 		else if (UProceduralMeshComponent* ProceduralMeshComponent = Cast<UProceduralMeshComponent>(PrimitiveComponent))
 		{
-			if (AppendComponent(ProceduralMeshComponent, SourceTransform))
+			if (!EnumHasAnyFlags(InParams.ComponentTypes, ECEMeshBuilderComponentType::ProceduralMeshComponent))
 			{
-				ComponentConverted++;
+				continue;
+			}
+			
+			if (AppendComponent(ProceduralMeshComponent, InSourceTransform))
+			{
+				AppendPrimitiveComponents.Add(PrimitiveComponent);
 			}
 		}
 		else if (UInstancedStaticMeshComponent* ISMComponent = Cast<UInstancedStaticMeshComponent>(PrimitiveComponent))
 		{
-			if (AppendComponent(ISMComponent, SourceTransform))
+			if (!EnumHasAnyFlags(InParams.ComponentTypes, ECEMeshBuilderComponentType::InstancedStaticMeshComponent))
+			{
+				continue;
+			}
+			
+			if (AppendComponent(ISMComponent, InSourceTransform))
 			{
-				ComponentConverted++;
+				AppendPrimitiveComponents.Add(PrimitiveComponent);
 			}
 		}
 		else if (USplineMeshComponent* SplineMeshComponent = Cast<USplineMeshComponent>(PrimitiveComponent))
 		{
-			if (AppendComponent(SplineMeshComponent, SourceTransform))
+			if (!EnumHasAnyFlags(InParams.ComponentTypes, ECEMeshBuilderComponentType::SplineMeshComponent))
+			{
+				continue;
+			}
+			
+			if (AppendComponent(SplineMeshComponent, InSourceTransform))
 			{
-				ComponentConverted++;
+				AppendPrimitiveComponents.Add(PrimitiveComponent);
 			}
 		}
 		else if (UStaticMeshComponent* StaticMeshComponent = Cast<UStaticMeshComponent>(PrimitiveComponent))
 		{
-			if (AppendComponent(StaticMeshComponent, SourceTransform))
+			if (!EnumHasAnyFlags(InParams.ComponentTypes, ECEMeshBuilderComponentType::StaticMeshComponent))
 			{
-				ComponentConverted++;
+				continue;
+			}
+			
+			if (AppendComponent(StaticMeshComponent, InSourceTransform))
+			{
+				AppendPrimitiveComponents.Add(PrimitiveComponent);
 			}
 		}
 		else if (UNiagaraComponent* NiagaraComponent = Cast<UNiagaraComponent>(PrimitiveComponent))
 		{
-			if (AppendComponent(NiagaraComponent, SourceTransform))
+			if (!EnumHasAnyFlags(InParams.ComponentTypes, ECEMeshBuilderComponentType::NiagaraComponent))
+			{
+				continue;
+			}
+			
+			if (AppendComponent(NiagaraComponent, InSourceTransform))
 			{
-				ComponentConverted++;
+				AppendPrimitiveComponents.Add(PrimitiveComponent);
 			}
 		}
 	}
 
-	return ComponentConverted;
+	return AppendPrimitiveComponents;
 }
 
 bool FCEMeshBuilder::AppendMesh(const UDynamicMesh* InMesh, const TArray<TWeakObjectPtr<UMaterialInterface>>& InMaterials, const FTransform& InTransform)
diff --git a/Engine/Plugins/Experimental/ClonerEffector/Source/ClonerEffectorMeshBuilder/Public/CEMeshBuilder.h b/Engine/Plugins/Experimental/ClonerEffector/Source/ClonerEffectorMeshBuilder/Public/CEMeshBuilder.h
--- a/Engine/Plugins/Experimental/ClonerEffector/Source/ClonerEffectorMeshBuilder/Public/CEMeshBuilder.h
+++ b/Engine/Plugins/Experimental/ClonerEffector/Source/ClonerEffectorMeshBuilder/Public/CEMeshBuilder.h
@@ -20,6 +20,21 @@ class USplineMeshComponent;
 class UStaticMesh;
 class UStaticMeshComponent;
 
+/** Enumerates all component type that can be converted */
+enum class ECEMeshBuilderComponentType : uint16
+{
+	None = 0 UMETA(Hidden),
+	DynamicMeshComponent = 1 << 1,
+	SkeletalMeshComponent = 1 << 2,
+	BrushComponent = 1 << 3,
+	ProceduralMeshComponent = 1 << 4,
+	InstancedStaticMeshComponent = 1 << 5,
+	SplineMeshComponent = 1 << 6,
+	StaticMeshComponent = 1 << 7,
+	NiagaraComponent = 1 << 8,
+	All = DynamicMeshComponent | SkeletalMeshComponent | BrushComponent | ProceduralMeshComponent | InstancedStaticMeshComponent | SplineMeshComponent | StaticMeshComponent | NiagaraComponent
+};
+
 /** Struct used to build a mesh based out of other meshes class */
 USTRUCT()
 struct CLONEREFFECTORMESHBUILDER_API FCEMeshBuilder
@@ -41,7 +56,17 @@ struct CLONEREFFECTORMESHBUILDER_API FCEMeshBuilder
 		bool bMergeMaterials = false;
 	};
 
-	static const FCEMeshBuilderParams DefaultParams;
+	struct FCEMeshBuilderAppendParams
+	{
+		/** Component types to append */
+		ECEMeshBuilderComponentType ComponentTypes = ECEMeshBuilderComponentType::All;
+
+		/** Specific excluded components */
+		TSet<UPrimitiveComponent*> ExcludeComponents;
+	};
+
+	static const FCEMeshBuilderParams DefaultBuildParams;
+	static const FCEMeshBuilderAppendParams DefaultAppendParams;
 
 	/** Checks if the component contains any geometry data */
 	static bool HasAnyGeometry(UActorComponent* InComponent);
@@ -70,7 +95,7 @@ struct CLONEREFFECTORMESHBUILDER_API FCEMeshBuilder
 	void Reset();
 
 	/** Appends supported components within actor */
-	int32 AppendActor(const AActor* InActor);
+	TArray<UPrimitiveComponent*> AppendActor(const AActor* InActor, const FTransform& InTransform, const FCEMeshBuilderAppendParams& InParams = DefaultAppendParams);
 
 	/** Appends Static Mesh component */
 	bool AppendComponent(const UStaticMeshComponent* InComponent, const FTransform& InTransform = FTransform::Identity);
@@ -103,10 +128,10 @@ struct CLONEREFFECTORMESHBUILDER_API FCEMeshBuilder
 	bool AppendMesh(UStaticMesh* InMesh, const TArray<TWeakObjectPtr<UMaterialInterface>>& InMaterials, const FTransform& InTransform = FTransform::Identity);
 
 	/** Builds a dynamic mesh by merging all the mesh data imported */
-	bool BuildDynamicMesh(UDynamicMesh* OutMesh, TArray<TWeakObjectPtr<UMaterialInterface>>& OutMaterials, const FCEMeshBuilderParams& InParams = DefaultParams);
+	bool BuildDynamicMesh(UDynamicMesh* OutMesh, TArray<TWeakObjectPtr<UMaterialInterface>>& OutMaterials, const FCEMeshBuilderParams& InParams = DefaultBuildParams);
 
 	/** Builds a static mesh by merging all the mesh data imported */
-	bool BuildStaticMesh(UStaticMesh* OutMesh, TArray<TWeakObjectPtr<UMaterialInterface>>& OutMaterials, const FCEMeshBuilderParams& InParams = DefaultParams);
+	bool BuildStaticMesh(UStaticMesh* OutMesh, TArray<TWeakObjectPtr<UMaterialInterface>>& OutMaterials, const FCEMeshBuilderParams& InParams = DefaultBuildParams);
 
 	/** Builds a static mesh for the specific instance index */
 	bool BuildStaticMesh(int32 InInstanceIndex, UStaticMesh* OutMesh, FCEMeshInstanceData& OutMeshInstance);
