[MotionDesign] Stagger Tool:
- Fix regression of not being able to select Sequencer outliner folder tracks for stagger
- Close tool window when Sequencer is closed
- Add FAvaSequencerAction::OnSequencerClosed() and hooked up OnSequencerCreated()

diff --git a/Engine/Plugins/Experimental/Avalanche/Source/AvalancheSequencer/Private/AvaSequencer.cpp b/Engine/Plugins/Experimental/Avalanche/Source/AvalancheSequencer/Private/AvaSequencer.cpp
--- a/Engine/Plugins/Experimental/Avalanche/Source/AvalancheSequencer/Private/AvaSequencer.cpp
+++ b/Engine/Plugins/Experimental/Avalanche/Source/AvalancheSequencer/Private/AvaSequencer.cpp
@@ -307,10 +307,16 @@ void FAvaSequencer::EnsureSequencer()
 	Sequencer->GetSelectionChangedObjectGuids().AddSP(this, &FAvaSequencer::OnSequencerSelectionChanged);
 	Sequencer->OnGetIsBindingVisible().BindSP(this, &FAvaSequencer::IsBindingSelected);
 	Sequencer->OnCameraCut().AddSP(this, &FAvaSequencer::OnUpdateCameraCut);
+	Sequencer->OnCloseEvent().AddSP(this, &FAvaSequencer::OnSequencerClosed);
 
 	if (IAvaSequenceProvider* const SequenceProvider = Provider.GetSequenceProvider())
 	{
 		SequenceProvider->OnEditorSequencerCreated(Sequencer);
+
+		for (const TSharedRef<FAvaSequencerAction>& SequencerAction : SequencerActions)
+		{
+			SequencerAction->OnSequencerCreated();
+		}
 	}
 
 	// Create ease curve tool and map commands
@@ -1870,4 +1876,17 @@ void FAvaSequencer::ExtendSidebarMarkedFramesMenu(FMenuBuilder& OutMenuBuilder)
 	}
 }
 
+void FAvaSequencer::OnSequencerClosed(const TSharedRef<ISequencer> InSequencer)
+{
+	if (InSequencer != GetSequencer())
+	{
+		return;
+	}
+
+	for (const TSharedRef<FAvaSequencerAction>& SequencerAction : SequencerActions)
+	{
+		SequencerAction->OnSequencerClosed();
+	}
+}
+
 #undef LOCTEXT_NAMESPACE
diff --git a/Engine/Plugins/Experimental/Avalanche/Source/AvalancheSequencer/Private/AvaSequencer.h b/Engine/Plugins/Experimental/Avalanche/Source/AvalancheSequencer/Private/AvaSequencer.h
--- a/Engine/Plugins/Experimental/Avalanche/Source/AvalancheSequencer/Private/AvaSequencer.h
+++ b/Engine/Plugins/Experimental/Avalanche/Source/AvalancheSequencer/Private/AvaSequencer.h
@@ -16,6 +16,7 @@ class FAvaEditorSelection;
 class FAvaSequencerAction;
 class FAvaSequencerCleanView;
 class FAvaSequencerController;
+class FAvaStaggerTool;
 class FExtender;
 class FMenuBuilder;
 class FProperty;
@@ -232,6 +233,8 @@ private:
 	void ExtendSidebarSelectionMenu(FMenuBuilder& OutMenuBuilder);
 	void ExtendSidebarMarkedFramesMenu(FMenuBuilder& OutMenuBuilder);
 
+	void OnSequencerClosed(const TSharedRef<ISequencer> InSequencer);
+
 	FOnViewedSequenceChanged OnViewedSequenceChanged;
 
 	IAvaSequencerProvider& Provider;
diff --git a/Engine/Plugins/Experimental/Avalanche/Source/AvalancheSequencer/Private/Commands/AvaSequencerAction.h b/Engine/Plugins/Experimental/Avalanche/Source/AvalancheSequencer/Private/Commands/AvaSequencerAction.h
--- a/Engine/Plugins/Experimental/Avalanche/Source/AvalancheSequencer/Private/Commands/AvaSequencerAction.h
+++ b/Engine/Plugins/Experimental/Avalanche/Source/AvalancheSequencer/Private/Commands/AvaSequencerAction.h
@@ -26,6 +26,8 @@ public:
 
 	virtual void OnSequencerCreated() {}
 
+	virtual void OnSequencerClosed() {}
+
 protected:
 	/** Utility function to directly retrieve the Sequencer Selection from the Owner Ava Sequencer */
 	TSharedPtr<UE::Sequencer::FSequencerSelection> GetSequencerSelection() const;
diff --git a/Engine/Plugins/Experimental/Avalanche/Source/AvalancheSequencer/Private/StaggerTool/AvaStaggerTool.cpp b/Engine/Plugins/Experimental/Avalanche/Source/AvalancheSequencer/Private/StaggerTool/AvaStaggerTool.cpp
--- a/Engine/Plugins/Experimental/Avalanche/Source/AvalancheSequencer/Private/StaggerTool/AvaStaggerTool.cpp
+++ b/Engine/Plugins/Experimental/Avalanche/Source/AvalancheSequencer/Private/StaggerTool/AvaStaggerTool.cpp
@@ -46,6 +46,7 @@ FAvaStaggerTool::FAvaStaggerTool(FAvaSequencer& InOwner)
 
 FAvaStaggerTool::~FAvaStaggerTool()
 {
+	CloseToolWindow();
 }
 
 TObjectPtr<UAvaSequencerStaggerSettings> FAvaStaggerTool::GetSettings() const
@@ -60,6 +61,11 @@ void FAvaStaggerTool::MapAction(const TSharedRef<FUICommandList>& InCommandList)
 		, FCanExecuteAction::CreateSP(this, &FAvaStaggerTool::CanExecute));
 }
 
+void FAvaStaggerTool::OnSequencerClosed()
+{
+	CloseToolWindow();
+}
+
 bool FAvaStaggerTool::CanExecute() const
 {
 	return true;
@@ -115,7 +121,7 @@ void FAvaStaggerTool::BindToSequencerSelectionChange()
 	{
 		if (const TSharedPtr<FSequencerSelection> Selection = ViewModel->GetSelection())
 		{
-			Selection->OnChanged.AddRaw(this, &FAvaStaggerTool::OnSequencerSelectionChanged);
+			Selection->OnChanged.AddSP(this, &FAvaStaggerTool::OnSequencerSelectionChanged);
 		}
 	}
 }
@@ -166,6 +172,16 @@ void FAvaStaggerTool::OnSequencerSelectionChanged()
 	}
 }
 
+void FAvaStaggerTool::CloseToolWindow()
+{
+	if (ToolWindow.IsValid())
+	{
+		ToolWindow->RequestDestroyWindow();
+		ToolWindow.Reset();
+		ToolWindow = nullptr;
+	}
+}
+
 bool FAvaStaggerTool::HasValidSelection() const
 {
 	return IsLayerBarSelection() || IsKeyFrameSelection();
@@ -426,16 +442,30 @@ TArray<TViewModelPtr<FLayerBarModel>> FAvaStaggerTool::GatherAllSelectedLayerBar
 
 		for (const FViewModelPtr& ViewModel : SequencerSelection->Outliner)
 		{
-			if (const TViewModelPtr<ITrackAreaExtension> TrackArea = ViewModel.ImplicitCast())
+			const TViewModelPtr<ITrackAreaExtension> Track = ViewModel.ImplicitCast();
+			if (!Track.IsValid())
 			{
-				for (const FViewModelPtr& TrackAreaModel : TrackArea->GetTopLevelChildTrackAreaModels())
+				continue;
+			}
+
+			for (const FViewModelPtr& TrackAreaModel : Track->GetTopLevelChildTrackAreaModels())
+			{
+				if (const TViewModelPtr<FLayerBarModel> LayerBarModel = TrackAreaModel.ImplicitCast())
 				{
-					if (const TViewModelPtr<FLayerBarModel> LayerBarModel = TrackAreaModel.ImplicitCast())
+					if (!LayerBarModels.Contains(LayerBarModel))
 					{
 						LayerBarModels.Add(LayerBarModel);
 					}
 				}
 			}
+
+			for (const TViewModelPtr<FLayerBarModel>& LayerBarModel : Track->GetTrackAreaModelListAs<FLayerBarModel>())
+			{
+				if (!LayerBarModels.Contains(LayerBarModel))
+				{
+					LayerBarModels.Add(LayerBarModel);
+				}
+			}
 		}
 	}
 	else if (OutlinerCount <= 1 && TrackAreaCount > 0)
@@ -451,25 +481,23 @@ TArray<TViewModelPtr<FLayerBarModel>> FAvaStaggerTool::GatherAllSelectedLayerBar
 		}
 	}
 
-	// Remove all layer bars that have a descendant that is selected
-	LayerBarModels.RemoveAll([SequencerSelection](const TViewModelPtr<FLayerBarModel>& InLayerBarModel)
+	// Disallow operations on all layer bars that have a descendant that is selected
+	for (const TViewModelPtr<FLayerBarModel>& InLayerBarModel : LayerBarModels)
+	{
+		const TViewModelPtr<IOutlinerExtension> OutlinerExtension = InLayerBarModel->GetLinkedOutlinerItem();
+		if (!OutlinerExtension.IsValid())
 		{
-			const TViewModelPtr<IOutlinerExtension> OutlinerExtension = InLayerBarModel->GetLinkedOutlinerItem();
-			if (!OutlinerExtension.IsValid())
-			{
-				return false;
-			}
+			continue;
+		}
 
-			for (const TViewModelPtr<IOutlinerExtension>& ChildOutlinerExtension : OutlinerExtension.AsModel()->GetDescendantsOfType<IOutlinerExtension>())
+		for (const TViewModelPtr<IOutlinerExtension>& ChildOutlinerExtension : OutlinerExtension.AsModel()->GetDescendantsOfType<IOutlinerExtension>())
+		{
+			if (SequencerSelection->Outliner.IsSelected(ChildOutlinerExtension))
 			{
-				if (SequencerSelection->Outliner.IsSelected(ChildOutlinerExtension))
-				{
-					return true;
-				}
+				return {};
 			}
-
-			return false;
-		});
+		}
+	}
 
 	return LayerBarModels;
 }
diff --git a/Engine/Plugins/Experimental/Avalanche/Source/AvalancheSequencer/Private/StaggerTool/AvaStaggerTool.h b/Engine/Plugins/Experimental/Avalanche/Source/AvalancheSequencer/Private/StaggerTool/AvaStaggerTool.h
--- a/Engine/Plugins/Experimental/Avalanche/Source/AvalancheSequencer/Private/StaggerTool/AvaStaggerTool.h
+++ b/Engine/Plugins/Experimental/Avalanche/Source/AvalancheSequencer/Private/StaggerTool/AvaStaggerTool.h
@@ -44,6 +44,8 @@ public:
 
 	void AlignToPlayhead();
 
+	void CloseToolWindow();
+
 private:
 	struct FStaggerLayerBarElement
 	{
@@ -66,6 +68,7 @@ private:
 
 	//~ Begin FAvaSequencerAction
 	virtual void MapAction(const TSharedRef<FUICommandList>& InCommandList) override;
+	virtual void OnSequencerClosed() override;
 	//~ End FAvaSequencerAction
 
 	bool CanExecute() const;
