MotionDesign : GeometryMask
- Fix canvas not clearing when disabling or removing mask write modifier
- Fix geometry/material not restoring to original state when removing mask write modifier after level is saved

diff --git a/Engine/Plugins/Experimental/Avalanche/Source/Avalanche/Private/Framework/AvaGizmoComponent.cpp b/Engine/Plugins/Experimental/Avalanche/Source/Avalanche/Private/Framework/AvaGizmoComponent.cpp
--- a/Engine/Plugins/Experimental/Avalanche/Source/Avalanche/Private/Framework/AvaGizmoComponent.cpp
+++ b/Engine/Plugins/Experimental/Avalanche/Source/Avalanche/Private/Framework/AvaGizmoComponent.cpp
@@ -567,6 +567,7 @@ void UAvaGizmoComponent::StoreComponentValues()
 					TArray<UMaterialInterface*> UsedMaterials;
 					InComponent->GetUsedMaterials(UsedMaterials);
 
+					StoredComponentValues.Materials.Reset(UsedMaterials.Num());
 					Algo::Transform(UsedMaterials, StoredComponentValues.Materials, [](UMaterialInterface* InMaterial)
 					{
 						return InMaterial;
diff --git a/Engine/Plugins/Experimental/Avalanche/Source/AvalancheMask/Private/AvaMaskUtilities.cpp b/Engine/Plugins/Experimental/Avalanche/Source/AvalancheMask/Private/AvaMaskUtilities.cpp
--- a/Engine/Plugins/Experimental/Avalanche/Source/AvalancheMask/Private/AvaMaskUtilities.cpp
+++ b/Engine/Plugins/Experimental/Avalanche/Source/AvalancheMask/Private/AvaMaskUtilities.cpp
@@ -90,7 +90,7 @@ namespace UE::AvaMask::Internal
 			Component->CreationMethod = EComponentCreationMethod::Instance;
 
 			// Add to SerializedComponents array so it gets saved
-			InActor->AddOwnedComponent(Component);
+			InActor->AddInstanceComponent(Component);
 			Component->OnComponentCreated();
 			Component->RegisterComponent();
 
diff --git a/Engine/Plugins/Experimental/Avalanche/Source/AvalancheMask/Private/Mask2D/AvaMask2DBaseModifier.cpp b/Engine/Plugins/Experimental/Avalanche/Source/AvalancheMask/Private/Mask2D/AvaMask2DBaseModifier.cpp
--- a/Engine/Plugins/Experimental/Avalanche/Source/AvalancheMask/Private/Mask2D/AvaMask2DBaseModifier.cpp
+++ b/Engine/Plugins/Experimental/Avalanche/Source/AvalancheMask/Private/Mask2D/AvaMask2DBaseModifier.cpp
@@ -240,6 +240,14 @@ void UAvaMask2DBaseModifier::LocalParamsToCanvas()
 }
 
 #if WITH_EDITOR
+void UAvaMask2DBaseModifier::VisualizeMask()
+{
+	if (GEngine)
+	{
+		GEngine->Exec(nullptr, TEXT("GeometryMask.Visualize"));
+	}
+}
+
 void UAvaMask2DBaseModifier::PostEditChangeProperty(FPropertyChangedEvent& InPropertyChangedEvent)
 {
 	Super::PostEditChangeProperty(InPropertyChangedEvent);
@@ -269,7 +277,6 @@ void UAvaMask2DBaseModifier::OnModifierAdded(EActorModifierCoreEnableReason InRe
 	else
 	{
 		CanvasParamsToLocal();
-		MarkModifierDirty();
 	}
 }
 
@@ -310,16 +317,14 @@ void UAvaMask2DBaseModifier::OnModifiedActorTransformed()
 void UAvaMask2DBaseModifier::SavePreState()
 {
 	Super::SavePreState();
-
-	TMap<TWeakObjectPtr<AActor>, FAvaMask2DActorData> PreviousActorDatas = ActorData;
 	
 	TMap<TWeakObjectPtr<AActor>, FAvaMask2DActorData> NewActorDatas;
-	NewActorDatas.Reserve(PreviousActorDatas.Num());
+	NewActorDatas.Reserve(ActorData.Num());
 
 	ForEachActor<AActor>(
 		[this, &NewActorDatas](AActor* InActor)
 		{
-			SaveActorPreState(InActor, NewActorDatas.Emplace(InActor));
+			SaveActorPreState(InActor, NewActorDatas.Emplace(InActor, ActorData.FindOrAdd(InActor)));
 			return true;
 		}
 		, EActorModifierCoreLookup::SelfAndAllChildren);
@@ -330,8 +335,6 @@ void UAvaMask2DBaseModifier::SavePreState()
 void UAvaMask2DBaseModifier::RestorePreState()
 {
 	Super::RestorePreState();
-
-	TGuardValue<bool> RestoreStateGuard(bIsRestoring, true);
 	
 	TMap<TWeakObjectPtr<AActor>, FAvaMask2DActorData> ActorDataCopy = ActorData;
 	for (const TPair<TWeakObjectPtr<AActor>, FAvaMask2DActorData>& ActorDataPair : ActorDataCopy)
diff --git a/Engine/Plugins/Experimental/Avalanche/Source/AvalancheMask/Public/Mask2D/AvaMask2DBaseModifier.h b/Engine/Plugins/Experimental/Avalanche/Source/AvalancheMask/Public/Mask2D/AvaMask2DBaseModifier.h
--- a/Engine/Plugins/Experimental/Avalanche/Source/AvalancheMask/Public/Mask2D/AvaMask2DBaseModifier.h
+++ b/Engine/Plugins/Experimental/Avalanche/Source/AvalancheMask/Public/Mask2D/AvaMask2DBaseModifier.h
@@ -166,7 +166,11 @@ protected:
 	
 	virtual void OnMaterialsChanged(UPrimitiveComponent* InPrimitiveComponent, const TArray<TSharedPtr<IAvaMaskMaterialHandle>>& InMaterialHandles);
 
-protected:
+#if WITH_EDITOR
+	UFUNCTION(CallInEditor, Category = "Mask2D")
+	void VisualizeMask();
+#endif
+
 	/** Whether to get the channel from the parent (first one that specifies a mask channel) */
 	UPROPERTY(EditInstanceOnly, BlueprintReadWrite, Getter = "UseParentChannel", Setter = "SetUseParentChannel", Category = "Mask2D", meta = (AllowPrivateAccess = "true", ValidEnumValues = "Self,Parent"))
 	bool bUseParentChannel = false;
@@ -231,8 +235,6 @@ protected:
 	/** Reference to the underlying canvas texture */
 	TWeakObjectPtr<UTexture> CanvasTextureWeak;
 
-	bool bIsRestoring = false;
-
 	UPROPERTY(DuplicateTransient)
 	TMap<TWeakObjectPtr<UObject>, FInstancedStruct> MaterialCollectionHandleData;
 
diff --git a/Engine/Plugins/Experimental/GeometryMask/Source/GeometryMask/Private/GeometryMaskCanvasResource.cpp b/Engine/Plugins/Experimental/GeometryMask/Source/GeometryMask/Private/GeometryMaskCanvasResource.cpp
--- a/Engine/Plugins/Experimental/GeometryMask/Source/GeometryMask/Private/GeometryMaskCanvasResource.cpp
+++ b/Engine/Plugins/Experimental/GeometryMask/Source/GeometryMask/Private/GeometryMaskCanvasResource.cpp
@@ -497,7 +497,9 @@ void UGeometryMaskCanvasResource::Draw(const ULevel* InLevel, FSceneView& InView
 			CanvasObject->SizeY = Texture->SizeY;
 			CanvasObject->Update();
 			CanvasObject->SetView(&InView);
-			CanvasObject->Canvas->Clear(GetRenderTargetTexture()->ClearColor);
+			CanvasObject->Canvas->Clear(Texture->ClearColor);
+
+			Texture->UpdateResourceImmediate(/** Clear */true);
 		}
 
 		// Contents
