Motion Design Transition: Fix crash with returning TOptional on the Execution Context causing it to be destroyed twice. Both TOptional and std::optional move operations do not reset the 'has value' flag and only move the value.

diff --git a/Engine/Plugins/Experimental/Avalanche/Source/AvalancheTransition/Private/Behavior/AvaTransitionBehaviorInstance.cpp b/Engine/Plugins/Experimental/Avalanche/Source/AvalancheTransition/Private/Behavior/AvaTransitionBehaviorInstance.cpp
--- a/Engine/Plugins/Experimental/Avalanche/Source/AvalancheTransition/Private/Behavior/AvaTransitionBehaviorInstance.cpp
+++ b/Engine/Plugins/Experimental/Avalanche/Source/AvalancheTransition/Private/Behavior/AvaTransitionBehaviorInstance.cpp
@@ -88,8 +88,9 @@ bool FAvaTransitionBehaviorInstance::Setup()
 	// Mark Run Status as running on setup so that on start all the nodes are on this status,
 	// even if their state tree hasn't started yet
 	RunStatus = EStateTreeRunStatus::Running;
-	TOptional<FAvaTransitionExecutionContext> Context = UpdateContext();
-	return Context.IsSet();
+
+	TOptional<FAvaTransitionExecutionContext> Context;
+	return UpdateContext(Context);
 }
 
 void FAvaTransitionBehaviorInstance::Start()
@@ -101,7 +102,8 @@ void FAvaTransitionBehaviorInstance::Start()
 		return;
 	}
 
-	if (TOptional<FAvaTransitionExecutionContext> Context = UpdateContext())
+	TOptional<FAvaTransitionExecutionContext> Context;
+	if (UpdateContext(Context))
 	{
 		const IAvaTransitionBehavior* Behavior = GetBehavior();
 		check(Behavior); // Assume that we have valid behavior if we were able to create the context.
@@ -124,7 +126,8 @@ void FAvaTransitionBehaviorInstance::Tick(float InDeltaSeconds)
 		return;
 	}
 
-	if (TOptional<FAvaTransitionExecutionContext> Context = UpdateContext())
+	TOptional<FAvaTransitionExecutionContext> Context;
+	if (UpdateContext(Context))
 	{
 		RunStatus = Context->Tick(InDeltaSeconds);
 	}
@@ -138,7 +141,8 @@ void FAvaTransitionBehaviorInstance::Tick(float InDeltaSeconds)
 
 void FAvaTransitionBehaviorInstance::Stop()
 {
-	if (TOptional<FAvaTransitionExecutionContext> Context = UpdateContext())
+	TOptional<FAvaTransitionExecutionContext> Context;
+	if (UpdateContext(Context))
 	{
 		RunStatus = Context->Stop();
 	}
@@ -172,15 +176,15 @@ void FAvaTransitionBehaviorInstance::ConditionallyStop()
 	}
 }
 
-TOptional<FAvaTransitionExecutionContext> FAvaTransitionBehaviorInstance::UpdateContext()
+bool FAvaTransitionBehaviorInstance::UpdateContext(TOptional<FAvaTransitionExecutionContext>& OutContext)
 {
 	if (ValidateTransitionScene())
 	{
 		IAvaTransitionBehavior* Behavior = GetBehavior();
 		UpdateTransitionLayers(Behavior);
-		return MakeContext(Behavior);
+		return MakeContext(Behavior, OutContext);
 	}
-	return TOptional<FAvaTransitionExecutionContext>();
+	return false;
 }
 
 bool FAvaTransitionBehaviorInstance::ValidateTransitionScene()
@@ -237,12 +241,12 @@ void FAvaTransitionBehaviorInstance::UpdateTransitionLayers(const IAvaTransition
 	}
 }
 
-TOptional<FAvaTransitionExecutionContext> FAvaTransitionBehaviorInstance::MakeContext(IAvaTransitionBehavior* InBehavior)
+bool FAvaTransitionBehaviorInstance::MakeContext(IAvaTransitionBehavior* InBehavior, TOptional<FAvaTransitionExecutionContext>& OutContext)
 {
 	FAvaTransitionScene* TransitionScene = TransitionContext.GetTransitionScene();
 	if (!InBehavior || !TransitionScene)
 	{
-		return TOptional<FAvaTransitionExecutionContext>();
+		return false;
 	}
 
 	const ULevel* Level = TransitionScene->GetLevel();
@@ -254,15 +258,15 @@ TOptional<FAvaTransitionExecutionContext> FAvaTransitionBehaviorInstance::MakeCo
 		UE_LOG(LogAvaTransition, Error
 			, TEXT("SetContextRequirements failed for '%s'. World is invalid")
 			, *LogContext);
-		return TOptional<FAvaTransitionExecutionContext>();
+		return false;
 	}
 
 	const FStateTreeReference& StateTreeReference = InBehavior->GetStateTreeReference();
 
-	FAvaTransitionExecutionContext Context(*this, InBehavior->AsUObject(), *StateTreeReference.GetStateTree(), InstanceData);
+	FAvaTransitionExecutionContext& Context = OutContext.Emplace(*this, InBehavior->AsUObject(), *StateTreeReference.GetStateTree(), InstanceData);
 	if (!Context.IsValid())
 	{
-		return TOptional<FAvaTransitionExecutionContext>();
+		return false;
 	}
 
 	Context.SetCollectExternalDataCallback(FOnCollectStateTreeExternalData::CreateLambda(
@@ -292,15 +296,15 @@ TOptional<FAvaTransitionExecutionContext> FAvaTransitionBehaviorInstance::MakeCo
 			return true;
 		})
 	);
-	
+
 	if (!Context.AreContextDataViewsValid())
 	{
 		UE_LOG(LogAvaTransition, Error
 			, TEXT("SetContextRequirements failed for '%s'. Missing external data requirements. StateTree will not update.")
 			, *LogContext);
 
-		return TOptional<FAvaTransitionExecutionContext>();
+		return false;
 	}
 
-	return Context;
+	return true;
 }
diff --git a/Engine/Plugins/Experimental/Avalanche/Source/AvalancheTransition/Public/Behavior/AvaTransitionBehaviorInstance.h b/Engine/Plugins/Experimental/Avalanche/Source/AvalancheTransition/Public/Behavior/AvaTransitionBehaviorInstance.h
--- a/Engine/Plugins/Experimental/Avalanche/Source/AvalancheTransition/Public/Behavior/AvaTransitionBehaviorInstance.h
+++ b/Engine/Plugins/Experimental/Avalanche/Source/AvalancheTransition/Public/Behavior/AvaTransitionBehaviorInstance.h
@@ -74,13 +74,13 @@ private:
 	/** Stops Execution if the Tree is no longer running */
 	void ConditionallyStop();
 
-	TOptional<FAvaTransitionExecutionContext> UpdateContext();
+	bool UpdateContext(TOptional<FAvaTransitionExecutionContext>& OutContext);
 
 	bool ValidateTransitionScene();
 
 	void UpdateTransitionLayers(const IAvaTransitionBehavior* InBehavior);
 
-	TOptional<FAvaTransitionExecutionContext> MakeContext(IAvaTransitionBehavior* InBehavior);
+	bool MakeContext(IAvaTransitionBehavior* InBehavior, TOptional<FAvaTransitionExecutionContext>& OutContext);
 
 	/** The Behavior this Instance is based on */
 	TWeakInterfacePtr<IAvaTransitionBehavior> BehaviorWeak;
