MotionDesign : ActorModifier
- Adding new plugin to regroup modifiers implementation outside of MotionDesign context to be used as standalone
- Moving resources (icons) into ActorModifier/Content/Icons
- Layout modifiers moved to their own module ActorModifierLayout (AlignBetween, AutoFollow, GridArrange, Justify, LookAt, RadialArrange)
- Shared modifiers files moved to their own module ActorModifier (extensions, shared, utilities)
- Modifier type customizations moved to their own module ActorModifierEditor
- Adding redirects to avoid breaking existing assets
- Tests are all passing

diff --git a/Engine/Plugins/Experimental/ActorModifier/Config/BaseActorModifier.ini b/Engine/Plugins/Experimental/ActorModifier/Config/BaseActorModifier.ini
new file mode 100644
--- /dev/null
+++ b/Engine/Plugins/Experimental/ActorModifier/Config/BaseActorModifier.ini
@@ -0,0 +1,32 @@
+﻿[CoreRedirects]
++EnumRedirects=(OldName="/Script/AvalancheModifiers.EAvaTransformSharedModifierState",NewName="/Script/ActorModifier.EActorModifierTransformSharedState")
++EnumRedirects=(OldName="/Script/Avalanche.EAvaReferenceContainer",NewName="/Script/ActorModifier.EActorModifierReferenceContainer")
++EnumRedirects=(OldName="/Script/AvalancheModifiers.EAvaVisibilityActor",NewName="/Script/ActorModifier.EActorModifierVisibilityActor")
++EnumRedirects=(OldName="/Script/Avalanche.EAvaCorner2D",NewName="/Script/ActorModifierLayout.EActorModifierGridArrangeCorner2D")
++EnumRedirects=(OldName="/Script/Avalanche.EAvaAxis",NewName="/Script/ActorModifier.EActorModifierAlignment")
++EnumRedirects=(OldName="/Script/AvalancheModifiers.EAvaGridArrangeDirection",NewName="/Script/ActorModifierLayout.EActorModifierGridArrangeDirection")
++EnumRedirects=(OldName="/Script/AvalancheModifiers.EAvaJustifyVertical",NewName="/Script/ActorModifierLayout.EActorModifierJustifyVertical")
++EnumRedirects=(OldName="/Script/AvalancheModifiers.EAvaJustifyHorizontal",NewName="/Script/ActorModifierLayout.EActorModifierJustifyHorizontal")
++EnumRedirects=(OldName="/Script/AvalancheModifiers.EAvaJustifyDepth",NewName="/Script/ActorModifierLayout.EActorModifierJustifyDepth")
++EnumRedirects=(OldName="/Script/AvalancheModifiers.EAvaRadialArrangement",NewName="/Script/ActorModifierLayout.EActorModifierRadialArrangeMode")
++EnumRedirects=(OldName="/Script/AvalancheModifiers.EAvaRadialArrangePlane",NewName="/Script/ActorModifierLayout.EActorModifierRadialArrangePlane")
++EnumRedirects=(OldName="/Script/AvalancheModifiers.EAvaModifiersAxis",NewName="/Script/ActorModifier.EActorModifierAxis")
+
++StructRedirects=(OldName="/Script/AvalancheModifiers.AvaTransformSharedModifierState",NewName="/Script/ActorModifier.ActorModifierTransformSharedModifierState")
++StructRedirects=(OldName="/Script/AvalancheModifiers.AvaTransformSharedActorState",NewName="/Script/ActorModifier.ActorModifierTransformSharedActorState")
++StructRedirects=(OldName="/Script/AvalancheModifiers.AvaSceneTreeActor",NewName="/Script/ActorModifier.ActorModifierSceneTreeActor")
++StructRedirects=(OldName="/Script/AvalancheModifiers.AvaVisibilitySharedModifierState",NewName="/Script/ActorModifier.ActorModifierVisibilitySharedModifierState")
++StructRedirects=(OldName="/Script/AvalancheModifiers.AvaVisibilitySharedActorState",NewName="/Script/ActorModifier.ActorModifierVisibilitySharedActorState")
++StructRedirects=(OldName="/Script/AvalancheModifiers.AvaAlignBetweenWeightedActor",NewName="/Script/ActorModifierLayout.ActorModifierAlignBetweenWeightedActor")
++StructRedirects=(OldName="/Script/Avalanche.AvaAnchorAlignment",NewName="/Script/ActorModifier.ActorModifierAnchorAlignment")
+
++ClassRedirects=(OldName="/Script/AvalancheModifiers.AvaTransformModifierShared",NewName="/Script/ActorModifier.ActorModifierTransformShared")
++ClassRedirects=(OldName="/Script/AvalancheModifiers.AvaVisibilityModifierShared",NewName="/Script/ActorModifier.ActorModifierVisibilityShared")
++ClassRedirects=(OldName="/Script/AvalancheModifiers.AvaAttachmentBaseModifier",NewName="/Script/ActorModifier.ActorModifierAttachmentBaseModifier")
++ClassRedirects=(OldName="/Script/AvalancheModifiers.AvaArrangeBaseModifier",NewName="/Script/ActorModifier.ActorModifierArrangeBaseModifier")
++ClassRedirects=(OldName="/Script/AvalancheModifiers.AvaAlignBetweenModifier",NewName="/Script/ActorModifierLayout.ActorModifierAlignBetweenModifier")
++ClassRedirects=(OldName="/Script/AvalancheModifiers.AvaAutoFollowModifier",NewName="/Script/ActorModifierLayout.ActorModifierAutoFollowModifier")
++ClassRedirects=(OldName="/Script/AvalancheModifiers.AvaGridArrangeModifier",NewName="/Script/ActorModifierLayout.ActorModifierGridArrangeModifier")
++ClassRedirects=(OldName="/Script/AvalancheModifiers.AvaJustifyModifier",NewName="/Script/ActorModifierLayout.ActorModifierJustifyModifier")
++ClassRedirects=(OldName="/Script/AvalancheModifiers.AvaLookAtModifier",NewName="/Script/ActorModifierLayout.ActorModifierLookAtModifier")
++ClassRedirects=(OldName="/Script/AvalancheModifiers.AvaRadialArrangeModifier",NewName="/Script/ActorModifierLayout.ActorModifierRadialArrangeModifier")
diff --git a/Engine/Plugins/Experimental/ActorModifier/Source/ActorModifier/Private/ActorModifierModule.cpp b/Engine/Plugins/Experimental/ActorModifier/Source/ActorModifier/Private/ActorModifierModule.cpp
new file mode 100644
--- /dev/null
+++ b/Engine/Plugins/Experimental/ActorModifier/Source/ActorModifier/Private/ActorModifierModule.cpp
@@ -0,0 +1,5 @@
+﻿// Copyright Epic Games, Inc. All Rights Reserved.
+
+#include "Modules/ModuleManager.h"
+	
+IMPLEMENT_MODULE(FDefaultModuleImpl, ActorModifier)
\ No newline at end of file
diff --git a/Engine/Plugins/Experimental/ActorModifier/Source/ActorModifierEditor/Private/ActorModifierEditorModule.cpp b/Engine/Plugins/Experimental/ActorModifier/Source/ActorModifierEditor/Private/ActorModifierEditorModule.cpp
new file mode 100644
--- /dev/null
+++ b/Engine/Plugins/Experimental/ActorModifier/Source/ActorModifierEditor/Private/ActorModifierEditorModule.cpp
@@ -0,0 +1,33 @@
+﻿// Copyright Epic Games, Inc. All Rights Reserved.
+
+#include "ActorModifierEditorModule.h"
+
+#include "ActorModifierTypes.h"
+#include "Customizations/ActorModifierEditorAnchorAlignmentPropertyTypeCustomization.h"
+#include "PropertyEditorModule.h"
+#include "Styles/ActorModifierEditorStyle.h"
+
+#define LOCTEXT_NAMESPACE "ActorModifierEditorModule"
+
+IMPLEMENT_MODULE(FActorModifierEditorModule, ActorModifierEditor)
+
+void FActorModifierEditorModule::StartupModule()
+{
+	FActorModifierEditorStyle::Get();
+	
+	FPropertyEditorModule& PropertyModule = FModuleManager::GetModuleChecked<FPropertyEditorModule>(TEXT("PropertyEditor"));
+
+	PropertyModule.RegisterCustomPropertyTypeLayout(FActorModifierAnchorAlignment::StaticStruct()->GetFName()
+	, FOnGetPropertyTypeCustomizationInstance::CreateStatic(&FActorModifierEditorAnchorAlignmentPropertyTypeCustomization::MakeInstance));
+}
+
+void FActorModifierEditorModule::ShutdownModule()
+{
+	if (FPropertyEditorModule* PropertyModule = FModuleManager::GetModulePtr<FPropertyEditorModule>(TEXT("PropertyEditor")))
+	{
+		PropertyModule->UnregisterCustomPropertyTypeLayout(FActorModifierAnchorAlignment::StaticStruct()->GetFName());
+	}
+}
+
+#undef LOCTEXT_NAMESPACE
+    
diff --git a/Engine/Plugins/Experimental/ActorModifier/Source/ActorModifierLayout/Private/ActorModifierLayoutModule.cpp b/Engine/Plugins/Experimental/ActorModifier/Source/ActorModifierLayout/Private/ActorModifierLayoutModule.cpp
new file mode 100644
--- /dev/null
+++ b/Engine/Plugins/Experimental/ActorModifier/Source/ActorModifierLayout/Private/ActorModifierLayoutModule.cpp
@@ -0,0 +1,5 @@
+﻿// Copyright Epic Games, Inc. All Rights Reserved.
+
+#include "Modules/ModuleManager.h"
+    
+IMPLEMENT_MODULE(FDefaultModuleImpl, ActorModifierLayout)
\ No newline at end of file
diff --git a/Engine/Plugins/Experimental/ActorModifierCore/Source/ActorModifierCore/Private/Modifiers/ActorModifierCoreExtension.cpp b/Engine/Plugins/Experimental/ActorModifierCore/Source/ActorModifierCore/Private/Modifiers/ActorModifierCoreExtension.cpp
--- a/Engine/Plugins/Experimental/ActorModifierCore/Source/ActorModifierCore/Private/Modifiers/ActorModifierCoreExtension.cpp
+++ b/Engine/Plugins/Experimental/ActorModifierCore/Source/ActorModifierCore/Private/Modifiers/ActorModifierCoreExtension.cpp
@@ -18,6 +18,7 @@ void FActorModifierCoreExtension::ConstructInternal(UActorModifierCoreBase* InMo
 	ModifierWeak = InModifier;
 	ExtensionType = InExtensionType;
 	bExtensionInitialized = true;
+	OnExtensionInitialized();
 }
 
 void FActorModifierCoreExtension::EnableExtension(EActorModifierCoreEnableReason InReason)
@@ -51,3 +52,9 @@ UWorld* FActorModifierCoreExtension::GetModifierWorld() const
 	const AActor* ModifiedActor = GetModifierActor();
 	return ModifiedActor ? ModifiedActor->GetWorld() : nullptr;
 }
+
+ULevel* FActorModifierCoreExtension::GetModifierLevel() const
+{
+	const AActor* ModifiedActor = GetModifierActor();
+	return ModifiedActor ? ModifiedActor->GetLevel() : nullptr;
+}
diff --git a/Engine/Plugins/Experimental/ActorModifierCore/Source/ActorModifierCore/Public/Modifiers/ActorModifierCoreExtension.h b/Engine/Plugins/Experimental/ActorModifierCore/Source/ActorModifierCore/Public/Modifiers/ActorModifierCoreExtension.h
--- a/Engine/Plugins/Experimental/ActorModifierCore/Source/ActorModifierCore/Public/Modifiers/ActorModifierCoreExtension.h
+++ b/Engine/Plugins/Experimental/ActorModifierCore/Source/ActorModifierCore/Public/Modifiers/ActorModifierCoreExtension.h
@@ -24,6 +24,9 @@ public:
 	ACTORMODIFIERCORE_API virtual ~FActorModifierCoreExtension();
 
 protected:
+	/** Overwrite this, called when the extension is created */
+	ACTORMODIFIERCORE_API virtual void OnExtensionInitialized() {}
+	
 	/** Overwrite this, called when the modifier is enabled */
 	ACTORMODIFIERCORE_API virtual void OnExtensionEnabled(EActorModifierCoreEnableReason InReason) {}
 
@@ -36,6 +39,9 @@ protected:
 	/** Get the modifier actor world */
 	ACTORMODIFIERCORE_API UWorld* GetModifierWorld() const;
 
+	/** Get the modifier actor level */
+	ACTORMODIFIERCORE_API ULevel* GetModifierLevel() const;
+
 	/** Get the casted modifier to which this extension is linked */
 	template <typename InModifierClass, typename = typename TEnableIf<TIsDerivedFrom<InModifierClass, UActorModifierCoreBase>::IsDerived>::Type>
 	InModifierClass* GetModifier() const
diff --git a/Engine/Plugins/Experimental/Avalanche/Avalanche.uplugin b/Engine/Plugins/Experimental/Avalanche/Avalanche.uplugin
--- a/Engine/Plugins/Experimental/Avalanche/Avalanche.uplugin
+++ b/Engine/Plugins/Experimental/Avalanche/Avalanche.uplugin
@@ -331,6 +331,10 @@
 			"Name": "ActorModifierCore",
 			"Enabled": true
 		},
+		{
+			"Name": "ActorModifier",
+			"Enabled": true
+		},
 		{
 			"Name": "OperatorStack",
 			"Enabled": true
diff --git a/Engine/Plugins/Experimental/Avalanche/Source/Avalanche/Private/AvaActorUtils.cpp b/Engine/Plugins/Experimental/Avalanche/Source/Avalanche/Private/AvaActorUtils.cpp
--- a/Engine/Plugins/Experimental/Avalanche/Source/Avalanche/Private/AvaActorUtils.cpp
+++ b/Engine/Plugins/Experimental/Avalanche/Source/Avalanche/Private/AvaActorUtils.cpp
@@ -135,134 +135,4 @@ IAvaSceneInterface* FAvaActorUtils::GetSceneInterfaceFromActor(const AActor* InA
 	}
 
 	return SceneSubsystem->GetSceneInterface(InActor->GetLevel());
-}
-
-AActor* FAvaActorUtils::ActorFromReferenceContainer(AActor* const InActor, const EAvaReferenceContainer InReferenceContainer, const bool bInIgnoreHiddenActors)
-{
-	if (InReferenceContainer == EAvaReferenceContainer::Other || !IsValid(InActor))
-	{
-		return nullptr;
-	}
-
-	// Note: Using Typed Outer World, instead of GetWorld(), since the typed outer World could be a streamed in world
-	// and GetWorld() only returns the main world.
-	UWorld* const World = InActor->GetTypedOuter<UWorld>();
-	if (!IsValid(World))
-	{
-		return nullptr;
-	}
-
-	AActor* const ParentActor = InActor->GetAttachParentActor();
-
-	bool bIsOutlinerAttachedActors = false;
-	TArray<AActor*> AttachedActors;
-
-#if WITH_EDITOR
-	UAvaOutlinerSubsystem* const OutlinerSubsystem = World->GetSubsystem<UAvaOutlinerSubsystem>();
-	if (IsValid(OutlinerSubsystem))
-	{
-		TSharedPtr<IAvaOutliner> AvaOutliner = OutlinerSubsystem->GetOutliner();
-		if (AvaOutliner.IsValid())
-		{
-			bIsOutlinerAttachedActors = true;
-			if (IsValid(ParentActor))
-			{
-				AttachedActors = FAvaOutlinerUtils::EditorOutlinerChildActors(AvaOutliner, ParentActor);
-			}
-			else // no valid parent, use world actors instead
-			{
-				TArray<FAvaOutlinerItemPtr> OutlinerRootChildren = AvaOutliner->GetTreeRoot()->GetChildren();
-				const bool bValidRootActors = FAvaOutlinerUtils::EditorOutlinerItemsToActors(OutlinerRootChildren, AttachedActors);
-			}
-		}
-	}
-#endif
-
-	if (!bIsOutlinerAttachedActors)
-	{
-		if (const IAvaSceneInterface* SceneInterface = FAvaActorUtils::GetSceneInterfaceFromActor(InActor))
-		{
-			const FAvaSceneTree& SceneTree = SceneInterface->GetSceneTree();
-
-			if (!IsValid(ParentActor))
-			{
-				for (const int32 ChildId : SceneTree.GetRootNode().GetChildrenIndices())
-				{
-					const FAvaSceneItem* Item = SceneTree.GetItemAtIndex(ChildId);
-					AttachedActors.Add(Item->Resolve<AActor>(World));
-				}
-			}
-			else
-			{
-				AttachedActors = SceneTree.GetChildActors(ParentActor);
-			}
-		}
-	}
-
-	if (AttachedActors.Num() == 0)
-	{
-		return nullptr;
-	}
-
-	if (bInIgnoreHiddenActors)
-	{
-		for (int32 ChildIndex = AttachedActors.Num() - 1; ChildIndex >= 0; --ChildIndex)
-		{
-			if (AttachedActors[ChildIndex]->IsHidden()
-#if WITH_EDITOR
-				|| AttachedActors[ChildIndex]->IsTemporarilyHiddenInEditor()
-#endif
-				)
-			{
-				AttachedActors.RemoveAt(ChildIndex);
-			}
-		}
-	}
-
-	int32 FromActorIndex = -1;
-	if (!AttachedActors.Find(InActor, FromActorIndex))
-	{
-		return nullptr;
-	}
-
-	switch (InReferenceContainer)
-	{
-		case EAvaReferenceContainer::Previous:
-		{
-			const int32 PreviousActorIndex = FromActorIndex - 1;
-			if (!AttachedActors.IsValidIndex(PreviousActorIndex))
-			{
-				return ParentActor;
-			}
-			return AttachedActors[PreviousActorIndex];
-		}
-		case EAvaReferenceContainer::Next:
-		{
-			const int32 NextActorIndex = FromActorIndex + 1;
-			if (!AttachedActors.IsValidIndex(NextActorIndex))
-			{
-				return nullptr;
-			}
-			return AttachedActors[NextActorIndex];
-		}
-		case EAvaReferenceContainer::First:
-		{
-			if (!AttachedActors.IsValidIndex(0))
-			{
-				return nullptr;
-			}
-			return AttachedActors[0] != InActor ? AttachedActors[0] : nullptr;
-		}
-		case EAvaReferenceContainer::Last:
-		{
-			AActor* const LastChildActor = AttachedActors.Last();
-			if (!IsValid(LastChildActor))
-			{
-				return nullptr;
-			}
-			return LastChildActor != InActor ? LastChildActor : nullptr;
-		}
-	}
-
-	return nullptr;
-}
+}
\ No newline at end of file
diff --git a/Engine/Plugins/Experimental/Avalanche/Source/Avalanche/Private/AvaOutlinerUtils.cpp b/Engine/Plugins/Experimental/Avalanche/Source/Avalanche/Private/AvaOutlinerUtils.cpp
--- a/Engine/Plugins/Experimental/Avalanche/Source/Avalanche/Private/AvaOutlinerUtils.cpp
+++ b/Engine/Plugins/Experimental/Avalanche/Source/Avalanche/Private/AvaOutlinerUtils.cpp
@@ -3,12 +3,12 @@
 #include "AvaOutlinerUtils.h"
 
 #if WITH_EDITOR
-
 #include "AvaOutlinerSubsystem.h"
 #include "Engine/World.h"
 #include "GameFramework/Actor.h"
 #include "IAvaOutliner.h"
 #include "Item/AvaOutlinerActor.h"
+#include "Item/AvaOutlinerLevel.h"
 
 TSharedPtr<IAvaOutliner> FAvaOutlinerUtils::EditorGetOutliner(const UWorld* const InWorld)
 {
@@ -41,6 +41,13 @@ bool FAvaOutlinerUtils::EditorOutlinerItemsToActors(const TArray<FAvaOutlinerIte
 			return false;
 		}
 
+		// Root of the tree is level
+		const FAvaOutlinerLevel* const LevelItem = InItems[i]->CastTo<FAvaOutlinerLevel>();
+		if (LevelItem)
+		{
+			return EditorOutlinerItemsToActors(LevelItem->GetChildren(), OutActors);
+		}
+
 		const FAvaOutlinerActor* const ActorItem = InItems[i]->CastTo<FAvaOutlinerActor>();
 		if (ActorItem)
 		{
diff --git a/Engine/Plugins/Experimental/Avalanche/Source/Avalanche/Public/AvaActorUtils.h b/Engine/Plugins/Experimental/Avalanche/Source/Avalanche/Public/AvaActorUtils.h
--- a/Engine/Plugins/Experimental/Avalanche/Source/Avalanche/Public/AvaActorUtils.h
+++ b/Engine/Plugins/Experimental/Avalanche/Source/Avalanche/Public/AvaActorUtils.h
@@ -9,7 +9,7 @@
 class AActor;
 class IAvaSceneInterface;
 class USceneComponent;
-enum class EAvaReferenceContainer : uint8;
+enum class EActorModifierReferenceContainer : uint8;
 struct FOrientedBox;
 
 struct AVALANCHE_API FAvaActorUtils
@@ -21,7 +21,4 @@ struct AVALANCHE_API FAvaActorUtils
 	static FBox GetComponentLocalBoundingBox(const USceneComponent* InComponent);
 
 	static IAvaSceneInterface* GetSceneInterfaceFromActor(const AActor* InActor);
-
-	/** Finds an actor relative to another based on the type of reference container used. */
-	static AActor* ActorFromReferenceContainer(AActor* const InActor, const EAvaReferenceContainer InReferenceContainer, const bool bInIgnoreHiddenActors = false);
 };
diff --git a/Engine/Plugins/Experimental/Avalanche/Source/Avalanche/Public/AvaDefs.h b/Engine/Plugins/Experimental/Avalanche/Source/Avalanche/Public/AvaDefs.h
--- a/Engine/Plugins/Experimental/Avalanche/Source/Avalanche/Public/AvaDefs.h
+++ b/Engine/Plugins/Experimental/Avalanche/Source/Avalanche/Public/AvaDefs.h
@@ -65,41 +65,6 @@ enum class EAvaDepthAlignment : uint8
 	Back
 };
 
-/** Specifies the method for finding a reference actor based on it's position in the parent's hierarchy. */
-UENUM(BlueprintType)
-enum class EAvaReferenceContainer : uint8
-{
-	/** Uses the previous actor in the parent's hierarchy. */
-	Previous,
-	/** Uses the next actor in the parent's hierarchy. */
-	Next,
-	/** Uses the first actor in the parent's hierarchy. */
-	First,
-	/** Uses the last actor in the parent's hierarchy. */
-	Last,
-	/** Uses a specified reference actor set by the user. */
-	Other
-};
-
-/** Specifies a corner for a 2D rectangle. */
-UENUM(BlueprintType)
-enum class EAvaCorner2D : uint8
-{
-	TopLeft,
-	TopRight,
-	BottomLeft,
-	BottomRight
-};
-
-/** Specifies a 3D axis. */
-UENUM(BlueprintType)
-enum class EAvaAxis : uint8
-{
-	Horizontal,
-	Vertical,
-	Depth
-};
-
 template<typename Enum>
 constexpr auto ToUnderlyingType(Enum Val)
 {
@@ -245,77 +210,3 @@ struct FAvaColorChangeData
 		return !operator==(InOtherColorData);
 	};
 };
-
-/** Specifies a set of anchor alignments, one for each 3D axis. */
-USTRUCT(BlueprintType)
-struct FAvaAnchorAlignment
-{
-	GENERATED_BODY()
-
-	UPROPERTY()
-	bool bUseHorizontal = true;
-	
-	UPROPERTY()
-	bool bUseVertical = true;
-	
-	UPROPERTY()
-	bool bUseDepth = true;
-
-	UPROPERTY(EditAnywhere, BlueprintReadWrite, Category="AnchorAlignment", meta=(EditCondition="bUseHorizontal", EditConditionHides))
-	EAvaHorizontalAlignment Horizontal = EAvaHorizontalAlignment::Center;
-
-	UPROPERTY(EditAnywhere, BlueprintReadWrite, Category="AnchorAlignment", meta=(EditCondition="bUseVertical", EditConditionHides))
-	EAvaVerticalAlignment Vertical = EAvaVerticalAlignment::Center;
-
-	UPROPERTY(EditAnywhere, BlueprintReadWrite, Category="AnchorAlignment", meta=(EditCondition="bUseDepth", EditConditionHides))
-	EAvaDepthAlignment Depth = EAvaDepthAlignment::Center;
-
-	FAvaAnchorAlignment() {}
-	FAvaAnchorAlignment(const EAvaHorizontalAlignment InHorizontal, const EAvaVerticalAlignment InVertical, const EAvaDepthAlignment InDepth)
-		: Horizontal(InHorizontal), Vertical(InVertical), Depth(InDepth)
-	{}
-
-	bool operator==(const FAvaAnchorAlignment& Other) const
-	{
-		return Horizontal == Other.Horizontal && Vertical == Other.Vertical && Depth == Other.Depth;
-	}
-
-	FORCEINLINE bool IsHorizontalMatch(const FAvaAnchorAlignment& Other) const { return Horizontal == Other.Horizontal; }
-	FORCEINLINE bool IsVerticalMatch(const FAvaAnchorAlignment& Other) const { return Vertical == Other.Vertical; }
-	FORCEINLINE bool IsDepthMatch(const FAvaAnchorAlignment& Other) const { return Depth == Other.Depth; }
-
-	/** Returns a point on the bounds extent determined by the alignment properties of the structure. */
-	FVector LocalBoundsOffset(const FBox& InBounds, const bool bInInverted = false) const
-	{
-		const FVector Extent = InBounds.GetExtent();
-
-		FVector BoundsOffset = FVector::ZeroVector;
-
-		auto HorizontalAlign = [bInInverted, Extent]() -> double { return bInInverted ? Extent.Y : -Extent.Y; };
-		auto VerticalAlign = [bInInverted, Extent]() -> double { return bInInverted ? Extent.Z : -Extent.Z; };
-		auto DepthAlign = [bInInverted, Extent]() -> double { return bInInverted ? Extent.X : -Extent.X; };
-
-		switch (Horizontal)
-		{
-		case EAvaHorizontalAlignment::Left: BoundsOffset.Y = HorizontalAlign(); break;
-		case EAvaHorizontalAlignment::Center: break;
-		case EAvaHorizontalAlignment::Right: BoundsOffset.Y = -HorizontalAlign(); break;
-		}
-
-		switch (Vertical)
-		{
-		case EAvaVerticalAlignment::Top: BoundsOffset.Z = VerticalAlign(); break;
-		case EAvaVerticalAlignment::Center: break;
-		case EAvaVerticalAlignment::Bottom: BoundsOffset.Z = -VerticalAlign(); break;
-		}
-
-		switch (Depth)
-		{
-		case EAvaDepthAlignment::Front: BoundsOffset.X = DepthAlign(); break;
-		case EAvaDepthAlignment::Center: break;
-		case EAvaDepthAlignment::Back: BoundsOffset.X = -DepthAlign(); break;
-		}
-
-		return BoundsOffset;
-	}
-};
diff --git a/Engine/Plugins/Experimental/Avalanche/Source/AvalancheEditor/AvalancheEditor.Build.cs b/Engine/Plugins/Experimental/Avalanche/Source/AvalancheEditor/AvalancheEditor.Build.cs
--- a/Engine/Plugins/Experimental/Avalanche/Source/AvalancheEditor/AvalancheEditor.Build.cs
+++ b/Engine/Plugins/Experimental/Avalanche/Source/AvalancheEditor/AvalancheEditor.Build.cs
@@ -22,6 +22,7 @@ public class AvalancheEditor : ModuleRules
 			new string[]
 			{
 				"ActorModifierCore",
+				"ActorModifier",
 				"ActorModifierCoreEditor",
 				"AdvancedPreviewScene",
 				"AdvancedRenamer",
diff --git a/Engine/Plugins/Experimental/Avalanche/Source/AvalancheEditor/Private/AvaEditorModule.cpp b/Engine/Plugins/Experimental/Avalanche/Source/AvalancheEditor/Private/AvaEditorModule.cpp
--- a/Engine/Plugins/Experimental/Avalanche/Source/AvalancheEditor/Private/AvaEditorModule.cpp
+++ b/Engine/Plugins/Experimental/Avalanche/Source/AvalancheEditor/Private/AvaEditorModule.cpp
@@ -40,7 +40,6 @@
 // Details View
 #include "AssetToolsModule.h"
 #include "AvaSceneSettings.h"
-#include "DetailView/Customizations/AvaAnchorAlignmentPropertyTypeCustomization.h"
 #include "DetailView/Customizations/AvaCategoryHiderCustomization.h"
 #include "DetailView/Customizations/AvaMeshesDetailCustomization.h"
 #include "DetailView/Customizations/AvaSceneSettingsCustomization.h"
@@ -284,8 +283,6 @@ void FAvaEditorModule::RegisterCustomLayouts()
 	
 	PropertyModule.RegisterCustomPropertyTypeLayout(FAvaViewportQualitySettings::StaticStruct()->GetFName()
 		, FOnGetPropertyTypeCustomizationInstance::CreateStatic(&FAvaViewportQualitySettingsPropertyTypeCustomization::MakeInstance));
-	PropertyModule.RegisterCustomPropertyTypeLayout(FAvaAnchorAlignment::StaticStruct()->GetFName()
-		, FOnGetPropertyTypeCustomizationInstance::CreateStatic(&FAvaAnchorAlignmentPropertyTypeCustomization::MakeInstance));
 }
 
 void FAvaEditorModule::UnregisterCustomLayouts()
@@ -309,7 +306,6 @@ void FAvaEditorModule::UnregisterCustomLayouts()
 		PropertyModule.UnregisterCustomPropertyTypeLayout(TEXT("Vector2D"));
 
 		PropertyModule.UnregisterCustomPropertyTypeLayout(FAvaViewportQualitySettings::StaticStruct()->GetFName());
-		PropertyModule.UnregisterCustomPropertyTypeLayout(FAvaAnchorAlignment::StaticStruct()->GetFName());
 	}
 }
 
diff --git a/Engine/Plugins/Experimental/Avalanche/Source/AvalancheEditor/Private/Outliner/AvaOutlinerExtension.cpp b/Engine/Plugins/Experimental/Avalanche/Source/AvalancheEditor/Private/Outliner/AvaOutlinerExtension.cpp
--- a/Engine/Plugins/Experimental/Avalanche/Source/AvalancheEditor/Private/Outliner/AvaOutlinerExtension.cpp
+++ b/Engine/Plugins/Experimental/Avalanche/Source/AvalancheEditor/Private/Outliner/AvaOutlinerExtension.cpp
@@ -30,7 +30,6 @@
 #include "LevelEditor.h"
 #include "Materials/Material.h"
 #include "Materials/MaterialInterface.h"
-#include "Modifiers/AvaBaseModifier.h"
 #include "Outliner/AvaOutlinerMaterialDesignerProxy.h"
 #include "ScopedTransaction.h"
 #include "Sequencer/AvaSequencerExtension.h"
@@ -38,6 +37,7 @@
 #include "Text3DActor.h"
 #include "ToolMenuContext/AvaOutlinerItemsContext.h"
 #include "ToolMenu.h"
+#include "Modifiers/ActorModifierCoreBase.h"
 #include "Viewport/AvaViewportExtension.h"
 
 #define LOCTEXT_NAMESPACE "AvaOutlinerExtension"
@@ -343,7 +343,7 @@ void FAvaOutlinerExtension::ExtendOutlinerItemFilters(TArray<TSharedPtr<IAvaOutl
 	OutItemFilters.Add(MakeShared<FAvaOutlinerItemTypeFilter>(TEXT("Modifiers")
 		, TArray<TSubclassOf<UObject>>{ UActorModifierCoreBase::StaticClass() }
 		, EAvaOutlinerTypeFilterMode::MatchesType | EAvaOutlinerTypeFilterMode::ContainerOfType
-		, FSlateIconFinder::FindIconForClass(UAvaBaseModifier::StaticClass()).GetIcon()
+		, FSlateIconFinder::FindIconForClass(UActorModifierCoreBase::StaticClass()).GetIcon()
 		, LOCTEXT("ModifiersFilterTooltip", "Modifiers")));
 
 	OutItemFilters.Add(MakeShared<FAvaOutlinerItemTypeFilter>(TEXT("Materials")
diff --git a/Engine/Plugins/Experimental/Avalanche/Source/AvalancheMask/AvalancheMask.Build.cs b/Engine/Plugins/Experimental/Avalanche/Source/AvalancheMask/AvalancheMask.Build.cs
--- a/Engine/Plugins/Experimental/Avalanche/Source/AvalancheMask/AvalancheMask.Build.cs
+++ b/Engine/Plugins/Experimental/Avalanche/Source/AvalancheMask/AvalancheMask.Build.cs
@@ -20,6 +20,7 @@ public class AvalancheMask : ModuleRules
 		PrivateDependencyModuleNames.AddRange(
 			new string[] {
 				"ActorModifierCore",
+				"ActorModifier",
 				"AssetRegistry",
 				"Avalanche",
 				"AvalancheCore",
diff --git a/Engine/Plugins/Experimental/Avalanche/Source/AvalancheMask/Public/Mask2D/AvaMask2DBaseModifier.h b/Engine/Plugins/Experimental/Avalanche/Source/AvalancheMask/Public/Mask2D/AvaMask2DBaseModifier.h
--- a/Engine/Plugins/Experimental/Avalanche/Source/AvalancheMask/Public/Mask2D/AvaMask2DBaseModifier.h
+++ b/Engine/Plugins/Experimental/Avalanche/Source/AvalancheMask/Public/Mask2D/AvaMask2DBaseModifier.h
@@ -5,7 +5,7 @@
 #include "AvaMaskTypes.h"
 #include "AvaPropertyChangeDispatcher.h"
 #include "GeometryMaskTypes.h"
-#include "Modifiers/AvaArrangeBaseModifier.h"
+#include "Modifiers/ActorModifierArrangeBaseModifier.h"
 #include "StructUtils/StructView.h"
 #include "UObject/ObjectPtr.h"
 
@@ -59,7 +59,7 @@ struct FAvaMask2DActorData
 /** Uses scene actors to create a mask texture and applies it to attached actors */
 UCLASS(Abstract)
 class AVALANCHEMASK_API UAvaMask2DBaseModifier
-	: public UAvaArrangeBaseModifier
+	: public UActorModifierArrangeBaseModifier
 {
 	GENERATED_BODY()
 
diff --git a/Engine/Plugins/Experimental/Avalanche/Source/AvalancheMask/Public/Mask2D/AvaMask2DReadModifier.h b/Engine/Plugins/Experimental/Avalanche/Source/AvalancheMask/Public/Mask2D/AvaMask2DReadModifier.h
--- a/Engine/Plugins/Experimental/Avalanche/Source/AvalancheMask/Public/Mask2D/AvaMask2DReadModifier.h
+++ b/Engine/Plugins/Experimental/Avalanche/Source/AvalancheMask/Public/Mask2D/AvaMask2DReadModifier.h
@@ -4,7 +4,7 @@
 
 #include "AvaMask2DBaseModifier.h"
 #include "AvaPropertyChangeDispatcher.h"
-#include "Modifiers/AvaArrangeBaseModifier.h"
+#include "Modifiers/ActorModifierArrangeBaseModifier.h"
 #include "UObject/ObjectPtr.h"
 #include "UObject/Package.h"
 #include "UObject/WeakObjectPtr.h"
diff --git a/Engine/Plugins/Experimental/Avalanche/Source/AvalancheMaskEditor/Private/AvaMaskEditorMode.cpp b/Engine/Plugins/Experimental/Avalanche/Source/AvalancheMaskEditor/Private/AvaMaskEditorMode.cpp
--- a/Engine/Plugins/Experimental/Avalanche/Source/AvalancheMaskEditor/Private/AvaMaskEditorMode.cpp
+++ b/Engine/Plugins/Experimental/Avalanche/Source/AvalancheMaskEditor/Private/AvaMaskEditorMode.cpp
@@ -408,7 +408,7 @@ UAvaMask2DBaseModifier* UAvaMaskEditorMode::FindOrAddMaskModifier(AActor* InActo
 			ModifierStack = ModifierSubsystem->AddActorModifierStack(InActor);
 		}
 
-		const FName MaskModifierName = GetDefault<UAvaBaseModifier>(InMaskModifierType)->GetModifierName();
+		const FName MaskModifierName = GetDefault<UActorModifierCoreBase>(InMaskModifierType)->GetModifierName();
 
 		TArray<UAvaMask2DBaseModifier*> FoundModifiers;
 		if (UActorModifierCoreBase* FoundModifier = ModifierStack->FindModifier(InMaskModifierType))
diff --git a/Engine/Plugins/Experimental/Avalanche/Source/AvalancheModifiers/AvalancheModifiers.Build.cs b/Engine/Plugins/Experimental/Avalanche/Source/AvalancheModifiers/AvalancheModifiers.Build.cs
--- a/Engine/Plugins/Experimental/Avalanche/Source/AvalancheModifiers/AvalancheModifiers.Build.cs
+++ b/Engine/Plugins/Experimental/Avalanche/Source/AvalancheModifiers/AvalancheModifiers.Build.cs
@@ -11,7 +11,9 @@ public class AvalancheModifiers : ModuleRules
 		PublicDependencyModuleNames.AddRange(
 			new string[]
 			{
+				"ActorModifier",
 				"ActorModifierCore",
+				"ActorModifierLayout",
 				"Avalanche",
 				"Core",
 				"CoreUObject",
diff --git a/Engine/Plugins/Experimental/Avalanche/Source/AvalancheModifiers/Private/AvaModifiersActorUtils.cpp b/Engine/Plugins/Experimental/Avalanche/Source/AvalancheModifiers/Private/AvaModifiersActorUtils.cpp
--- a/Engine/Plugins/Experimental/Avalanche/Source/AvalancheModifiers/Private/AvaModifiersActorUtils.cpp
+++ b/Engine/Plugins/Experimental/Avalanche/Source/AvalancheModifiers/Private/AvaModifiersActorUtils.cpp
@@ -2,250 +2,226 @@
 
 #include "AvaModifiersActorUtils.h"
 
+#include "AvaActorUtils.h"
+#include "AvaSceneTree.h"
 #include "Components/PrimitiveComponent.h"
+#include "Extensions/ActorModifierSceneTreeUpdateExtension.h"
 #include "GameFramework/Actor.h"
+#include "IAvaSceneInterface.h"
+#include "Utilities/ActorModifierActorUtils.h"
 
 #if WITH_EDITOR
+#include "AvaOutlinerSubsystem.h"
 #include "AvaOutlinerUtils.h"
 #include "IAvaOutliner.h"
 #endif
 
-FBox FAvaModifiersActorUtils::GetActorsBounds(const TSet<TWeakObjectPtr<AActor>>& InActors, const FTransform& InReferenceTransform, bool bInSkipHidden)
+bool FAvaModifiersActorUtils::IsActorNotIsolated(const AActor* InActor)
 {
-	FBox ActorsBounds = FBox(EForceInit::ForceInit);
-	ActorsBounds.IsValid = 0;
-
-	FVector OrientedVerts[8];
-
-	FTransform AccumulatedTransform = InReferenceTransform;
-
-	for (const TWeakObjectPtr<AActor>& ActorWeak : InActors)
+	if (!InActor)
 	{
-		const AActor* Actor = ActorWeak.Get();
-
-		if (!Actor)
-		{
-			continue;
-		}
-
-		if (bInSkipHidden && !IsActorVisible(Actor))
-		{
-			continue;
-		}
-
-		FBox ActorBounds = FAvaModifiersActorUtils::GetActorBounds(Actor);
-
-		if (ActorBounds.IsValid > 0)
-		{
-			FTransform ActorTransform = Actor->GetTransform();
-			ActorTransform.SetScale3D(FVector::OneVector);
-
-			const FOrientedBox OrientedBox = FAvaModifiersActorUtils::GetOrientedBox(ActorBounds, ActorTransform);
-
-			OrientedBox.CalcVertices(OrientedVerts);
-
-			for (const FVector& OrientedVert : OrientedVerts)
-			{
-				ActorsBounds += AccumulatedTransform.InverseTransformPositionNoScale(OrientedVert);
-			}
-
-			ActorsBounds.IsValid = 1;
-		}
+		return false;
 	}
 
-	AccumulatedTransform.SetScale3D(FVector::OneVector);
-
-	return ActorsBounds.TransformBy(AccumulatedTransform);
-}
+#if WITH_EDITOR
+	bool bIsIsolatingActors = false;
+	TArray<TWeakObjectPtr<const AActor>> IsolatedActors;
 
-FBox FAvaModifiersActorUtils::GetActorsBounds(AActor* InActor, bool bInIncludeChildren, bool bInSkipHidden)
-{
-	if (!InActor)
+	TSharedPtr<IAvaOutliner> AvaOutliner = FAvaOutlinerUtils::EditorGetOutliner(InActor->GetWorld());
+	if (AvaOutliner.IsValid())
 	{
-		return FBox();
+		bIsIsolatingActors = FAvaOutlinerUtils::EditorActorIsolationInfo(AvaOutliner, IsolatedActors);
 	}
 
-	TSet<TWeakObjectPtr<AActor>> AttachedModifyActors {InActor};
-
-	if (bInIncludeChildren)
+	if (bIsIsolatingActors && !IsolatedActors.Contains(InActor))
 	{
-		TArray<AActor*> AttachedActors;
-		InActor->GetAttachedActors(AttachedActors, false, true);
-		Algo::Transform(AttachedActors, AttachedModifyActors, [](AActor* InAttachedActor)
-		{
-			return InAttachedActor;
-		});
+		return true;
 	}
+#endif
 
-	return FAvaModifiersActorUtils::GetActorsBounds(AttachedModifyActors, InActor->GetActorTransform(), bInSkipHidden);
+	return false;
 }
 
-FBox FAvaModifiersActorUtils::GetActorBounds(const AActor* InActor)
+AActor* FAvaModifiersActorUtils::FindActorFromReferenceContainer(AActor* const InActor, const EActorModifierReferenceContainer InReferenceContainer, const bool bInIgnoreHiddenActors)
 {
-	FBox Box(ForceInit);
-	Box.IsValid = 0;
+	if (InReferenceContainer == EActorModifierReferenceContainer::Other || !IsValid(InActor))
+	{
+		return nullptr;
+	}
 
-	if (!InActor || !InActor->GetRootComponent())
+	// Note: Using Typed Outer World, instead of GetWorld(), since the typed outer World could be a streamed in world
+	// and GetWorld() only returns the main world.
+	UWorld* const World = InActor->GetTypedOuter<UWorld>();
+	if (!IsValid(World))
 	{
-		return Box;
+		return nullptr;
 	}
 
-	FTransform ActorToWorld = InActor->GetTransform();
-	ActorToWorld.SetScale3D(FVector::OneVector);
-	const FTransform WorldToActor = ActorToWorld.Inverse();
+	AActor* const ParentActor = InActor->GetAttachParentActor();
 
-	InActor->ForEachComponent<UPrimitiveComponent>(true, [&WorldToActor, &Box](const UPrimitiveComponent* InPrimitiveComponent)
-	{
-		if (!IsValid(InPrimitiveComponent))
-		{
-			return;
-		}
+	bool bIsOutlinerAttachedActors = false;
+	TArray<AActor*> AttachedActors;
 
 #if WITH_EDITOR
-		// Ignore Visualization Components, but don't consider them as failed components.
-		if (InPrimitiveComponent->IsVisualizationComponent())
+	UAvaOutlinerSubsystem* const OutlinerSubsystem = World->GetSubsystem<UAvaOutlinerSubsystem>();
+	if (IsValid(OutlinerSubsystem))
+	{
+		TSharedPtr<IAvaOutliner> AvaOutliner = OutlinerSubsystem->GetOutliner();
+		if (AvaOutliner.IsValid())
 		{
-			return;
+			bIsOutlinerAttachedActors = true;
+			if (IsValid(ParentActor))
+			{
+				AttachedActors = FAvaOutlinerUtils::EditorOutlinerChildActors(AvaOutliner, ParentActor);
+			}
+			else // no valid parent, use world actors instead
+			{
+				TArray<FAvaOutlinerItemPtr> OutlinerRootChildren = AvaOutliner->GetTreeRoot()->GetChildren();
+				const bool bValidRootActors = FAvaOutlinerUtils::EditorOutlinerItemsToActors(OutlinerRootChildren, AttachedActors);
+			}
 		}
-#endif
-
-		const FTransform ComponentToActor = InPrimitiveComponent->GetComponentTransform() * WorldToActor;
-		const FBox ComponentBox = InPrimitiveComponent->CalcBounds(ComponentToActor).GetBox();
-
-		Box += ComponentBox;
-		Box.IsValid = 1;
-	});
-
-	return Box;
-}
-
-FOrientedBox FAvaModifiersActorUtils::GetOrientedBox(const FBox& InLocalBox, const FTransform& InWorldTransform)
-{
-	FOrientedBox OutOrientedBox;
-	OutOrientedBox.Center = InWorldTransform.TransformPosition(InLocalBox.GetCenter());
-
-	OutOrientedBox.AxisX = InWorldTransform.TransformVector(FVector::UnitX());
-	OutOrientedBox.AxisY = InWorldTransform.TransformVector(FVector::UnitY());
-	OutOrientedBox.AxisZ = InWorldTransform.TransformVector(FVector::UnitZ());
-
-	OutOrientedBox.ExtentX = (InLocalBox.Max.X - InLocalBox.Min.X) / 2.f;
-	OutOrientedBox.ExtentY = (InLocalBox.Max.Y - InLocalBox.Min.Y) / 2.f;
-	OutOrientedBox.ExtentZ = (InLocalBox.Max.Z - InLocalBox.Min.Z) / 2.f;
-	return OutOrientedBox;
-}
-
-FVector FAvaModifiersActorUtils::GetVectorAxis(int32 InAxis)
-{
-	FVector FollowedAxisVector = FVector::ZeroVector;
-
-	if (EnumHasAnyFlags(static_cast<EAvaModifiersAxis>(InAxis), EAvaModifiersAxis::X))
-	{
-		FollowedAxisVector.X = 1;
 	}
+#endif
 
-	if (EnumHasAnyFlags(static_cast<EAvaModifiersAxis>(InAxis), EAvaModifiersAxis::Y))
+	if (!bIsOutlinerAttachedActors)
 	{
-		FollowedAxisVector.Y = 1;
-	}
+		if (const IAvaSceneInterface* SceneInterface = FAvaActorUtils::GetSceneInterfaceFromActor(InActor))
+		{
+			const FAvaSceneTree& SceneTree = SceneInterface->GetSceneTree();
 
-	if (EnumHasAnyFlags(static_cast<EAvaModifiersAxis>(InAxis), EAvaModifiersAxis::Z))
-	{
-		FollowedAxisVector.Z = 1;
+			if (!IsValid(ParentActor))
+			{
+				for (const int32 ChildId : SceneTree.GetRootNode().GetChildrenIndices())
+				{
+					const FAvaSceneItem* Item = SceneTree.GetItemAtIndex(ChildId);
+					AttachedActors.Add(Item->Resolve<AActor>(World));
+				}
+			}
+			else
+			{
+				AttachedActors = SceneTree.GetChildActors(ParentActor);
+			}
+		}
 	}
 
-	return FollowedAxisVector;
-}
-
-bool FAvaModifiersActorUtils::IsAxisVectorEquals(const FVector& InVectorA, const FVector& InVectorB, int32 InCompareAxis)
-{
-	// Compare InVectorA and InVectorB but only component of InCompareAxis
-	const FVector FollowedAxisVector = FAvaModifiersActorUtils::GetVectorAxis(InCompareAxis);
-	if ((InVectorA * FollowedAxisVector).Equals(InVectorB * FollowedAxisVector, 0.1))
+	if (AttachedActors.Num() == 0)
 	{
-		return true;
+		return nullptr;
 	}
 
-	return false;
-}
-
-bool FAvaModifiersActorUtils::IsActorNotIsolated(const AActor* InActor)
-{
-	if (!InActor)
+	if (bInIgnoreHiddenActors)
 	{
-		return false;
-	}
-
+		for (int32 ChildIndex = AttachedActors.Num() - 1; ChildIndex >= 0; --ChildIndex)
+		{
+			if (AttachedActors[ChildIndex]->IsHidden()
 #if WITH_EDITOR
-	bool bIsIsolatingActors = false;
-	TArray<TWeakObjectPtr<const AActor>> IsolatedActors;
+				|| AttachedActors[ChildIndex]->IsTemporarilyHiddenInEditor()
+#endif
+				)
+			{
+				AttachedActors.RemoveAt(ChildIndex);
+			}
+		}
+	}
 
-	TSharedPtr<IAvaOutliner> AvaOutliner = FAvaOutlinerUtils::EditorGetOutliner(InActor->GetWorld());
-	if (AvaOutliner.IsValid())
+	int32 FromActorIndex = -1;
+	if (!AttachedActors.Find(InActor, FromActorIndex))
 	{
-		bIsIsolatingActors = FAvaOutlinerUtils::EditorActorIsolationInfo(AvaOutliner, IsolatedActors);
+		return nullptr;
 	}
 
-	if (bIsIsolatingActors && !IsolatedActors.Contains(InActor))
+	switch (InReferenceContainer)
 	{
-		return true;
+		case EActorModifierReferenceContainer::Previous:
+		{
+			const int32 PreviousActorIndex = FromActorIndex - 1;
+			if (!AttachedActors.IsValidIndex(PreviousActorIndex))
+			{
+				return ParentActor;
+			}
+			return AttachedActors[PreviousActorIndex];
+		}
+		case EActorModifierReferenceContainer::Next:
+		{
+			const int32 NextActorIndex = FromActorIndex + 1;
+			if (!AttachedActors.IsValidIndex(NextActorIndex))
+			{
+				return nullptr;
+			}
+			return AttachedActors[NextActorIndex];
+		}
+		case EActorModifierReferenceContainer::First:
+		{
+			if (!AttachedActors.IsValidIndex(0))
+			{
+				return nullptr;
+			}
+			return AttachedActors[0] != InActor ? AttachedActors[0] : nullptr;
+		}
+		case EActorModifierReferenceContainer::Last:
+		{
+			AActor* const LastChildActor = AttachedActors.Last();
+			if (!IsValid(LastChildActor))
+			{
+				return nullptr;
+			}
+			return LastChildActor != InActor ? LastChildActor : nullptr;
+		}
+		default: ;
 	}
-#endif
 
-	return false;
+	return nullptr;
 }
 
-FRotator FAvaModifiersActorUtils::FindLookAtRotation(const FVector& InEyePosition, const FVector& InTargetPosition, EAvaModifiersAxis InAxis, bool bInFlipAxis)
+TArray<AActor*> FAvaModifiersActorUtils::GetReferenceActors(const FActorModifierSceneTreeActor* InTrackedActor)
 {
-	const FVector Direction = bInFlipAxis ? (InTargetPosition - InEyePosition).GetSafeNormal() : (InEyePosition - InTargetPosition).GetSafeNormal();
+	TArray<AActor*> ReferenceActors;
 
-	if (Direction.IsNearlyZero())
+	if (!InTrackedActor)
 	{
-		return FRotator::ZeroRotator;
+		return ReferenceActors;
 	}
 
-	const FRotator BaseRotation = Direction.Rotation();
-
-	FQuat AxisQuat;
-	switch (InAxis)
+	AActor* LocalActor = InTrackedActor->GetLocalActor();
+	if (!IsValid(LocalActor))
 	{
-	case EAvaModifiersAxis::X:
-		AxisQuat = FQuat::Identity;
-		break;
-
-	case EAvaModifiersAxis::Y:
-		AxisQuat = FQuat(FVector::ZAxisVector, FMath::DegreesToRadians(90.0f));
-		break;
-
-	case EAvaModifiersAxis::Z:
-		AxisQuat = FQuat(FVector::YAxisVector, FMath::DegreesToRadians(-90.0f));
-		break;
-
-	default:
-		return FRotator::ZeroRotator;
+		return ReferenceActors;
 	}
 
-	const FQuat FinalQuat = BaseRotation.Quaternion() * AxisQuat;
+	if (InTrackedActor->ReferenceContainer == EActorModifierReferenceContainer::Other)
+	{
+		if (AActor* ReferenceActor = InTrackedActor->ReferenceActorWeak.Get())
+		{
+			ReferenceActors.Add(ReferenceActor);
+		}
 
-	return FinalQuat.Rotator();
-}
+		return ReferenceActors;
+	}
 
-bool FAvaModifiersActorUtils::IsActorVisible(const AActor* InActor)
-{
-	if (!InActor)
+	AActor* ContextActor = LocalActor;
+	while (AActor* NewReferenceActor = FAvaModifiersActorUtils::FindActorFromReferenceContainer(ContextActor, InTrackedActor->ReferenceContainer, false))
 	{
-		return false;
-	}
+		// Take only siblings
+		if (NewReferenceActor->GetAttachParentActor() != LocalActor->GetAttachParentActor())
+		{
+			break;
+		}
 
-	const bool bIsActorVisible = !InActor->IsHidden();
+		ReferenceActors.Add(NewReferenceActor);
+		ContextActor = NewReferenceActor;
 
+		if (InTrackedActor->bSkipHiddenActors
+				&& (NewReferenceActor->IsHidden()
 #if WITH_EDITOR
-	const bool bIsActorVisibleInEditor = !InActor->IsTemporarilyHiddenInEditor();
-#else
-	constexpr bool bIsActorVisibleInEditor = true;
+				|| NewReferenceActor->IsTemporarilyHiddenInEditor()
 #endif
+				))
+		{
+			continue;
+		}
 
-	// If we don't have a root component, consider this true to skip it in final condition
-	const bool bIsRootComponentVisible = !InActor->GetRootComponent() || InActor->GetRootComponent()->IsVisible();
+		break;
+	}
 
-	return bIsActorVisible && bIsActorVisibleInEditor && bIsRootComponentVisible;
-}
+	return ReferenceActors;
+}
\ No newline at end of file
diff --git a/Engine/Plugins/Experimental/Avalanche/Source/AvalancheModifiers/Private/AvaModifiersActorUtils.h b/Engine/Plugins/Experimental/Avalanche/Source/AvalancheModifiers/Private/AvaModifiersActorUtils.h
--- a/Engine/Plugins/Experimental/Avalanche/Source/AvalancheModifiers/Private/AvaModifiersActorUtils.h
+++ b/Engine/Plugins/Experimental/Avalanche/Source/AvalancheModifiers/Private/AvaModifiersActorUtils.h
@@ -7,38 +7,18 @@
 #include "Math/OrientedBox.h"
 #include "Math/TransformNonVectorized.h"
 #include "UObject/WeakObjectPtrTemplates.h"
-#include "AvaModifiersActorUtils.generated.h"
 
-enum class EAvaAxis : uint8;
+enum class EActorModifierReferenceContainer : uint8;
+enum class EActorModifierAlignment : uint8;
 class AActor;
-
-/** Axis to use */
-UENUM(BlueprintType, meta=(Bitflags, UseEnumValuesAsMaskValuesInEditor="true"))
-enum class EAvaModifiersAxis : uint8
-{
-	None = 0 UMETA(Hidden),
-	X = 1 << 0,
-	Y = 1 << 1,
-	Z = 1 << 2,
-};
-ENUM_CLASS_FLAGS(EAvaModifiersAxis);
+struct FActorModifierSceneTreeActor;
 
 // All operations that can be reused or shared in modifiers should go here
 struct FAvaModifiersActorUtils
 {
-	/** Begin Bounds */
-	static FBox GetActorsBounds(const TSet<TWeakObjectPtr<AActor>>& InActors, const FTransform& InReferenceTransform, bool bInSkipHidden = false);
-	static FBox GetActorsBounds(AActor* InActor, bool bInIncludeChildren, bool bInSkipHidden = false);
-	static FBox GetActorBounds(const AActor* InActor);
-	static FOrientedBox GetOrientedBox(const FBox& InLocalBox, const FTransform& InWorldTransform);
-	static FVector GetVectorAxis(int32 InAxis);
-	static bool IsAxisVectorEquals(const FVector& InVectorA, const FVector& InVectorB, int32 InCompareAxis);
-	/** End Bounds */
-
 	/** Begin Outliner */
 	static bool IsActorNotIsolated(const AActor* InActor);
+	static AActor* FindActorFromReferenceContainer(AActor* const InActor, const EActorModifierReferenceContainer InReferenceContainer, const bool bInIgnoreHiddenActors);
+	static TArray<AActor*> GetReferenceActors(const FActorModifierSceneTreeActor* InTrackedActor);
 	/** End Outliner */
-
-	static FRotator FindLookAtRotation(const FVector& InEyePosition, const FVector& InTargetPosition, EAvaModifiersAxis InAxis, bool bInFlipAxis);
-	static bool IsActorVisible(const AActor* InActor);
 };
\ No newline at end of file
diff --git a/Engine/Plugins/Experimental/Avalanche/Source/AvalancheModifiers/Private/AvaModifiersModule.cpp b/Engine/Plugins/Experimental/Avalanche/Source/AvalancheModifiers/Private/AvaModifiersModule.cpp
--- a/Engine/Plugins/Experimental/Avalanche/Source/AvalancheModifiers/Private/AvaModifiersModule.cpp
+++ b/Engine/Plugins/Experimental/Avalanche/Source/AvalancheModifiers/Private/AvaModifiersModule.cpp
@@ -1,5 +1,31 @@
 ﻿// Copyright Epic Games, Inc. All Rights Reserved.
 
+#include "AvaModifiersModule.h"
+
+#include "AvaModifiersSceneTreeResolver.h"
+#include "Engine/Level.h"
+#include "Extensions/ActorModifierSceneTreeUpdateExtension.h"
+#include "GameFramework/Actor.h"
 #include "Modules/ModuleManager.h"
 
-IMPLEMENT_MODULE(FDefaultModuleImpl, AvalancheModifiers)
+IMPLEMENT_MODULE(FAvalancheModifiersModule, AvalancheModifiers)
+
+void FAvalancheModifiersModule::StartupModule()
+{
+	FActorModifierSceneTreeUpdateExtension::OnGetSceneTreeResolver().BindRaw(this, &FAvalancheModifiersModule::GetSceneTreeResolver);
+}
+
+void FAvalancheModifiersModule::ShutdownModule()
+{
+	FActorModifierSceneTreeUpdateExtension::OnGetSceneTreeResolver().Unbind();
+}
+
+TSharedPtr<IActorModifierSceneTreeCustomResolver> FAvalancheModifiersModule::GetSceneTreeResolver(ULevel* InLevel)
+{
+	if (IsValid(InLevel))
+	{
+		return MakeShared<FAvaModifiersSceneTreeResolver>(InLevel);
+	}
+
+	return nullptr;
+}
diff --git a/Engine/Plugins/Experimental/Avalanche/Source/AvalancheModifiers/Private/Modifiers/AvaAutoSizeModifier.cpp b/Engine/Plugins/Experimental/Avalanche/Source/AvalancheModifiers/Private/Modifiers/AvaAutoSizeModifier.cpp
--- a/Engine/Plugins/Experimental/Avalanche/Source/AvalancheModifiers/Private/Modifiers/AvaAutoSizeModifier.cpp
+++ b/Engine/Plugins/Experimental/Avalanche/Source/AvalancheModifiers/Private/Modifiers/AvaAutoSizeModifier.cpp
@@ -8,7 +8,8 @@
 #include "DynamicMesh/MeshTransforms.h"
 #include "AvaShapeActor.h"
 #include "DynamicMeshes/AvaShape2DDynMeshBase.h"
-#include "Shared/AvaTransformModifierShared.h"
+#include "Shared/ActorModifierTransformShared.h"
+#include "Utilities/ActorModifierActorUtils.h"
 
 #define LOCTEXT_NAMESPACE "AvaAutoSizeModifier"
 
@@ -22,7 +23,7 @@ bool UAvaAutoSizeModifier::IsModifierDirtyable() const
 		return Super::IsModifierDirtyable();
 	}
 
-	const FBox ReferenceActorLocalBounds = FAvaModifiersActorUtils::GetActorsBounds(TrackedActor, bIncludeChildren, true);
+	const FBox ReferenceActorLocalBounds = UE::ActorModifier::ActorUtils::GetActorsBounds(TrackedActor, bIncludeChildren, true);
 
 	if (ReferenceActorLocalBounds.GetSize().Equals(CachedReferenceBounds.GetSize(), 0.01))
 	{
@@ -69,11 +70,11 @@ void UAvaAutoSizeModifier::OnModifierAdded(EActorModifierCoreEnableReason InReas
 {
 	Super::OnModifierAdded(InReason);
 
-	AddExtension<FAvaTransformUpdateModifierExtension>(this);
-	AddExtension<FAvaRenderStateUpdateModifierExtension>(this);
-	AddExtension<FAvaSceneTreeUpdateModifierExtension>(this);
+	AddExtension<FActorModifierTransformUpdateExtension>(this);
+	AddExtension<FActorModifierRenderStateUpdateExtension>(this);
+	AddExtension<FActorModifierSceneTreeUpdateExtension>(this);
 
-	if (FAvaSceneTreeUpdateModifierExtension* SceneExtension = GetExtension<FAvaSceneTreeUpdateModifierExtension>())
+	if (FActorModifierSceneTreeUpdateExtension* SceneExtension = GetExtension<FActorModifierSceneTreeUpdateExtension>())
 	{
 		SceneExtension->TrackSceneTree(0, &ReferenceActor);
 	}
@@ -86,7 +87,7 @@ void UAvaAutoSizeModifier::OnModifierEnabled(EActorModifierCoreEnableReason InRe
 	Super::OnModifierEnabled(InReason);
 
 	// Save actor layout state
-	if (UAvaTransformModifierShared* LayoutShared = GetShared<UAvaTransformModifierShared>(true))
+	if (UActorModifierTransformShared* LayoutShared = GetShared<UActorModifierTransformShared>(true))
 	{
 		LayoutShared->SaveActorState(this, GetModifiedActor());
 	}
@@ -109,7 +110,7 @@ void UAvaAutoSizeModifier::OnModifierDisabled(EActorModifierCoreDisableReason In
 	Super::OnModifierDisabled(InReason);
 
 	// Restore actor layout state
-	if (UAvaTransformModifierShared* LayoutShared = GetShared<UAvaTransformModifierShared>(false))
+	if (UActorModifierTransformShared* LayoutShared = GetShared<UActorModifierTransformShared>(false))
 	{
 		LayoutShared->RestoreActorState(this, GetModifiedActor());
 	}
@@ -138,14 +139,14 @@ void UAvaAutoSizeModifier::Apply()
 		return;
 	}
 
-	if (FAvaRenderStateUpdateModifierExtension* RenderExtension = GetExtension<FAvaRenderStateUpdateModifierExtension>())
+	if (FActorModifierRenderStateUpdateExtension* RenderExtension = GetExtension<FActorModifierRenderStateUpdateExtension>())
 	{
 		RenderExtension->SetTrackedActorVisibility(TrackedActor, bIncludeChildren);
 	}
 
-	CachedReferenceBounds = FAvaModifiersActorUtils::GetActorsBounds(TrackedActor, bIncludeChildren, true);
+	CachedReferenceBounds = UE::ActorModifier::ActorUtils::GetActorsBounds(TrackedActor, bIncludeChildren, true);
 
-	const FBox ModifiedActorBounds = FAvaModifiersActorUtils::GetActorsBounds(CurrentlyModifiedActor, false);
+	const FBox ModifiedActorBounds = UE::ActorModifier::ActorUtils::GetActorsBounds(CurrentlyModifiedActor, false);
 
 	FVector ModifiedActorBoundsOrigin;
 	FVector ModifiedActorBoundsExtent;
@@ -233,7 +234,7 @@ void UAvaAutoSizeModifier::Apply()
 void UAvaAutoSizeModifier::PostLoad()
 {
 	if (!bDeprecatedPropertiesMigrated
-		&& ReferenceActor.ReferenceContainer == EAvaReferenceContainer::Other
+		&& ReferenceActor.ReferenceContainer == EActorModifierReferenceContainer::Other
 		&& ReferenceActor.ReferenceActorWeak == nullptr)
 	{
 		ReferenceActor.ReferenceContainer = ReferenceContainer_DEPRECATED;
@@ -282,7 +283,7 @@ void UAvaAutoSizeModifier::PostEditChangeProperty(FPropertyChangedEvent& Propert
 }
 #endif
 
-void UAvaAutoSizeModifier::SetReferenceActor(const FAvaSceneTreeActor& InReferenceActor)
+void UAvaAutoSizeModifier::SetReferenceActor(const FActorModifierSceneTreeActor& InReferenceActor)
 {
 	if (ReferenceActor == InReferenceActor)
 	{
@@ -390,7 +391,7 @@ void UAvaAutoSizeModifier::OnSceneTreeTrackedActorChanged(int32 InIdx, AActor* I
 	}
 
 	// Untrack reference actor
-	if (FAvaTransformUpdateModifierExtension* TransformExtension = GetExtension<FAvaTransformUpdateModifierExtension>())
+	if (FActorModifierTransformUpdateExtension* TransformExtension = GetExtension<FActorModifierTransformUpdateExtension>())
 	{
 		TransformExtension->UntrackActor(InPreviousActor);
 		TransformExtension->TrackActor(InNewActor, true);
@@ -419,7 +420,7 @@ void UAvaAutoSizeModifier::OnReferenceActorChanged()
 		ReferenceActor.ReferenceActorWeak = nullptr;
 	}
 
-	if (const FAvaSceneTreeUpdateModifierExtension* SceneExtension = GetExtension<FAvaSceneTreeUpdateModifierExtension>())
+	if (const FActorModifierSceneTreeUpdateExtension* SceneExtension = GetExtension<FActorModifierSceneTreeUpdateExtension>())
 	{
 		SceneExtension->CheckTrackedActorUpdate(0);
 	}
diff --git a/Engine/Plugins/Experimental/Avalanche/Source/AvalancheModifiers/Private/Modifiers/AvaBooleanModifier.cpp b/Engine/Plugins/Experimental/Avalanche/Source/AvalancheModifiers/Private/Modifiers/AvaBooleanModifier.cpp
--- a/Engine/Plugins/Experimental/Avalanche/Source/AvalancheModifiers/Private/Modifiers/AvaBooleanModifier.cpp
+++ b/Engine/Plugins/Experimental/Avalanche/Source/AvalancheModifiers/Private/Modifiers/AvaBooleanModifier.cpp
@@ -7,7 +7,6 @@
 #include "Containers/Ticker.h"
 #include "DynamicMesh/MeshTransforms.h"
 #include "DynamicMeshes/AvaShapeDynMeshBase.h"
-#include "Extensions/AvaTransformUpdateModifierExtension.h"
 #include "GeometryScript/MeshBasicEditFunctions.h"
 #include "Materials/MaterialInstanceDynamic.h"
 #include "Modifiers/ActorModifierCoreStack.h"
@@ -16,7 +15,7 @@
 #include "Operations/OffsetMeshRegion.h"
 #include "Profilers/AvaBooleanModifierProfiler.h"
 #include "Shared/AvaBooleanModifierShared.h"
-#include "Shared/AvaVisibilityModifierShared.h"
+#include "Shared/ActorModifierVisibilityShared.h"
 
 #define LOCTEXT_NAMESPACE "AvaBooleanModifier"
 
@@ -50,7 +49,7 @@ void UAvaBooleanModifier::OnModifierAdded(EActorModifierCoreEnableReason InReaso
 {
 	Super::OnModifierAdded(InReason);
 
-	if (FAvaTransformUpdateModifierExtension* TransformExtension = AddExtension<FAvaTransformUpdateModifierExtension>(this))
+	if (FActorModifierTransformUpdateExtension* TransformExtension = AddExtension<FActorModifierTransformUpdateExtension>(this))
 	{
 		TransformExtension->TrackActor(GetModifiedActor(), true);
 	}
@@ -576,7 +575,7 @@ void UAvaBooleanModifier::UpdateMaskVisibility()
 {
 	if (AActor* ActorModified = GetModifiedActor())
 	{
-		UAvaVisibilityModifierShared* VisibilityShared = GetShared<UAvaVisibilityModifierShared>(true);
+		UActorModifierVisibilityShared* VisibilityShared = GetShared<UActorModifierVisibilityShared>(true);
 
 		if (Mode != EAvaBooleanMode::None)
 		{
@@ -584,7 +583,7 @@ void UAvaBooleanModifier::UpdateMaskVisibility()
 			if (IsModifierEnabled())
 			{
 				VisibilityShared->SaveActorState(this, ActorModified);
-				VisibilityShared->SetActorVisibility(this, ActorModified, true, false, EAvaVisibilityActor::Game);
+				VisibilityShared->SetActorVisibility(this, ActorModified, true, false, EActorModifierVisibilityActor::Game);
 				return;
 			}
 		}
diff --git a/Engine/Plugins/Experimental/Avalanche/Source/AvalancheModifiers/Private/Modifiers/AvaDynamicMeshConverterModifier.cpp b/Engine/Plugins/Experimental/Avalanche/Source/AvalancheModifiers/Private/Modifiers/AvaDynamicMeshConverterModifier.cpp
--- a/Engine/Plugins/Experimental/Avalanche/Source/AvalancheModifiers/Private/Modifiers/AvaDynamicMeshConverterModifier.cpp
+++ b/Engine/Plugins/Experimental/Avalanche/Source/AvalancheModifiers/Private/Modifiers/AvaDynamicMeshConverterModifier.cpp
@@ -6,7 +6,7 @@
 #include "Components/DynamicMeshComponent.h"
 #include "DynamicMesh/DynamicMesh3.h"
 #include "Engine/StaticMesh.h"
-#include "Extensions/AvaSceneTreeUpdateModifierExtension.h"
+#include "Extensions/ActorModifierSceneTreeUpdateExtension.h"
 #include "GeometryScript/MeshAssetFunctions.h"
 #include "GeometryScript/MeshBasicEditFunctions.h"
 
@@ -63,11 +63,11 @@ void UAvaDynamicMeshConverterModifier::OnModifierAdded(EActorModifierCoreEnableR
 
 	AddDynamicMeshComponent();
 
-	AddExtension<FAvaRenderStateUpdateModifierExtension>(this);
+	AddExtension<FActorModifierRenderStateUpdateExtension>(this);
 
-	if (FAvaSceneTreeUpdateModifierExtension* SceneExtension = AddExtension<FAvaSceneTreeUpdateModifierExtension>(this))
+	if (FActorModifierSceneTreeUpdateExtension* SceneExtension = AddExtension<FActorModifierSceneTreeUpdateExtension>(this))
 	{
-		TrackedActor.ReferenceContainer = EAvaReferenceContainer::Other;
+		TrackedActor.ReferenceContainer = EActorModifierReferenceContainer::Other;
 		TrackedActor.ReferenceActorWeak = SourceActorWeak.Get();
 		TrackedActor.bSkipHiddenActors = false;
 		SceneExtension->TrackSceneTree(0, &TrackedActor);
@@ -487,7 +487,7 @@ void UAvaDynamicMeshConverterModifier::OnSourceActorChanged()
 
 	bHideConvertedMesh = SourceActor == ActorModified || SourceActor->IsAttachedTo(ActorModified);
 
-	if (const FAvaSceneTreeUpdateModifierExtension* SceneExtension = GetExtension<FAvaSceneTreeUpdateModifierExtension>())
+	if (const FActorModifierSceneTreeUpdateExtension* SceneExtension = GetExtension<FActorModifierSceneTreeUpdateExtension>())
 	{
 		TrackedActor.ReferenceActorWeak = SourceActor;
 		SceneExtension->CheckTrackedActorUpdate(0);
diff --git a/Engine/Plugins/Experimental/Avalanche/Source/AvalancheModifiers/Private/Modifiers/AvaHideEmptyModifier.cpp b/Engine/Plugins/Experimental/Avalanche/Source/AvalancheModifiers/Private/Modifiers/AvaHideEmptyModifier.cpp
--- a/Engine/Plugins/Experimental/Avalanche/Source/AvalancheModifiers/Private/Modifiers/AvaHideEmptyModifier.cpp
+++ b/Engine/Plugins/Experimental/Avalanche/Source/AvalancheModifiers/Private/Modifiers/AvaHideEmptyModifier.cpp
@@ -2,7 +2,7 @@
 
 #include "Modifiers/AvaHideEmptyModifier.h"
 
-#include "Shared/AvaVisibilityModifierShared.h"
+#include "Shared/ActorModifierVisibilityShared.h"
 #include "Text3DComponent.h"
 
 #define LOCTEXT_NAMESPACE "AvaHideEmptyModifier"
@@ -31,7 +31,7 @@ void UAvaHideEmptyModifier::OnModifierAdded(EActorModifierCoreEnableReason InRea
 	{
 		if (AActor* ActorModified = GetModifiedActor())
 		{
-			ReferenceActor.ReferenceContainer = EAvaReferenceContainer::Other;
+			ReferenceActor.ReferenceContainer = EActorModifierReferenceContainer::Other;
 			ReferenceActor.ReferenceActorWeak = ActorModified;
 			ReferenceActor.bSkipHiddenActors = false;
 
@@ -71,7 +71,7 @@ void UAvaHideEmptyModifier::Apply()
 		return;
 	}
 
-	UAvaVisibilityModifierShared* VisibilityShared = GetShared<UAvaVisibilityModifierShared>(true);
+	UActorModifierVisibilityShared* VisibilityShared = GetShared<UActorModifierVisibilityShared>(true);
 
 	if (Text3DComponent->GetText().IsEmpty())
 	{
@@ -162,11 +162,11 @@ void UAvaHideEmptyModifier::OnTextChanged()
 
 void UAvaHideEmptyModifier::OnContainerActorChanged()
 {
-	ReferenceActor.ReferenceContainer = EAvaReferenceContainer::Other;
+	ReferenceActor.ReferenceContainer = EActorModifierReferenceContainer::Other;
 	ReferenceActor.ReferenceActorWeak = ContainerActorWeak;
 	ReferenceActor.bSkipHiddenActors = false;
 
-	if (const FAvaSceneTreeUpdateModifierExtension* SceneExtension = GetExtension<FAvaSceneTreeUpdateModifierExtension>())
+	if (const FActorModifierSceneTreeUpdateExtension* SceneExtension = GetExtension<FActorModifierSceneTreeUpdateExtension>())
 	{
 		SceneExtension->CheckTrackedActorUpdate(0);
 	}
@@ -181,7 +181,7 @@ void UAvaHideEmptyModifier::OnSceneTreeTrackedActorChanged(int32 InIdx, AActor*
 {
 	Super::OnSceneTreeTrackedActorChanged(InIdx, InPreviousActor, InNewActor);
 
-	if (UAvaVisibilityModifierShared* VisibilityShared = GetShared<UAvaVisibilityModifierShared>(false))
+	if (UActorModifierVisibilityShared* VisibilityShared = GetShared<UActorModifierVisibilityShared>(false))
 	{
 		VisibilityShared->RestoreActorState(this, InPreviousActor);
 	}
diff --git a/Engine/Plugins/Experimental/Avalanche/Source/AvalancheModifiers/Private/Modifiers/AvaSplineSweepModifier.cpp b/Engine/Plugins/Experimental/Avalanche/Source/AvalancheModifiers/Private/Modifiers/AvaSplineSweepModifier.cpp
--- a/Engine/Plugins/Experimental/Avalanche/Source/AvalancheModifiers/Private/Modifiers/AvaSplineSweepModifier.cpp
+++ b/Engine/Plugins/Experimental/Avalanche/Source/AvalancheModifiers/Private/Modifiers/AvaSplineSweepModifier.cpp
@@ -6,7 +6,6 @@
 #include "Components/DynamicMeshComponent.h"
 #include "Components/SplineComponent.h"
 #include "DynamicMesh/Operations/MergeCoincidentMeshEdges.h"
-#include "Extensions/AvaRenderStateUpdateModifierExtension.h"
 #include "GeometryScript/MeshPrimitiveFunctions.h"
 #include "GeometryScript/PolyPathFunctions.h"
 #include "Operations/InsetMeshRegion.h"
@@ -234,7 +233,7 @@ void UAvaSplineSweepModifier::OnModifierAdded(EActorModifierCoreEnableReason InR
 {
 	Super::OnModifierAdded(InReason);
 
-	AddExtension<FAvaRenderStateUpdateModifierExtension>(this);
+	AddExtension<FActorModifierRenderStateUpdateExtension>(this);
 }
 
 void UAvaSplineSweepModifier::Apply()
diff --git a/Engine/Plugins/Experimental/Avalanche/Source/AvalancheModifiers/Private/Modifiers/AvaTranslucentPriorityModifier.cpp b/Engine/Plugins/Experimental/Avalanche/Source/AvalancheModifiers/Private/Modifiers/AvaTranslucentPriorityModifier.cpp
--- a/Engine/Plugins/Experimental/Avalanche/Source/AvalancheModifiers/Private/Modifiers/AvaTranslucentPriorityModifier.cpp
+++ b/Engine/Plugins/Experimental/Avalanche/Source/AvalancheModifiers/Private/Modifiers/AvaTranslucentPriorityModifier.cpp
@@ -121,7 +121,7 @@ void UAvaTranslucentPriorityModifier::SavePreState()
 
 	SharedObject->SetComponentsState(this, PrimitiveComponentsWeak);
 
-	if (FAvaTransformUpdateModifierExtension* TransformExtension = GetExtension<FAvaTransformUpdateModifierExtension>())
+	if (FActorModifierTransformUpdateExtension* TransformExtension = GetExtension<FActorModifierTransformUpdateExtension>())
 	{
 		TransformExtension->TrackActors(ChildrenActorsWeak, true);
 	}
@@ -131,7 +131,7 @@ void UAvaTranslucentPriorityModifier::RestorePreState()
 {
 	Super::RestorePreState();
 
-	if (FAvaTransformUpdateModifierExtension* TransformExtension = GetExtension<FAvaTransformUpdateModifierExtension>())
+	if (FActorModifierTransformUpdateExtension* TransformExtension = GetExtension<FActorModifierTransformUpdateExtension>())
 	{
 		TransformExtension->UntrackActors(ChildrenActorsWeak);
 	}
@@ -252,7 +252,7 @@ void UAvaTranslucentPriorityModifier::OnModifierDisabled(EActorModifierCoreDisab
 		SharedObject->RestoreComponentsState(this, true);
 	}
 
-	if (FAvaTransformUpdateModifierExtension* TransformExtension = GetExtension<FAvaTransformUpdateModifierExtension>())
+	if (FActorModifierTransformUpdateExtension* TransformExtension = GetExtension<FActorModifierTransformUpdateExtension>())
 	{
 		TransformExtension->UntrackActors(ChildrenActorsWeak);
 	}
diff --git a/Engine/Plugins/Experimental/Avalanche/Source/AvalancheModifiers/Private/Modifiers/AvaVisibilityModifier.cpp b/Engine/Plugins/Experimental/Avalanche/Source/AvalancheModifiers/Private/Modifiers/AvaVisibilityModifier.cpp
--- a/Engine/Plugins/Experimental/Avalanche/Source/AvalancheModifiers/Private/Modifiers/AvaVisibilityModifier.cpp
+++ b/Engine/Plugins/Experimental/Avalanche/Source/AvalancheModifiers/Private/Modifiers/AvaVisibilityModifier.cpp
@@ -5,7 +5,7 @@
 #include "AvaModifiersActorUtils.h"
 #include "GameFramework/Actor.h"
 #include "Misc/ITransaction.h"
-#include "Shared/AvaVisibilityModifierShared.h"
+#include "Shared/ActorModifierVisibilityShared.h"
 
 #define LOCTEXT_NAMESPACE "AvaVisibilityModifier"
 
@@ -20,7 +20,7 @@ void UAvaVisibilityModifier::Apply()
 		return;
 	}
 
-	const FAvaSceneTreeUpdateModifierExtension* SceneExtension = GetExtension<FAvaSceneTreeUpdateModifierExtension>();
+	const FActorModifierSceneTreeUpdateExtension* SceneExtension = GetExtension<FActorModifierSceneTreeUpdateExtension>();
 	if (!SceneExtension)
 	{
 		Fail(LOCTEXT("InvalidSceneExtension", "Scene extension could not be found"));
@@ -28,7 +28,7 @@ void UAvaVisibilityModifier::Apply()
 	}
 
 	const TArray<TWeakObjectPtr<AActor>> AttachedActors = SceneExtension->GetDirectChildrenActor(ModifyActor);
-	UAvaVisibilityModifierShared* VisibilityShared = GetShared<UAvaVisibilityModifierShared>(true);
+	UActorModifierVisibilityShared* VisibilityShared = GetShared<UActorModifierVisibilityShared>(true);
 
 #if WITH_EDITOR
 	const bool bHiddenInEditor = ModifyActor->IsTemporarilyHiddenInEditor(/** IncludeParent */false);
@@ -113,7 +113,7 @@ void UAvaVisibilityModifier::Apply()
 
 	ChildrenActorsWeak = NewChildrenActorsWeak;
 
-	FAvaRenderStateUpdateModifierExtension* RenderStateExtension = GetExtension<FAvaRenderStateUpdateModifierExtension>();
+	FActorModifierRenderStateUpdateExtension* RenderStateExtension = GetExtension<FActorModifierRenderStateUpdateExtension>();
 	RenderStateExtension->SetTrackedActorsVisibility(ChildrenActorsWeak);
 
 	Next();
@@ -259,7 +259,7 @@ bool UAvaVisibilityModifier::IsChildActorHidden(AActor* InActor) const
 
 UAvaVisibilityModifier* UAvaVisibilityModifier::GetFirstModifierAbove(AActor* InActor)
 {
-	UAvaVisibilityModifierShared* VisibilityShared = GetShared<UAvaVisibilityModifierShared>(false);
+	UActorModifierVisibilityShared* VisibilityShared = GetShared<UActorModifierVisibilityShared>(false);
 	UAvaVisibilityModifier* FirstModifierAbove = nullptr;
 
 	if (!InActor || !VisibilityShared)
@@ -267,11 +267,11 @@ UAvaVisibilityModifier* UAvaVisibilityModifier::GetFirstModifierAbove(AActor* In
 		return FirstModifierAbove;
 	}
 
-	if (FAvaVisibilitySharedActorState* ActorState = VisibilityShared->FindActorState(InActor))
+	if (FActorModifierVisibilitySharedActorState* ActorState = VisibilityShared->FindActorState(InActor))
 	{
-		for (const FAvaVisibilitySharedModifierState& ModifierState : ActorState->ModifierStates)
+		for (const FActorModifierVisibilitySharedModifierState& ModifierState : ActorState->ModifierStates)
 		{
-			UAvaBaseModifier* Modifier = ModifierState.ModifierWeak.Get();
+			UActorModifierCoreBase* Modifier = ModifierState.ModifierWeak.Get();
 			if (!Modifier || Modifier->GetModifiedActor() != InActor->GetAttachParentActor())
 			{
 				continue;
diff --git a/Engine/Plugins/Experimental/Avalanche/Source/AvalancheModifiers/Private/Tests/AvaModifiersAlignBetween.spec.cpp b/Engine/Plugins/Experimental/Avalanche/Source/AvalancheModifiers/Private/Tests/AvaModifiersAlignBetween.spec.cpp
--- a/Engine/Plugins/Experimental/Avalanche/Source/AvalancheModifiers/Private/Tests/AvaModifiersAlignBetween.spec.cpp
+++ b/Engine/Plugins/Experimental/Avalanche/Source/AvalancheModifiers/Private/Tests/AvaModifiersAlignBetween.spec.cpp
@@ -2,7 +2,7 @@
 
 #include "Misc/AutomationTest.h"
 #include "Modifiers/ActorModifierCoreStack.h"
-#include "Modifiers/AvaAlignBetweenModifier.h"
+#include "Modifiers/ActorModifierAlignBetweenModifier.h"
 #include "Subsystems/ActorModifierCoreSubsystem.h"
 #include "Tests/Framework/AvaModifiersTestUtils.h"
 #include "Tests/Framework/AvaTestDynamicMeshActor.h"
@@ -23,7 +23,7 @@ BEGIN_DEFINE_SPEC(AvalancheModifiersAlignBetween, "Avalanche.Modifiers.AlignBetw
 	FVector InitialReferenceTwoMeshActorLocation;
 	FVector ExpectedModifiedActorLocation;
 
-	UAvaAlignBetweenModifier* AlignBetweenModifier;
+	UActorModifierAlignBetweenModifier* AlignBetweenModifier;
 
 	TSharedPtr<FAvaTestUtils> TestUtils = MakeShared<FAvaTestUtils>();
 	TSharedPtr<FAvaModifierTestUtils> ModifierTestUtils = MakeShared<FAvaModifierTestUtils>(TestUtils);
@@ -51,20 +51,20 @@ void AvalancheModifiersAlignBetween::Define()
 
 		// Set up modifier
 		const UActorModifierCoreSubsystem* ModifierSubsystem = UActorModifierCoreSubsystem::Get();
-		const FName ModifierName = ModifierTestUtils->GetModifierName(UAvaAlignBetweenModifier::StaticClass());
+		const FName ModifierName = ModifierTestUtils->GetModifierName(UActorModifierAlignBetweenModifier::StaticClass());
 		FActorModifierCoreStackInsertOp InsertOp = ModifierTestUtils->GenerateInsertOp(ModifierName);
 		UActorModifierCoreStack* ModifierStack = ModifierTestUtils->GenerateModifierStackForActor(ModifiedMeshActor);
-		AlignBetweenModifier = Cast<UAvaAlignBetweenModifier>(ModifierSubsystem->InsertModifier(ModifierStack, InsertOp));
+		AlignBetweenModifier = Cast<UActorModifierAlignBetweenModifier>(ModifierSubsystem->InsertModifier(ModifierStack, InsertOp));
 		check(AlignBetweenModifier);
 
 		// Set up the modifier
-		FAvaAlignBetweenWeightedActor ReferenceOneActor = FAvaAlignBetweenWeightedActor(ReferenceOneMeshActor);
+		FActorModifierAlignBetweenWeightedActor ReferenceOneActor = FActorModifierAlignBetweenWeightedActor(ReferenceOneMeshActor);
 		ReferenceOneActor.bEnabled = true;
 		ReferenceOneActor.Weight = ReferenceOneWeight;
-		FAvaAlignBetweenWeightedActor ReferenceTwoActor = FAvaAlignBetweenWeightedActor(ReferenceTwoMeshActor);
+		FActorModifierAlignBetweenWeightedActor ReferenceTwoActor = FActorModifierAlignBetweenWeightedActor(ReferenceTwoMeshActor);
 		ReferenceTwoActor.bEnabled = true;
 		ReferenceTwoActor.Weight = ReferenceTwoWeight;
-		TSet<FAvaAlignBetweenWeightedActor> ReferenceActors = {
+		TSet<FActorModifierAlignBetweenWeightedActor> ReferenceActors = {
 			ReferenceOneActor,
 			ReferenceTwoActor
 		};
diff --git a/Engine/Plugins/Experimental/Avalanche/Source/AvalancheModifiers/Private/Tests/AvaModifiersAutoFollow.spec.cpp b/Engine/Plugins/Experimental/Avalanche/Source/AvalancheModifiers/Private/Tests/AvaModifiersAutoFollow.spec.cpp
--- a/Engine/Plugins/Experimental/Avalanche/Source/AvalancheModifiers/Private/Tests/AvaModifiersAutoFollow.spec.cpp
+++ b/Engine/Plugins/Experimental/Avalanche/Source/AvalancheModifiers/Private/Tests/AvaModifiersAutoFollow.spec.cpp
@@ -3,7 +3,7 @@
 #include "Misc/AutomationTest.h"
 #include "Modifiers/ActorModifierCoreBase.h"
 #include "Modifiers/ActorModifierCoreStack.h"
-#include "Modifiers/AvaAutoFollowModifier.h"
+#include "Modifiers/ActorModifierAutoFollowModifier.h"
 #include "Subsystems/ActorModifierCoreSubsystem.h"
 #include "Tests/Framework/AvaModifiersTestUtils.h"
 #include "Tests/Framework/AvaTestDynamicMeshActor.h"
@@ -18,7 +18,7 @@ BEGIN_DEFINE_SPEC(AvalancheModifiersAutoFollowTests, "Avalanche.Modifiers.Autofo
 	AAvaTestDynamicMeshActor* DynamicMeshActorToFollow;
 	UActorModifierCoreSubsystem* ModifierSubsystem;
 	UActorModifierCoreStack* ActorModifierStack;
-	UAvaAutoFollowModifier* AutoFollowModifier;
+	UActorModifierAutoFollowModifier* AutoFollowModifier;
 	FTransform InitialActorState;
 
 	TSharedPtr<FAvaTestUtils> TestUtils = MakeShared<FAvaTestUtils>();
@@ -33,7 +33,7 @@ void AvalancheModifiersAutoFollowTests::Define()
 		TestUtils->Init();
 		DynamicMeshActor = ModifierTestUtils->SpawnTestDynamicMeshActor();
 		ModifierSubsystem = UActorModifierCoreSubsystem::Get();
-		ModifierAutoFollowName = ModifierTestUtils->GetModifierName(UAvaAutoFollowModifier::StaticClass());
+		ModifierAutoFollowName = ModifierTestUtils->GetModifierName(UActorModifierAutoFollowModifier::StaticClass());
 		check(ModifierSubsystem);
 	});
 
@@ -66,7 +66,7 @@ void AvalancheModifiersAutoFollowTests::Define()
 
 			InsertOp.NewModifierName = ModifierAutoFollowName;
 			InsertOp.FailReason = OutFailReasonPtr;
-			AutoFollowModifier = Cast<UAvaAutoFollowModifier>(ModifierSubsystem->InsertModifier(ActorModifierStack, InsertOp));
+			AutoFollowModifier = Cast<UActorModifierAutoFollowModifier>(ModifierSubsystem->InsertModifier(ActorModifierStack, InsertOp));
 		});
 
 		It("Should restore initial Dynamic mesh state once Autofollow modifier is disapplied", [this]()
@@ -91,8 +91,8 @@ void AvalancheModifiersAutoFollowTests::Define()
 			if (AutoFollowModifier)
 			{
 				TestTrue("Modifier is not a stack", !AutoFollowModifier->IsModifierStack());
-				FAvaSceneTreeActor ReferenceActor;
-				ReferenceActor.ReferenceContainer = EAvaReferenceContainer::Other;
+				FActorModifierSceneTreeActor ReferenceActor;
+				ReferenceActor.ReferenceContainer = EActorModifierReferenceContainer::Other;
 				ReferenceActor.ReferenceActorWeak = DynamicMeshActorToFollow;
 				AutoFollowModifier->SetReferenceActor(ReferenceActor);
 				FTransform NewActorState = DynamicMeshActor->GetActorTransform();
diff --git a/Engine/Plugins/Experimental/Avalanche/Source/AvalancheModifiers/Private/Tests/AvaModifiersAutoSize.spec.cpp b/Engine/Plugins/Experimental/Avalanche/Source/AvalancheModifiers/Private/Tests/AvaModifiersAutoSize.spec.cpp
--- a/Engine/Plugins/Experimental/Avalanche/Source/AvalancheModifiers/Private/Tests/AvaModifiersAutoSize.spec.cpp
+++ b/Engine/Plugins/Experimental/Avalanche/Source/AvalancheModifiers/Private/Tests/AvaModifiersAutoSize.spec.cpp
@@ -8,6 +8,7 @@
 #include "Tests/Framework/AvaModifiersTestUtils.h"
 #include "Tests/Framework/AvaTestDynamicMeshActor.h"
 #include "Tests/Framework/AvaTestUtils.h"
+#include "Utilities/ActorModifierActorUtils.h"
 
 #if WITH_AUTOMATION_TESTS
 
@@ -16,7 +17,7 @@ BEGIN_DEFINE_SPEC(AvalancheModifiersAutoSize, "Avalanche.Modifiers.AutoSize",
 
 	AAvaTestDynamicMeshActor* ReferenceActor;
 	AAvaTestDynamicMeshActor* ModifiedActor;
-	FAvaSceneTreeActor ActorQuery;
+	FActorModifierSceneTreeActor ActorQuery;
 
 	FBox InitialReferencedActorBounds;
 	FBox InitialModifiedActorBounds;
@@ -42,8 +43,8 @@ void AvalancheModifiersAutoSize::Define()
 		ActorQuery.ReferenceActorWeak = ReferenceActor;
 
 		// Store initial actors data
-		InitialReferencedActorBounds = FAvaModifiersActorUtils::GetActorBounds(ReferenceActor);
-		InitialModifiedActorBounds = FAvaModifiersActorUtils::GetActorBounds(ModifiedActor);
+		InitialReferencedActorBounds = UE::ActorModifier::ActorUtils::GetActorBounds(ReferenceActor);
+		InitialModifiedActorBounds = UE::ActorModifier::ActorUtils::GetActorBounds(ModifiedActor);
 
 		// Set up modifier stack
 		const UActorModifierCoreSubsystem* ModifierSubsystem = UActorModifierCoreSubsystem::Get();
@@ -65,8 +66,8 @@ void AvalancheModifiersAutoSize::Define()
 	{
 		It("Should change it's size according to a reference object and the modifier settings", [this]
 		{
-			FBox ActualModifiedActorBounds = FAvaModifiersActorUtils::GetActorBounds(ModifiedActor);
-			FBox ActualReferenceActorBounds = FAvaModifiersActorUtils::GetActorBounds(ReferenceActor);
+			FBox ActualModifiedActorBounds = UE::ActorModifier::ActorUtils::GetActorBounds(ModifiedActor);
+			FBox ActualReferenceActorBounds = UE::ActorModifier::ActorUtils::GetActorBounds(ReferenceActor);
 			
 			TestNotEqual("Modified actor bounds have changed", InitialModifiedActorBounds.Max,
 			             ActualModifiedActorBounds.Max);
diff --git a/Engine/Plugins/Experimental/Avalanche/Source/AvalancheModifiers/Private/Tests/AvaModifiersGridArrange.spec.cpp b/Engine/Plugins/Experimental/Avalanche/Source/AvalancheModifiers/Private/Tests/AvaModifiersGridArrange.spec.cpp
--- a/Engine/Plugins/Experimental/Avalanche/Source/AvalancheModifiers/Private/Tests/AvaModifiersGridArrange.spec.cpp
+++ b/Engine/Plugins/Experimental/Avalanche/Source/AvalancheModifiers/Private/Tests/AvaModifiersGridArrange.spec.cpp
@@ -2,7 +2,7 @@
 
 #include "Misc/AutomationTest.h"
 #include "Modifiers/ActorModifierCoreStack.h"
-#include "Modifiers/AvaGridArrangeModifier.h"
+#include "Modifiers/ActorModifierGridArrangeModifier.h"
 #include "Subsystems/ActorModifierCoreSubsystem.h"
 #include "Tests/Framework/AvaModifiersTestUtils.h"
 #include "Tests/Framework/AvaTestDynamicMeshActor.h"
@@ -23,7 +23,7 @@ BEGIN_DEFINE_SPEC(AvalancheModifiersGridArrange, "Avalanche.Modifiers.GridArrang
 	FTransform InitialChildThreeActorState;
 	FTransform InitialChildFourActorState;
 
-	UAvaGridArrangeModifier* GridArrangeModifier;
+	UActorModifierGridArrangeModifier* GridArrangeModifier;
 	FIntPoint GridCount;
 	FVector2d GridSpread;
 
@@ -91,10 +91,10 @@ void AvalancheModifiersGridArrange::Define()
 
 		// Set up modifier
 		const UActorModifierCoreSubsystem* ModifierSubsystem = UActorModifierCoreSubsystem::Get();
-		const FName ModifierGridArrangeName = ModifierTestUtils->GetModifierName(UAvaGridArrangeModifier::StaticClass());
+		const FName ModifierGridArrangeName = ModifierTestUtils->GetModifierName(UActorModifierGridArrangeModifier::StaticClass());
 		FActorModifierCoreStackInsertOp InsertOp = ModifierTestUtils->GenerateInsertOp(ModifierGridArrangeName);
 		UActorModifierCoreStack* ModifierStack = ModifierTestUtils->GenerateModifierStackForActor(ParentMeshActor);
-		GridArrangeModifier = Cast<UAvaGridArrangeModifier>(ModifierSubsystem->InsertModifier(ModifierStack, InsertOp));
+		GridArrangeModifier = Cast<UActorModifierGridArrangeModifier>(ModifierSubsystem->InsertModifier(ModifierStack, InsertOp));
 	});
 
 	AfterEach([this]
diff --git a/Engine/Plugins/Experimental/Avalanche/Source/AvalancheModifiers/Private/Tests/AvaModifiersJustify.spec.cpp b/Engine/Plugins/Experimental/Avalanche/Source/AvalancheModifiers/Private/Tests/AvaModifiersJustify.spec.cpp
--- a/Engine/Plugins/Experimental/Avalanche/Source/AvalancheModifiers/Private/Tests/AvaModifiersJustify.spec.cpp
+++ b/Engine/Plugins/Experimental/Avalanche/Source/AvalancheModifiers/Private/Tests/AvaModifiersJustify.spec.cpp
@@ -2,7 +2,7 @@
 
 #include "Misc/AutomationTest.h"
 #include "Modifiers/ActorModifierCoreStack.h"
-#include "Modifiers/AvaJustifyModifier.h"
+#include "Modifiers/ActorModifierJustifyModifier.h"
 #include "Subsystems/ActorModifierCoreSubsystem.h"
 #include "Tests/Framework/AvaModifiersTestUtils.h"
 #include "Tests/Framework/AvaTestDynamicMeshActor.h"
@@ -24,7 +24,7 @@ BEGIN_DEFINE_SPEC(AvalancheModifiersJustify, "Avalanche.Modifiers.Justify",
 	FVector InitialChildTwoLocation;
 	FVector InitialChildThreeLocation;
 
-	UAvaJustifyModifier* JustifyModifier;
+	UActorModifierJustifyModifier* JustifyModifier;
 
 	TSharedPtr<FAvaTestUtils> TestUtils = MakeShared<FAvaTestUtils>();
 	TSharedPtr<FAvaModifierTestUtils> ModifierTestUtils = MakeShared<FAvaModifierTestUtils>(TestUtils);
@@ -82,13 +82,13 @@ void AvalancheModifiersJustify::Define()
 
 		// Set up modifier
 		const UActorModifierCoreSubsystem* ModifierSubsystem = UActorModifierCoreSubsystem::Get();
-		const FName ModifierJustifyName = ModifierTestUtils->GetModifierName(UAvaJustifyModifier::StaticClass());
+		const FName ModifierJustifyName = ModifierTestUtils->GetModifierName(UActorModifierJustifyModifier::StaticClass());
 		FActorModifierCoreStackInsertOp InsertOp = ModifierTestUtils->GenerateInsertOp(ModifierJustifyName);
 		UActorModifierCoreStack* ModifierStack = ModifierTestUtils->GenerateModifierStackForActor(ParentActor);
-		JustifyModifier = Cast<UAvaJustifyModifier>(ModifierSubsystem->InsertModifier(ModifierStack, InsertOp));
-		JustifyModifier->SetHorizontalAlignment(EAvaJustifyHorizontal::Right);
-		JustifyModifier->SetVerticalAlignment(EAvaJustifyVertical::Bottom);
-		JustifyModifier->SetDepthAlignment(EAvaJustifyDepth::Center);
+		JustifyModifier = Cast<UActorModifierJustifyModifier>(ModifierSubsystem->InsertModifier(ModifierStack, InsertOp));
+		JustifyModifier->SetHorizontalAlignment(EActorModifierJustifyHorizontal::Right);
+		JustifyModifier->SetVerticalAlignment(EActorModifierJustifyVertical::Bottom);
+		JustifyModifier->SetDepthAlignment(EActorModifierJustifyDepth::Center);
 		JustifyModifier->SetDepthAnchor(AnchorValue);
 	});
 
diff --git a/Engine/Plugins/Experimental/Avalanche/Source/AvalancheModifiers/Private/Tests/AvaModifiersLookAt.spec.cpp b/Engine/Plugins/Experimental/Avalanche/Source/AvalancheModifiers/Private/Tests/AvaModifiersLookAt.spec.cpp
--- a/Engine/Plugins/Experimental/Avalanche/Source/AvalancheModifiers/Private/Tests/AvaModifiersLookAt.spec.cpp
+++ b/Engine/Plugins/Experimental/Avalanche/Source/AvalancheModifiers/Private/Tests/AvaModifiersLookAt.spec.cpp
@@ -3,7 +3,7 @@
 #include "Kismet/KismetMathLibrary.h"
 #include "Misc/AutomationTest.h"
 #include "Modifiers/ActorModifierCoreStack.h"
-#include "Modifiers/AvaLookAtModifier.h"
+#include "Modifiers/ActorModifierLookAtModifier.h"
 #include "Subsystems/ActorModifierCoreSubsystem.h"
 #include "Tests/Framework/AvaModifiersTestUtils.h"
 #include "Tests/Framework/AvaTestDynamicMeshActor.h"
@@ -20,7 +20,7 @@ BEGIN_DEFINE_SPEC(AvalancheModifiersLookAt, "Avalanche.Modifiers.LookAt",
 	FTransform InitialModifiedState;
 	FTransform InitialReferenceSate;
 
-	UAvaLookAtModifier* LookAtModifier;
+	UActorModifierLookAtModifier* LookAtModifier;
 
 	TSharedPtr<FAvaTestUtils> TestUtils = MakeShared<FAvaTestUtils>();
 	TSharedPtr<FAvaModifierTestUtils> ModifierTestUtils = MakeShared<FAvaModifierTestUtils>(TestUtils);
@@ -49,11 +49,11 @@ void AvalancheModifiersLookAt::Define()
 
 		// Set up modifier
 		const UActorModifierCoreSubsystem* ModifierSubsystem = UActorModifierCoreSubsystem::Get();
-		const FName ModifierLookAtName = ModifierTestUtils->GetModifierName(UAvaLookAtModifier::StaticClass());
+		const FName ModifierLookAtName = ModifierTestUtils->GetModifierName(UActorModifierLookAtModifier::StaticClass());
 		FActorModifierCoreStackInsertOp InsertOp = ModifierTestUtils->GenerateInsertOp(ModifierLookAtName);
 		UActorModifierCoreStack* ModifierStack = ModifierTestUtils->GenerateModifierStackForActor(ModifiedActor);
-		LookAtModifier = Cast<UAvaLookAtModifier>(ModifierSubsystem->InsertModifier(ModifierStack, InsertOp));
-		LookAtModifier->SetReferenceActor(FAvaSceneTreeActor(ReferenceActor));
+		LookAtModifier = Cast<UActorModifierLookAtModifier>(ModifierSubsystem->InsertModifier(ModifierStack, InsertOp));
+		LookAtModifier->SetReferenceActor(FActorModifierSceneTreeActor(ReferenceActor));
 	});
 
 	AfterEach([this]
diff --git a/Engine/Plugins/Experimental/Avalanche/Source/AvalancheModifiers/Private/Tests/AvaModifiersRadialArrange.cpp b/Engine/Plugins/Experimental/Avalanche/Source/AvalancheModifiers/Private/Tests/AvaModifiersRadialArrange.cpp
--- a/Engine/Plugins/Experimental/Avalanche/Source/AvalancheModifiers/Private/Tests/AvaModifiersRadialArrange.cpp
+++ b/Engine/Plugins/Experimental/Avalanche/Source/AvalancheModifiers/Private/Tests/AvaModifiersRadialArrange.cpp
@@ -2,7 +2,7 @@
 
 #include "Misc/AutomationTest.h"
 #include "Modifiers/ActorModifierCoreStack.h"
-#include "Modifiers/AvaRadialArrangeModifier.h"
+#include "Modifiers/ActorModifierRadialArrangeModifier.h"
 #include "Subsystems/ActorModifierCoreSubsystem.h"
 #include "Tests/Framework/AvaModifiersTestUtils.h"
 #include "Tests/Framework/AvaTestDynamicMeshActor.h"
@@ -17,7 +17,7 @@ BEGIN_DEFINE_SPEC(AvalancheModifiersRadialArrange, "Avalanche.Modifiers.RadialAr
 
 	FVector InitialParentLocation;
 
-	UAvaRadialArrangeModifier* RadialArrangeModifier;
+	UActorModifierRadialArrangeModifier* RadialArrangeModifier;
 
 	TSharedPtr<FAvaTestUtils> TestUtils = MakeShared<FAvaTestUtils>();
 	TSharedPtr<FAvaModifierTestUtils> ModifierTestUtils = MakeShared<FAvaModifierTestUtils>(TestUtils);
@@ -39,10 +39,10 @@ void AvalancheModifiersRadialArrange::Define()
 
 		// Set up modifier
 		const UActorModifierCoreSubsystem* ModifierSubsystem = UActorModifierCoreSubsystem::Get();
-		const FName ModifierName = ModifierTestUtils->GetModifierName(UAvaRadialArrangeModifier::StaticClass());
+		const FName ModifierName = ModifierTestUtils->GetModifierName(UActorModifierRadialArrangeModifier::StaticClass());
 		FActorModifierCoreStackInsertOp InsertOp = ModifierTestUtils->GenerateInsertOp(ModifierName);
 		UActorModifierCoreStack* ModifierStack = ModifierTestUtils->GenerateModifierStackForActor(ParentActor);
-		RadialArrangeModifier = Cast<UAvaRadialArrangeModifier>(ModifierSubsystem->InsertModifier(ModifierStack, InsertOp));
+		RadialArrangeModifier = Cast<UActorModifierRadialArrangeModifier>(ModifierSubsystem->InsertModifier(ModifierStack, InsertOp));
 		InnerRadius = FMath::RandRange(-1000, 1000);
 		RadialArrangeModifier->SetInnerRadius(InnerRadius);
 	});
diff --git a/Engine/Plugins/Experimental/Avalanche/Source/AvalancheModifiers/Public/Modifiers/AvaAutoSizeModifier.h b/Engine/Plugins/Experimental/Avalanche/Source/AvalancheModifiers/Public/Modifiers/AvaAutoSizeModifier.h
--- a/Engine/Plugins/Experimental/Avalanche/Source/AvalancheModifiers/Public/Modifiers/AvaAutoSizeModifier.h
+++ b/Engine/Plugins/Experimental/Avalanche/Source/AvalancheModifiers/Public/Modifiers/AvaAutoSizeModifier.h
@@ -2,11 +2,10 @@
 
 #pragma once
 
-#include "AvaDefs.h"
 #include "AvaGeometryBaseModifier.h"
-#include "Extensions/AvaRenderStateUpdateModifierExtension.h"
-#include "Extensions/AvaSceneTreeUpdateModifierExtension.h"
-#include "Extensions/AvaTransformUpdateModifierExtension.h"
+#include "Extensions/ActorModifierRenderStateUpdateExtension.h"
+#include "Extensions/ActorModifierSceneTreeUpdateExtension.h"
+#include "Extensions/ActorModifierTransformUpdateExtension.h"
 #include "Layout/Margin.h"
 #include "AvaAutoSizeModifier.generated.h"
 
@@ -26,18 +25,18 @@ enum class EAvaAutoSizeFitMode : uint8
  */
 UCLASS(MinimalAPI, BlueprintType)
 class UAvaAutoSizeModifier : public UAvaGeometryBaseModifier
-	, public IAvaTransformUpdateHandler
-	, public IAvaRenderStateUpdateHandler
-	, public IAvaSceneTreeUpdateHandler
+	, public IActorModifierTransformUpdateHandler
+	, public IActorModifierRenderStateUpdateHandler
+	, public IActorModifierSceneTreeUpdateHandler
 {
 	GENERATED_BODY()
 
 public:
 	UFUNCTION(BlueprintCallable, Category="Motion Design|Modifiers|AutoSize")
-	AVALANCHEMODIFIERS_API void SetReferenceActor(const FAvaSceneTreeActor& InReferenceActor);
+	AVALANCHEMODIFIERS_API void SetReferenceActor(const FActorModifierSceneTreeActor& InReferenceActor);
 
 	UFUNCTION(BlueprintPure, Category="Motion Design|Modifiers|AutoSize")
-	const FAvaSceneTreeActor& GetReferenceActor() const
+	const FActorModifierSceneTreeActor& GetReferenceActor() const
 	{
 		return ReferenceActor;
 	}
@@ -116,11 +115,11 @@ protected:
 	void OnReferenceActorChanged();
 
 	UPROPERTY(EditInstanceOnly, Setter="SetReferenceActor", Getter="GetReferenceActor", Category="AutoSize", meta=(ShowOnlyInnerProperties, AllowPrivateAccess="true"))
-	FAvaSceneTreeActor ReferenceActor;
+	FActorModifierSceneTreeActor ReferenceActor;
 
 	/** The method for finding a reference actor based on it's position in the parent's hierarchy. */
 	UPROPERTY(meta=(DeprecatedProperty, DeprecationMessage="Use ReferenceActor instead"))
-	EAvaReferenceContainer ReferenceContainer_DEPRECATED = EAvaReferenceContainer::Other;
+	EActorModifierReferenceContainer ReferenceContainer_DEPRECATED = EActorModifierReferenceContainer::Other;
 
 	/** The actor affecting the modifier. This is user selectable if the Reference Container is set to "Other". */
 	UPROPERTY(meta=(DeprecatedProperty, DeprecationMessage="Use ReferenceActor instead"))
diff --git a/Engine/Plugins/Experimental/Avalanche/Source/AvalancheModifiers/Public/Modifiers/AvaBooleanModifier.h b/Engine/Plugins/Experimental/Avalanche/Source/AvalancheModifiers/Public/Modifiers/AvaBooleanModifier.h
--- a/Engine/Plugins/Experimental/Avalanche/Source/AvalancheModifiers/Public/Modifiers/AvaBooleanModifier.h
+++ b/Engine/Plugins/Experimental/Avalanche/Source/AvalancheModifiers/Public/Modifiers/AvaBooleanModifier.h
@@ -4,7 +4,7 @@
 
 #include "AvaGeometryBaseModifier.h"
 #include "AvaModifiersParametricMaterial.h"
-#include "Extensions/AvaTransformUpdateModifierExtension.h"
+#include "Extensions/ActorModifierTransformUpdateExtension.h"
 #include "Shared/AvaBooleanModifierShared.h"
 #include "AvaBooleanModifier.generated.h"
 
@@ -24,7 +24,7 @@ enum class EAvaBooleanMode : uint8
 /** This modifier allows you to apply a mask on a certain shape, this will affect every shape it collides with that matches options */
 UCLASS(MinimalAPI, BlueprintType)
 class UAvaBooleanModifier : public UAvaGeometryBaseModifier
-	, public IAvaTransformUpdateHandler
+	, public IActorModifierTransformUpdateHandler
 {
 	GENERATED_BODY()
 
diff --git a/Engine/Plugins/Experimental/Avalanche/Source/AvalancheModifiers/Public/Modifiers/AvaDynamicMeshConverterModifier.h b/Engine/Plugins/Experimental/Avalanche/Source/AvalancheModifiers/Public/Modifiers/AvaDynamicMeshConverterModifier.h
--- a/Engine/Plugins/Experimental/Avalanche/Source/AvalancheModifiers/Public/Modifiers/AvaDynamicMeshConverterModifier.h
+++ b/Engine/Plugins/Experimental/Avalanche/Source/AvalancheModifiers/Public/Modifiers/AvaDynamicMeshConverterModifier.h
@@ -4,8 +4,8 @@
 
 #include "AvaGeometryBaseModifier.h"
 #include "CEMeshBuilder.h"
-#include "Extensions/AvaRenderStateUpdateModifierExtension.h"
-#include "Extensions/AvaSceneTreeUpdateModifierExtension.h"
+#include "Extensions/ActorModifierRenderStateUpdateExtension.h"
+#include "Extensions/ActorModifierSceneTreeUpdateExtension.h"
 #include "Templates/SubclassOf.h"
 #include "AvaDynamicMeshConverterModifier.generated.h"
 
@@ -82,8 +82,8 @@ struct FAvaDynamicMeshConverterModifierComponentState
 UCLASS(MinimalAPI, BlueprintType)
 class UAvaDynamicMeshConverterModifier
 	: public UAvaGeometryBaseModifier
-	, public IAvaSceneTreeUpdateHandler
-	, public IAvaRenderStateUpdateHandler
+	, public IActorModifierSceneTreeUpdateHandler
+	, public IActorModifierRenderStateUpdateHandler
 {
 	GENERATED_BODY()
 
@@ -252,7 +252,7 @@ protected:
 	UPROPERTY(Transient, DuplicateTransient, NonTransactional)
 	FCEMeshBuilder MeshBuilder;
 
-	FAvaSceneTreeActor TrackedActor;
+	FActorModifierSceneTreeActor TrackedActor;
 
 	/** Time elasped since */
 	double LastTransformUpdateTime = 0;
diff --git a/Engine/Plugins/Experimental/Avalanche/Source/AvalancheModifiers/Public/Modifiers/AvaGeometryBaseModifier.h b/Engine/Plugins/Experimental/Avalanche/Source/AvalancheModifiers/Public/Modifiers/AvaGeometryBaseModifier.h
--- a/Engine/Plugins/Experimental/Avalanche/Source/AvalancheModifiers/Public/Modifiers/AvaGeometryBaseModifier.h
+++ b/Engine/Plugins/Experimental/Avalanche/Source/AvalancheModifiers/Public/Modifiers/AvaGeometryBaseModifier.h
@@ -2,10 +2,10 @@
 
 #pragma once
 
-#include "AvaBaseModifier.h"
 #include "DynamicMesh/DynamicMesh3.h"
 #include "DynamicMesh/DynamicMeshAttributeSet.h"
 #include "Misc/Optional.h"
+#include "Modifiers/ActorModifierCoreBase.h"
 #include "UObject/ObjectPtr.h"
 #include "AvaGeometryBaseModifier.generated.h"
 
@@ -14,7 +14,7 @@ class UDynamicMeshComponent;
 
 /** Base class used for geometry modifier that uses dynamic mesh components */
 UCLASS(Abstract)
-class UAvaGeometryBaseModifier : public UAvaBaseModifier
+class UAvaGeometryBaseModifier : public UActorModifierCoreBase
 {
 	GENERATED_BODY()
 
diff --git a/Engine/Plugins/Experimental/Avalanche/Source/AvalancheModifiers/Public/Modifiers/AvaHideEmptyModifier.h b/Engine/Plugins/Experimental/Avalanche/Source/AvalancheModifiers/Public/Modifiers/AvaHideEmptyModifier.h
--- a/Engine/Plugins/Experimental/Avalanche/Source/AvalancheModifiers/Public/Modifiers/AvaHideEmptyModifier.h
+++ b/Engine/Plugins/Experimental/Avalanche/Source/AvalancheModifiers/Public/Modifiers/AvaHideEmptyModifier.h
@@ -2,13 +2,13 @@
 
 #pragma once
 
-#include "AvaArrangeBaseModifier.h"
+#include "Modifiers/ActorModifierArrangeBaseModifier.h"
 #include "AvaHideEmptyModifier.generated.h"
 
 class UText3DComponent;
 
 UCLASS(MinimalAPI, BlueprintType)
-class UAvaHideEmptyModifier : public UAvaArrangeBaseModifier
+class UAvaHideEmptyModifier : public UActorModifierArrangeBaseModifier
 {
 	GENERATED_BODY()
 
diff --git a/Engine/Plugins/Experimental/Avalanche/Source/AvalancheModifiers/Public/Modifiers/AvaMaterialParameterModifier.h b/Engine/Plugins/Experimental/Avalanche/Source/AvalancheModifiers/Public/Modifiers/AvaMaterialParameterModifier.h
--- a/Engine/Plugins/Experimental/Avalanche/Source/AvalancheModifiers/Public/Modifiers/AvaMaterialParameterModifier.h
+++ b/Engine/Plugins/Experimental/Avalanche/Source/AvalancheModifiers/Public/Modifiers/AvaMaterialParameterModifier.h
@@ -2,7 +2,7 @@
 
 #pragma once
 
-#include "AvaArrangeBaseModifier.h"
+#include "Modifiers/ActorModifierArrangeBaseModifier.h"
 #include "UObject/ObjectPtr.h"
 #include "AvaMaterialParameterModifier.generated.h"
 
@@ -35,7 +35,7 @@ struct FAvaMaterialParameterMap
 
 /** This modifier sets specified dynamic materials parameters on an actor and its children */
 UCLASS(MinimalAPI, BlueprintType)
-class UAvaMaterialParameterModifier : public UAvaArrangeBaseModifier
+class UAvaMaterialParameterModifier : public UActorModifierArrangeBaseModifier
 {
 	GENERATED_BODY()
 
diff --git a/Engine/Plugins/Experimental/Avalanche/Source/AvalancheModifiers/Public/Modifiers/AvaSplineSweepModifier.h b/Engine/Plugins/Experimental/Avalanche/Source/AvalancheModifiers/Public/Modifiers/AvaSplineSweepModifier.h
--- a/Engine/Plugins/Experimental/Avalanche/Source/AvalancheModifiers/Public/Modifiers/AvaSplineSweepModifier.h
+++ b/Engine/Plugins/Experimental/Avalanche/Source/AvalancheModifiers/Public/Modifiers/AvaSplineSweepModifier.h
@@ -4,7 +4,7 @@
 
 #include "AvaGeometryBaseModifier.h"
 #include "AvaPropertyChangeDispatcher.h"
-#include "Extensions/AvaRenderStateUpdateModifierExtension.h"
+#include "Extensions/ActorModifierRenderStateUpdateExtension.h"
 #include "AvaSplineSweepModifier.generated.h"
 
 class USplineComponent;
@@ -20,7 +20,7 @@ enum class EAvaSplineSweepSampleMode : uint8
 
 UCLASS(MinimalAPI, BlueprintType)
 class UAvaSplineSweepModifier : public UAvaGeometryBaseModifier
-	, public IAvaRenderStateUpdateHandler
+	, public IActorModifierRenderStateUpdateHandler
 {
 	GENERATED_BODY()
 
diff --git a/Engine/Plugins/Experimental/Avalanche/Source/AvalancheModifiers/Public/Modifiers/AvaTranslucentPriorityModifier.h b/Engine/Plugins/Experimental/Avalanche/Source/AvalancheModifiers/Public/Modifiers/AvaTranslucentPriorityModifier.h
--- a/Engine/Plugins/Experimental/Avalanche/Source/AvalancheModifiers/Public/Modifiers/AvaTranslucentPriorityModifier.h
+++ b/Engine/Plugins/Experimental/Avalanche/Source/AvalancheModifiers/Public/Modifiers/AvaTranslucentPriorityModifier.h
@@ -2,8 +2,7 @@
 
 #pragma once
 
-#include "AvaArrangeBaseModifier.h"
-#include "AvaBaseModifier.h"
+#include "Modifiers/ActorModifierArrangeBaseModifier.h"
 #include "Shared/AvaTranslucentPriorityModifierShared.h"
 #include "UObject/WeakObjectPtrTemplatesFwd.h"
 #include "AvaTranslucentPriorityModifier.generated.h"
@@ -23,7 +22,7 @@ enum class EAvaTranslucentPriorityModifierMode : uint8
 };
 
 UCLASS(MinimalAPI, BlueprintType)
-class UAvaTranslucentPriorityModifier : public UAvaArrangeBaseModifier
+class UAvaTranslucentPriorityModifier : public UActorModifierArrangeBaseModifier
 {
 	GENERATED_BODY()
 
diff --git a/Engine/Plugins/Experimental/Avalanche/Source/AvalancheModifiers/Public/Modifiers/AvaVisibilityModifier.h b/Engine/Plugins/Experimental/Avalanche/Source/AvalancheModifiers/Public/Modifiers/AvaVisibilityModifier.h
--- a/Engine/Plugins/Experimental/Avalanche/Source/AvalancheModifiers/Public/Modifiers/AvaVisibilityModifier.h
+++ b/Engine/Plugins/Experimental/Avalanche/Source/AvalancheModifiers/Public/Modifiers/AvaVisibilityModifier.h
@@ -2,7 +2,7 @@
 
 #pragma once
 
-#include "AvaArrangeBaseModifier.h"
+#include "Modifiers/ActorModifierArrangeBaseModifier.h"
 #include "AvaVisibilityModifier.generated.h"
 
 class AActor;
@@ -11,7 +11,7 @@ class AActor;
  * Controls the visibility of a range of child actors by index.
  */
 UCLASS(MinimalAPI, BlueprintType)
-class UAvaVisibilityModifier : public UAvaArrangeBaseModifier
+class UAvaVisibilityModifier : public UActorModifierArrangeBaseModifier
 {
 	GENERATED_BODY()
 
diff --git a/Engine/Plugins/Experimental/Avalanche/Source/AvalancheModifiersEditor/AvalancheModifiersEditor.Build.cs b/Engine/Plugins/Experimental/Avalanche/Source/AvalancheModifiersEditor/AvalancheModifiersEditor.Build.cs
--- a/Engine/Plugins/Experimental/Avalanche/Source/AvalancheModifiersEditor/AvalancheModifiersEditor.Build.cs
+++ b/Engine/Plugins/Experimental/Avalanche/Source/AvalancheModifiersEditor/AvalancheModifiersEditor.Build.cs
@@ -21,14 +21,16 @@ public class AvalancheModifiersEditor : ModuleRules
 		PrivateDependencyModuleNames.AddRange(
 			new string[]
 			{
+				"ActorModifier",
 				"ActorModifierCore",
 				"ActorModifierCoreEditor",
+				"ActorModifierEditor",
 				"CoreUObject",
 				"Engine",
 				"Projects",
 				"Slate",
 				"ToolMenus",
-				"UnrealEd"
+				"UnrealEd",
 			}
 		);
 	}
diff --git a/Engine/Plugins/Experimental/Avalanche/Source/AvalancheModifiersEditor/Private/AvaModifiersEditorStyle.cpp b/Engine/Plugins/Experimental/Avalanche/Source/AvalancheModifiersEditor/Private/AvaModifiersEditorStyle.cpp
--- a/Engine/Plugins/Experimental/Avalanche/Source/AvalancheModifiersEditor/Private/AvaModifiersEditorStyle.cpp
+++ b/Engine/Plugins/Experimental/Avalanche/Source/AvalancheModifiersEditor/Private/AvaModifiersEditorStyle.cpp
@@ -7,6 +7,7 @@
 #include "Interfaces/IPluginManager.h"
 #include "Math/MathFwd.h"
 #include "Modifiers/ActorModifierCoreBase.h"
+#include "Styles/ActorModifierEditorStyle.h"
 #include "Styling/StyleColors.h"
 #include "Styling/SlateColor.h"
 #include "Styling/SlateTypes.h"
@@ -15,26 +16,13 @@
 #include "Subsystems/ActorModifierCoreSubsystem.h"
 
 FAvaModifiersEditorStyle::FAvaModifiersEditorStyle()
-	: FSlateStyleSet(TEXT("AvaModifiersEditor"))
+	: FSlateStyleSet(UE_MODULE_NAME)
 {
-	const FVector2f Icon16x16(16.f, 16.f);
-
 	const TSharedPtr<IPlugin> Plugin = IPluginManager::Get().FindPlugin(UE_PLUGIN_NAME);
 	check(Plugin.IsValid());
 
 	SetContentRoot(Plugin->GetBaseDir() / TEXT("Resources"));
 
-	Set("ClassIcon.AvaBaseModifier", new IMAGE_BRUSH_SVG("Icons/ModifierIcons/BaseModifier", Icon16x16));
-
-	ModifierCategoriesColors =
-	{
-		{TEXT("Unspecified"), FSlateColor(FActorModifierCoreMetadata::DefaultColor)},
-		{TEXT("Geometry"), FSlateColor(FLinearColor::Yellow.Desaturate(0.25))},
-		{TEXT("Transform"), FSlateColor(FLinearColor::Blue.Desaturate(0.25))},
-		{TEXT("Layout"), FSlateColor(FLinearColor::Green.Desaturate(0.25))},
-		{TEXT("Rendering"), FSlateColor(FLinearColor::Red.Desaturate(0.25))}
-	};
-
 	if (const UActorModifierCoreSubsystem* ModifierSubsystem = UActorModifierCoreSubsystem::Get())
 	{
 		// loop through already registered factories if subsystem exists
@@ -50,16 +38,6 @@ FAvaModifiersEditorStyle::FAvaModifiersEditorStyle()
 	UActorModifierCoreSubsystem::OnModifierClassRegistered().AddRaw(this, &FAvaModifiersEditorStyle::OnModifierClassRegistered);
 }
 
-const FSlateColor& FAvaModifiersEditorStyle::GetModifierCategoryColor(FName CategoryName)
-{
-	if (ModifierCategoriesColors.Contains(CategoryName))
-	{
-		return ModifierCategoriesColors[CategoryName];
-	}
-
-	return ModifierCategoriesColors[TEXT("Unspecified")];
-}
-
 void FAvaModifiersEditorStyle::OnModifierClassRegistered(const FActorModifierCoreMetadata& InMetadata)
 {
 	// We need to update the metadata to set various options
@@ -67,8 +45,8 @@ void FAvaModifiersEditorStyle::OnModifierClassRegistered(const FActorModifierCor
 
 	const FVector2f Icon16x16(16.f, 16.f);
 	const FString ModifierIconPath = TEXT("Icons/ModifierIcons/") + InMetadata.GetName().ToString() + TEXT("Modifier");
-
-	MutableMetadata.SetColor(GetModifierCategoryColor(InMetadata.GetCategory()).GetSpecifiedColor());
+	
+	MutableMetadata.SetColor(FActorModifierEditorStyle::GetModifierCategoryColor(InMetadata.GetCategory()).GetSpecifiedColor());
 
 	if (FPlatformFileManager::Get().GetPlatformFile().FileExists(*RootToContentDir(ModifierIconPath, TEXT(".svg"))))
 	{
diff --git a/Engine/Plugins/Experimental/Avalanche/Source/AvalancheModifiersEditor/Private/AvaModifiersEditorStyle.h b/Engine/Plugins/Experimental/Avalanche/Source/AvalancheModifiersEditor/Private/AvaModifiersEditorStyle.h
--- a/Engine/Plugins/Experimental/Avalanche/Source/AvalancheModifiersEditor/Private/AvaModifiersEditorStyle.h
+++ b/Engine/Plugins/Experimental/Avalanche/Source/AvalancheModifiersEditor/Private/AvaModifiersEditorStyle.h
@@ -19,10 +19,6 @@ public:
 	FAvaModifiersEditorStyle();
 	virtual ~FAvaModifiersEditorStyle() override;
 
-	const FSlateColor& GetModifierCategoryColor(FName CategoryName);
-
 private:
 	void OnModifierClassRegistered(const FActorModifierCoreMetadata& InMetadata);
-
-	TMap<FName, FSlateColor> ModifierCategoriesColors;
 };
diff --git a/Engine/Plugins/Experimental/Avalanche/Source/AvalancheModifiersEditor/Private/Outliner/AvaOutlinerModifier.cpp b/Engine/Plugins/Experimental/Avalanche/Source/AvalancheModifiersEditor/Private/Outliner/AvaOutlinerModifier.cpp
--- a/Engine/Plugins/Experimental/Avalanche/Source/AvalancheModifiersEditor/Private/Outliner/AvaOutlinerModifier.cpp
+++ b/Engine/Plugins/Experimental/Avalanche/Source/AvalancheModifiersEditor/Private/Outliner/AvaOutlinerModifier.cpp
@@ -3,7 +3,6 @@
 #include "Outliner/AvaOutlinerModifier.h"
 #include "Modifiers/ActorModifierCoreBase.h"
 #include "Modifiers/ActorModifierCoreStack.h"
-#include "Modifiers/AvaBaseModifier.h"
 #include "Selection/AvaOutlinerScopedSelection.h"
 #include "Styling/SlateIconFinder.h"
 #include "Subsystems/ActorModifierCoreEditorSubsystem.h"
@@ -14,7 +13,7 @@ FAvaOutlinerModifier::FAvaOutlinerModifier(IAvaOutliner& InOutliner, UActorModif
 	, Modifier(InModifier)
 {
 	ModifierName = FText::FromName(InModifier->GetModifierName());
-	ModifierIcon = FSlateIconFinder::FindIconForClass(UAvaBaseModifier::StaticClass());
+	ModifierIcon = FSlateIconFinder::FindIconForClass(UActorModifierCoreBase::StaticClass());
 	ModifierTooltip = FText::GetEmpty();
 
 	if (const UActorModifierCoreSubsystem* ModifierSubsystem = UActorModifierCoreSubsystem::Get())
