MotionDesign : ActorModifier

diff --git a/Engine/Plugins/Experimental/Avalanche/Source/AvalancheModifiersEditor/Private/Outliner/AvaOutlinerModifierProxy.cpp b/Engine/Plugins/Experimental/Avalanche/Source/AvalancheModifiersEditor/Private/Outliner/AvaOutlinerModifierProxy.cpp
--- a/Engine/Plugins/Experimental/Avalanche/Source/AvalancheModifiersEditor/Private/Outliner/AvaOutlinerModifierProxy.cpp
+++ b/Engine/Plugins/Experimental/Avalanche/Source/AvalancheModifiersEditor/Private/Outliner/AvaOutlinerModifierProxy.cpp
@@ -4,7 +4,6 @@
 #include "IAvaOutliner.h"
 #include "Item/AvaOutlinerActor.h"
 #include "Modifiers/ActorModifierCoreStack.h"
-#include "Modifiers/AvaBaseModifier.h"
 #include "Outliner/AvaOutlinerModifier.h"
 #include "Selection/AvaOutlinerScopedSelection.h"
 #include "Styling/SlateIconFinder.h"
@@ -15,7 +14,7 @@
 FAvaOutlinerModifierProxy::FAvaOutlinerModifierProxy(IAvaOutliner& InOutliner, const FAvaOutlinerItemPtr& InParentItem)
 	: Super(InOutliner, InParentItem)
 {
-	ModifierIcon = FSlateIconFinder::FindIconForClass(UAvaBaseModifier::StaticClass());
+	ModifierIcon = FSlateIconFinder::FindIconForClass(UActorModifierCoreBase::StaticClass());
 }
 
 UActorModifierCoreStack* FAvaOutlinerModifierProxy::GetModifierStack() const
diff --git a/Engine/Plugins/Experimental/Avalanche/Source/AvalancheSceneTree/Private/AvaSceneTree.cpp b/Engine/Plugins/Experimental/Avalanche/Source/AvalancheSceneTree/Private/AvaSceneTree.cpp
--- a/Engine/Plugins/Experimental/Avalanche/Source/AvalancheSceneTree/Private/AvaSceneTree.cpp
+++ b/Engine/Plugins/Experimental/Avalanche/Source/AvalancheSceneTree/Private/AvaSceneTree.cpp
@@ -1,6 +1,7 @@
 // Copyright Epic Games, Inc. All Rights Reserved.
 
 #include "AvaSceneTree.h"
+#include "Engine/Level.h"
 #include "Engine/World.h"
 #include "GameFramework/Actor.h"
 #include "Modules/ModuleManager.h"
@@ -195,4 +196,28 @@ TArray<AActor*> FAvaSceneTree::GetChildActors(AActor* const InParentActor) const
 	}
 
 	return OutActors;
-}
\ No newline at end of file
+}
+
+TArray<AActor*> FAvaSceneTree::GetRootActors(const ULevel* InLevel) const
+{
+	TArray<AActor*> OutActors;
+
+	UWorld* Outer = InLevel->GetTypedOuter<UWorld>();
+	check(IsValid(Outer));
+
+	const TConstArrayView<int32> ChildIndices = RootNode.GetChildrenIndices();
+	OutActors.Reserve(ChildIndices.Num());
+	
+	for (const int32 ChildIndex : ChildIndices)
+	{
+		const FAvaSceneItem* const ActorSceneItem = GetItemAtIndex(ChildIndex);
+		
+		if (ActorSceneItem)
+		{
+			AActor* const ChildActor = ActorSceneItem->Resolve<AActor>(Outer);
+			OutActors.Add(ChildActor);
+		}
+	}
+	
+	return OutActors;
+}
diff --git a/Engine/Plugins/Experimental/Avalanche/Source/AvalancheSceneTree/Public/AvaSceneTree.h b/Engine/Plugins/Experimental/Avalanche/Source/AvalancheSceneTree/Public/AvaSceneTree.h
--- a/Engine/Plugins/Experimental/Avalanche/Source/AvalancheSceneTree/Public/AvaSceneTree.h
+++ b/Engine/Plugins/Experimental/Avalanche/Source/AvalancheSceneTree/Public/AvaSceneTree.h
@@ -8,6 +8,7 @@
 #include "AvaSceneTree.generated.h"
 
 class AActor;
+class ULevel;
 
 USTRUCT()
 struct AVALANCHESCENETREE_API FAvaSceneTree
@@ -39,6 +40,9 @@ struct AVALANCHESCENETREE_API FAvaSceneTree
 	/** Returns all resolved child actors of a specified parent actor in the scene tree. */
 	TArray<AActor*> GetChildActors(AActor* const InParentActor) const;
 
+	/** Returns all root actors in this tree */
+	TArray<AActor*> GetRootActors(const ULevel* InLevel) const;
+
 private:
 	void UpdateTreeNodes();
 
