[Motion Design] Rundown Server: Support for creating Combo Templates and various page editing improvements.

New Commands:
- Added "CreateComboTemplate"
- Added "GetPlayableAssets" command to have the server list all it's available playable assets. Needed for clients UI implementation to obtain the asset path for  CreateTemplate command.
- Added "ChangePageName": Allows renaming either template or instance pages.

Modified Commands:
- "CreatePage": Adding the PageId generator parameters to support page id insertion.
- "CreateTemplate": Adding the PageId generator parameters to support template id insertion.
- "CreateTemplate": Adding the asset path to be able to create the template in a single command. This avoid having to delete the invalid template if setting up the asset fails.
- "ChangeTemplateBP": Adding bReimport parameter to allow reimporting the source asset, i.e. refreshing the template's information from the source asset.

New events:
- PageNameChanged: when a page or template is renamed.

diff --git a/Engine/Plugins/Experimental/Avalanche/Source/AvalancheMedia/Internal/Rundown/AvaRundownMessages.h b/Engine/Plugins/Experimental/Avalanche/Source/AvalancheMedia/Internal/Rundown/AvaRundownMessages.h
--- a/Engine/Plugins/Experimental/Avalanche/Source/AvalancheMedia/Internal/Rundown/AvaRundownMessages.h
+++ b/Engine/Plugins/Experimental/Avalanche/Source/AvalancheMedia/Internal/Rundown/AvaRundownMessages.h
@@ -233,6 +233,42 @@ struct FAvaRundownServerInfo : public FAvaRundownMsgBase
 	uint32 RemoteControlWebSocketServerPort = 0;
 };
 
+/**
+ * Requests a list of playable assets that can be added to a
+ * rundown template.
+ * Response is FAvaRundownPlayableAssets.
+ */
+USTRUCT()
+struct FAvaRundownGetPlayableAssets : public FAvaRundownMsgBase
+{
+	GENERATED_BODY()
+
+	/**
+	 * The search query which will be compared with the asset names.
+	*/
+	UPROPERTY()
+	FString Query;
+
+	/**
+	 * The maximum number of search results returned.
+	*/
+	UPROPERTY()
+	int32 Limit = 0;
+};
+
+/**
+ *	List of all available playable assets on the server.
+ *	Expected Response from FAvaRundownGetPlayableAssets.
+ */
+USTRUCT()
+struct FAvaRundownPlayableAssets : public FAvaRundownMsgBase
+{
+	GENERATED_BODY()
+
+	UPROPERTY()
+	TArray<FSoftObjectPath> Assets;
+};
+
 /**
  *	Requests the list of rundowns that can be opened on the current server.
  *	Response is FAvaRundownRundowns.
@@ -450,6 +486,29 @@ public:
 	FString Rundown;
 };
 
+/**
+ * Defines the parameters for the page id generator algorithm.
+ * The Id generator uses a sequence strategy to search for an unused id.
+ * It is defined by a starting id and a search direction.
+ */
+USTRUCT()
+struct FAvaRundownCreatePageIdGeneratorParams
+{
+	GENERATED_BODY()
+
+	/** Starting Id for the search. */
+	UPROPERTY()
+	int32 ReferenceId = FAvaRundownPage::InvalidPageId;
+
+	/**
+	 * @brief (Initial) Search increment.
+	 * @remark For negative increment search, the limit of the search space can be reached. If no unique id is found,
+	 *		   the search will continue in the positive direction instead.		   
+	 */
+	UPROPERTY()
+	int32 Increment = 1;
+};
+
 /**
  * Requests a new page be created from the specified template in the given rundown.
  */
@@ -462,6 +521,10 @@ struct FAvaRundownCreatePage : public FAvaRundownMsgBase
 	UPROPERTY()
 	FString Rundown;
 
+	/** Defines the parameters for the page id generator algorithm. */
+	UPROPERTY()
+	FAvaRundownCreatePageIdGeneratorParams IdGeneratorParams;
+	
 	/** Specifies the template for the newly created page. */
 	UPROPERTY()
 	int32 TemplateId = FAvaRundownPage::InvalidPageId;
@@ -498,6 +561,40 @@ struct FAvaRundownCreateTemplate : public FAvaRundownMsgBase
 	/** Rundown asset path: [PackagePath]/[AssetName].[AssetName] */
 	UPROPERTY()
 	FString Rundown;
+
+	/** Defines the parameters for the page id generator algorithm. */
+	UPROPERTY()
+	FAvaRundownCreatePageIdGeneratorParams IdGeneratorParams;
+
+	/** Specifies the asset path to assign to the template. */
+	UPROPERTY()
+	FString AssetPath;
+};
+
+/**
+ * Requests the creation of a new combo template.
+ * If successful, the response is FAvaRundownServerMsg with a "Template [Id] Created" text.
+ * The id of the created template can be parsed from that message's text.
+ * Also a secondary FAvaRundownPageListChanged event with added template id will be sent.
+ *
+ * @remark A combination template can only be created using transition logic templates that are in different transition layers.
+ */
+USTRUCT()
+struct FAvaRundownCreateComboTemplate : public FAvaRundownMsgBase
+{
+	GENERATED_BODY()
+
+	/** Rundown asset path: [PackagePath]/[AssetName].[AssetName] */
+	UPROPERTY()
+	FString Rundown;
+
+	/** Defines the parameters for the page id generator algorithm. */
+	UPROPERTY()
+	FAvaRundownCreatePageIdGeneratorParams IdGeneratorParams;
+
+	/** Specifies the template ids that are combined. */
+	UPROPERTY()
+	TArray<int32> CombinedTemplateIds;
 };
 
 /** Requests deletion of the given template. */
@@ -534,6 +631,10 @@ public:
 	/** Specifies the asset path to assign. */
 	UPROPERTY()
 	FString AssetPath;
+
+	/** If true, the asset will be re-imported and the template information will be refresh from the source asset. */
+	UPROPERTY()
+	bool bReimport = false;
 };
 
 /** Page Information */
@@ -776,6 +877,31 @@ struct FAvaRundownPageChannelChanged : public FAvaRundownMsgBase
 	FString ChannelName;
 };
 
+/**
+ * Event sent when a page's name is modified.
+ */
+USTRUCT()
+struct FAvaRundownPageNameChanged : public FAvaRundownMsgBase
+{
+	GENERATED_BODY()
+
+	/** Rundown asset path: [PackagePath]/[AssetName].[AssetName] */
+	UPROPERTY()
+	FString Rundown;
+
+	/** Specified the modified page id. */
+	UPROPERTY()
+	int32 PageId = FAvaRundownPage::InvalidPageId;
+
+	/** new page name is currently assigned to (post modification). */
+	UPROPERTY()
+	FString PageName;
+
+	/** Indicate whether the name or friendly name that changed. */ 
+	UPROPERTY()
+	bool bFriendlyName = true;
+};
+
 /**
  * Event sent when a page's animation settings is modified.
  */
@@ -816,6 +942,39 @@ struct FAvaRundownPageChangeChannel : public FAvaRundownMsgBase
 	FString ChannelName;
 };
 
+/**
+ * Sets page name. Works for template or instance pages.
+ * By default, the command will set the page's "friendly" name as it is the one used for
+ * display purposes. The page name is reserved for native code uses.
+ * Along with the corresponding response, this will also trigger a FAvaRundownPageNameChanged event.
+ */
+USTRUCT()
+struct FAvaRundownChangePageName : public FAvaRundownMsgBase
+{
+	GENERATED_BODY()
+
+	/** Rundown asset path: [PackagePath]/[AssetName].[AssetName] */
+	UPROPERTY()
+	FString Rundown;
+
+	/** Specifies the page or template that will be modified. */
+	UPROPERTY()
+	int32 PageId = FAvaRundownPage::InvalidPageId;
+
+	/** Specifies the new page name. */
+	UPROPERTY()
+	FString PageName;
+
+	/**
+	 * If true, the page's friendly name will be set.
+	 * The page name is usually set by the native code.
+	 * For display purposes, it is preferable to use the "friendly" name.
+	 */
+	UPROPERTY()
+	bool bSetFriendlyName = true;
+};
+
+
 /** This is a request to save the managed Remote Control Preset (RCP) back to the corresponding page values. */
 USTRUCT()
 struct FAvaRundownUpdatePageFromRCP : public FAvaRundownMsgBase
diff --git a/Engine/Plugins/Experimental/Avalanche/Source/AvalancheMedia/Private/Rundown/AvaRundown.cpp b/Engine/Plugins/Experimental/Avalanche/Source/AvalancheMedia/Private/Rundown/AvaRundown.cpp
--- a/Engine/Plugins/Experimental/Avalanche/Source/AvalancheMedia/Private/Rundown/AvaRundown.cpp
+++ b/Engine/Plugins/Experimental/Avalanche/Source/AvalancheMedia/Private/Rundown/AvaRundown.cpp
@@ -299,6 +299,61 @@ int32 UAvaRundown::AddComboTemplate(const TArray<int32>& InTemplateIds, const FA
 	});
 }
 
+TArray<int32> UAvaRundown::ValidateTemplateIdsForComboTemplate(const TArray<int32>& InTemplateIds, FOutputDevice& InErrorLog)
+{
+	TSet<FAvaTagId> LayerIds;
+	TArray<int32> TemplateIds;
+	TemplateIds.Reserve(InTemplateIds.Num());
+	FAvaPlayableRemoteControlValues MergedValues;
+	
+	for (const int32 SelectedPageId : InTemplateIds)
+	{
+		const FAvaRundownPage& Page = GetPage(SelectedPageId);
+		if (!Page.IsValidPage())
+		{
+			InErrorLog.Logf(TEXT("Template %d is not valid."), SelectedPageId);
+			continue;
+		}
+		if (!Page.IsTemplate())
+		{
+			InErrorLog.Logf(TEXT("Page %d is not a template."), SelectedPageId);
+			continue;
+		}
+		if (Page.IsComboTemplate())
+		{
+			InErrorLog.Logf(TEXT("Template %d is already a combo template."), SelectedPageId);
+			continue;
+		}
+		if (!Page.HasTransitionLogic(this))
+		{
+			InErrorLog.Logf(TEXT("Template %d doesn't have transition logic."), SelectedPageId);
+			continue;
+		}
+		if (!Page.GetTransitionLayer(this).IsValid())
+		{
+			InErrorLog.Logf(TEXT("Template %d doesn't have a valid transition logic layer."), SelectedPageId);
+			continue;
+		}
+		if (LayerIds.Contains(Page.GetTransitionLayer(this).TagId))
+		{
+			InErrorLog.Logf(TEXT("Template %d's layer %s is already in the selection."), SelectedPageId, *Page.GetTransitionLayer(this).ToString());
+			continue;
+		}
+
+		// Make sure the RC values can merge correctly. If not, original template needs fixing.
+		if (MergedValues.HasIdCollisions(Page.GetRemoteControlValues()))
+		{
+			InErrorLog.Logf(TEXT("Template %d's RemoteControl values have Id collisions with other templates in the selection."), SelectedPageId);
+			continue;
+		}
+		
+		MergedValues.Merge(Page.GetRemoteControlValues());
+		LayerIds.Add(Page.GetTransitionLayer(this).TagId);
+		TemplateIds.Add(Page.GetPageId());
+	}
+	return TemplateIds;
+}
+
 TArray<int32> UAvaRundown::AddTemplates(const TArray<FAvaRundownPage>& InSourceTemplates)
 {
 	if (!CanAddPage() || InSourceTemplates.IsEmpty())
@@ -2017,7 +2072,7 @@ namespace UE::AvaMedia::Rundown::Private
 			// Check if we have a corresponding template.
 			if (PlayingTemplate.IsComboTemplate())
 			{
-				if(!PlayingTemplate.GetCombinedTemplateIds().Contains(SubTemplate.GetPageId()))
+				if (!PlayingTemplate.GetCombinedTemplateIds().Contains(SubTemplate.GetPageId()))
 				{
 					continue;
 				}
diff --git a/Engine/Plugins/Experimental/Avalanche/Source/AvalancheMedia/Private/Rundown/AvaRundownServer.cpp b/Engine/Plugins/Experimental/Avalanche/Source/AvalancheMedia/Private/Rundown/AvaRundownServer.cpp
--- a/Engine/Plugins/Experimental/Avalanche/Source/AvalancheMedia/Private/Rundown/AvaRundownServer.cpp
+++ b/Engine/Plugins/Experimental/Avalanche/Source/AvalancheMedia/Private/Rundown/AvaRundownServer.cpp
@@ -399,6 +399,7 @@ void FAvaRundownServer::Init(const FString& InAssignedHostName)
 	MessageEndpoint = FMessageEndpoint::Builder("MotionDesignRundownServer")
 	.Handling<FAvaRundownPing>(this, &FAvaRundownServer::HandleRundownPing)
 	.Handling<FAvaRundownGetServerInfo>(this, &FAvaRundownServer::HandleGetRundownServerInfo)
+	.Handling<FAvaRundownGetPlayableAssets>(this, &FAvaRundownServer::HandleGetPlayableAssets)
 	.Handling<FAvaRundownGetRundowns>(this, &FAvaRundownServer::HandleGetRundowns)
 	.Handling<FAvaRundownLoadRundown>(this, &FAvaRundownServer::HandleLoadRundown)
 	.Handling<FAvaRundownCreateRundown>(this, &FAvaRundownServer::HandleCreateRundown)
@@ -409,11 +410,13 @@ void FAvaRundownServer::Init(const FString& InAssignedHostName)
 	.Handling<FAvaRundownCreatePage>(this, &FAvaRundownServer::HandleCreatePage)
 	.Handling<FAvaRundownDeletePage>(this, &FAvaRundownServer::HandleDeletePage)
 	.Handling<FAvaRundownCreateTemplate>(this, &FAvaRundownServer::HandleCreateTemplate)
+	.Handling<FAvaRundownCreateComboTemplate>(this, &FAvaRundownServer::HandleCreateComboTemplate)
 	.Handling<FAvaRundownDeleteTemplate>(this, &FAvaRundownServer::HandleDeleteTemplate)
 	.Handling<FAvaRundownChangeTemplateBP>(this, &FAvaRundownServer::HandleChangeTemplateBP)
 	.Handling<FAvaRundownGetPages>(this, &FAvaRundownServer::HandleGetPages)
 	.Handling<FAvaRundownGetPageDetails>(this, &FAvaRundownServer::HandleGetPageDetails)
 	.Handling<FAvaRundownPageChangeChannel>(this, &FAvaRundownServer::HandleChangePageChannel)
+	.Handling<FAvaRundownChangePageName>(this, &FAvaRundownServer::HandleChangePageName)
 	.Handling<FAvaRundownUpdatePageFromRCP>(this, &FAvaRundownServer::HandleUpdatePageFromRCP)
 	.Handling<FAvaRundownPageAction>(this, &FAvaRundownServer::HandlePageAction)
 	.Handling<FAvaRundownPagePreviewAction>(this, &FAvaRundownServer::HandlePagePreviewAction)
@@ -524,46 +527,33 @@ void FAvaRundownServer::OnPageListChanged(const FAvaRundownPageListChangeParams&
 
 void FAvaRundownServer::OnPagesChanged(const UAvaRundown* InRundown, const FAvaRundownPage& InPage, EAvaRundownPageChanges InChange) const
 {
-	switch (InChange)
+	if (EnumHasAnyFlags(InChange, EAvaRundownPageChanges::AnimationSettings))
 	{
-	case EAvaRundownPageChanges::AnimationSettings:
-		{
-			PageAnimSettingsChanged(InRundown, InPage);
-			break;
-		}
-	case EAvaRundownPageChanges::Blueprint:
-		{
-			PageBlueprintChanged(InRundown, InPage, InPage.GetAssetPath(InRundown).ToString());
-			break;
-		}
-	case EAvaRundownPageChanges::Status:
-		{
-			PageStatusChanged(InRundown, InPage);
-			break;
-		}
-	case EAvaRundownPageChanges::Channel:
-		{
-			PageChannelChanged(InRundown, InPage, InPage.GetChannelName().ToString());
-			break;
-		}
-	case EAvaRundownPageChanges::RemoteControlValues:
-		{
-			break;
-		}
-	case EAvaRundownPageChanges::All:
-		{
-			break;
-		}
-	case EAvaRundownPageChanges::None:
-		{
-			break;
-		}
-	default:
-		{
-			break;
-		}
+		PageAnimSettingsChanged(InRundown, InPage);
 	}
-	
+	if (EnumHasAnyFlags(InChange, EAvaRundownPageChanges::Blueprint))
+	{
+		PageBlueprintChanged(InRundown, InPage, InPage.GetAssetPath(InRundown).ToString());
+	}
+	if (EnumHasAnyFlags(InChange, EAvaRundownPageChanges::Status))
+	{
+		PageStatusChanged(InRundown, InPage);
+	}
+	if (EnumHasAnyFlags(InChange, EAvaRundownPageChanges::Channel))
+	{
+		PageChannelChanged(InRundown, InPage, InPage.GetChannelName().ToString());
+	}
+	if (EnumHasAnyFlags(InChange, EAvaRundownPageChanges::Name))
+	{
+		PageNameChanged(InRundown, InPage, /*bInFriendlyName*/ false);
+	}
+	if (EnumHasAnyFlags(InChange, EAvaRundownPageChanges::FriendlyName))
+	{
+		PageNameChanged(InRundown, InPage, /*bInFriendlyName*/ true);
+	}
+
+	// todo: EAvaRundownPageChanges::RemoteControlValues
+	// -> tbd: rundown server api doesn't expose RC value directly.
 }
 void FAvaRundownServer::PageStatusChanged(const UAvaRundown* InRundown, const FAvaRundownPage& InPage) const
 {
@@ -591,6 +581,16 @@ void FAvaRundownServer::PageChannelChanged(const UAvaRundown* InRundown, const F
 	SendResponse(ReplyMessage, ClientAddresses);
 }
 
+void FAvaRundownServer::PageNameChanged(const UAvaRundown* InRundown, const FAvaRundownPage& InPage, bool bInFriendlyName) const
+{
+	FAvaRundownPageNameChanged* ReplyMessage = FMessageEndpoint::MakeMessage<FAvaRundownPageNameChanged>();
+	ReplyMessage->Rundown = FSoftObjectPath(InRundown).ToString();
+	ReplyMessage->PageId = InPage.GetPageId();
+	ReplyMessage->PageName = bInFriendlyName ? InPage.GetPageFriendlyName().ToString() : InPage.GetPageName();
+	ReplyMessage->bFriendlyName = bInFriendlyName;
+	SendResponse(ReplyMessage, ClientAddresses);
+}
+
 void FAvaRundownServer::PageAnimSettingsChanged(const UAvaRundown* InRundown, const FAvaRundownPage& InPage) const
 {
 	FAvaRundownPageAnimSettingsChanged* ReplyMessage = FMessageEndpoint::MakeMessage<FAvaRundownPageAnimSettingsChanged>();
@@ -721,6 +721,48 @@ void FAvaRundownServer::HandleGetRundownServerInfo(const FAvaRundownGetServerInf
 	SendResponse(ServerInfo, InContext->GetSender());
 }
 
+void FAvaRundownServer::HandleGetPlayableAssets(const FAvaRundownGetPlayableAssets& InMessage, const TSharedRef<IMessageContext, ESPMode::ThreadSafe>& InContext)
+{
+	FAvaRundownPlayableAssets* Response = FMessageEndpoint::MakeMessage<FAvaRundownPlayableAssets>();
+	Response->RequestId = InMessage.RequestId;
+	
+	if (const IAssetRegistry* AssetRegistry = IAssetRegistry::Get())
+	{
+		FARFilter Filter;
+		
+		// todo: Add all supported playable asset types. Hardcoded for now, need a extensible factory system.
+		Filter.ClassPaths.Add(UWorld::StaticClass()->GetClassPathName());
+		
+		TArray<FAssetData> Assets;
+		AssetRegistry->GetAssets(Filter, Assets);
+
+		for (const FAssetData& AssetData : Assets)
+		{
+			if (FAvaPlaybackUtils::IsPlayableAsset(AssetData))
+			{
+				if (!InMessage.Query.IsEmpty())
+				{
+					if (AssetData.AssetName.ToString().Contains(*InMessage.Query))
+					{
+						Response->Assets.Add(AssetData.ToSoftObjectPath());
+					}
+				}
+				else
+				{
+					Response->Assets.Add(AssetData.ToSoftObjectPath());
+				}
+
+				if (InMessage.Limit > 0 && Response->Assets.Num() >= InMessage.Limit)
+				{
+					break;
+				}
+			}
+		}
+	}
+
+	SendResponse(Response, InContext->GetSender());
+}
+
 void FAvaRundownServer::HandleGetRundowns(const FAvaRundownGetRundowns& InMessage, const TSharedRef<IMessageContext, ESPMode::ThreadSafe>& InContext)
 {
 	FAvaRundownRundowns* ReplyMessage = FMessageEndpoint::MakeMessage<FAvaRundownRundowns>();
@@ -1190,19 +1232,52 @@ void FAvaRundownServer::HandleCreateTemplate(const FAvaRundownCreateTemplate& In
 	{
 		return; // Response sent by GetOrLoadRundownForEdit.
 	}
+
+	const FAvaRundownPageIdGeneratorParams IdGeneratorParams(InMessage.IdGeneratorParams.ReferenceId, InMessage.IdGeneratorParams.Increment);
+	const FString& AssetPath = InMessage.AssetPath;
+	FString ErrorString;
+
+	// Note: using AddTemplateInternal because it doesn't add the template if there is an error and we can capture the error in the lambda.
+	const int32 TemplateId = Rundown->AddTemplateInternal(IdGeneratorParams, [&AssetPath, &ErrorString](FAvaRundownPage& InNewTemplate)
+	{
+		if (!AssetPath.IsEmpty())
+		{
+			constexpr bool bInReimportPage = true; // Ensures the asset is updated in the template.
+			if (!InNewTemplate.UpdateAsset(AssetPath, bInReimportPage))
+			{
+				ErrorString = FString::Printf(TEXT("asset \"%s\" is invalid"), *AssetPath);
+				return false;
+			}
+		}
+		return true;
+	});
 	
-	const int32 TemplateId = Rundown->AddTemplate();
 	if (TemplateId != FAvaRundownPage::InvalidPageId)
 	{
 		SendMessage(InContext->GetSender(), InMessage.RequestId, ELogVerbosity::Log, TEXT("Template %d Created"), TemplateId);
 	}
+	else if (!ErrorString.IsEmpty())
+	{
+		SendMessage(InContext->GetSender(), InMessage.RequestId, ELogVerbosity::Error, TEXT("Failed to create a new template: %s"), *ErrorString);
+	}
 	else
 	{
 		SendMessage(InContext->GetSender(), InMessage.RequestId, ELogVerbosity::Error, TEXT("Failed to create a new template"));
 	}
 }
 
-void FAvaRundownServer::HandleChangeTemplateBP(const FAvaRundownChangeTemplateBP& InMessage,
+class FAvaRundownServerErrorContext : public FOutputDevice
+{
+public:
+	TArray<FString> Errors;
+	
+	virtual void Serialize(const TCHAR* InText, ELogVerbosity::Type InVerbosity, const class FName& InCategory) override
+	{
+		Errors.Add(InText);
+	}
+};
+
+void FAvaRundownServer::HandleCreateComboTemplate(const FAvaRundownCreateComboTemplate& InMessage,
 	const TSharedRef<IMessageContext, ESPMode::ThreadSafe>& InContext)
 {
 	UAvaRundown* Rundown = GetOrLoadRundownForEdit(InContext->GetSender(), InMessage.RequestId, InMessage.Rundown);
@@ -1211,17 +1286,65 @@ void FAvaRundownServer::HandleChangeTemplateBP(const FAvaRundownChangeTemplateBP
 		return; // Response sent by GetOrLoadRundownForEdit.
 	}
 
-	FAvaRundownPage& Page = Rundown->GetPage(InMessage.TemplateId);
-	if (Page.IsValidPage() && Page.IsTemplate())
+	FAvaRundownServerErrorContext ErrorContext;
+	const TArray<int32> TemplateIds = Rundown->ValidateTemplateIdsForComboTemplate(InMessage.CombinedTemplateIds, ErrorContext);
+
+	if (TemplateIds.Num() > 2)
 	{
-		if (Page.UpdateAsset(InMessage.AssetPath))
+		const FAvaRundownPageIdGeneratorParams IdGeneratorParams(InMessage.IdGeneratorParams.ReferenceId, InMessage.IdGeneratorParams.Increment);
+		const int32 TemplateId = Rundown->AddComboTemplate(TemplateIds, IdGeneratorParams);
+	
+		if (TemplateId != FAvaRundownPage::InvalidPageId)
 		{
-			SendMessage(InContext->GetSender(), InMessage.RequestId, ELogVerbosity::Log, TEXT("Asset change of template: %d to %s"), InMessage.TemplateId, *InMessage.AssetPath);
-			Rundown->GetOnPagesChanged().Broadcast(Rundown, Page, EAvaRundownPageChanges::Blueprint);
-			return;
+			SendMessage(InContext->GetSender(), InMessage.RequestId, ELogVerbosity::Log, TEXT("Template %d Created"), TemplateId);
+		}
+		else
+		{
+			SendMessage(InContext->GetSender(), InMessage.RequestId, ELogVerbosity::Error, TEXT("Failed to create a new combo template"));
 		}
 	}
-	LogAndSendMessage(InContext->GetSender(), InMessage.RequestId, ELogVerbosity::Error, TEXT("Asset change of template: %d to %s failed."), InMessage.TemplateId, *InMessage.AssetPath);
+	else
+	{
+		SendMessage(InContext->GetSender(), InMessage.RequestId, ELogVerbosity::Error,
+			TEXT("Need at least 2 suitable templates to create a combo template: %s"), *FString::Join(ErrorContext.Errors, TEXT("; ")));
+	}
+}
+
+void FAvaRundownServer::HandleChangeTemplateBP(const FAvaRundownChangeTemplateBP& InMessage,
+	const TSharedRef<IMessageContext, ESPMode::ThreadSafe>& InContext)
+{
+	UAvaRundown* Rundown = GetOrLoadRundownForEdit(InContext->GetSender(), InMessage.RequestId, InMessage.Rundown);
+	if (!Rundown)
+	{
+		return; // Response sent by GetOrLoadRundownForEdit.
+	}
+
+	FAvaRundownPage& Page = Rundown->GetPage(InMessage.TemplateId);
+
+	if (!Page.IsValidPage())
+	{
+		SendMessage(InContext->GetSender(), InMessage.RequestId, ELogVerbosity::Error,
+			TEXT("Asset change of template failed: PageId %d is not a valid page."), InMessage.TemplateId);
+		return;
+	}
+
+	if (!Page.IsTemplate())
+	{
+		SendMessage(InContext->GetSender(), InMessage.RequestId, ELogVerbosity::Error,
+			TEXT("Asset change of template failed: PageId %d is not a template."), InMessage.TemplateId);
+		return;
+	}
+	
+	if (Page.UpdateAsset(InMessage.AssetPath, InMessage.bReimport))
+	{
+		SendMessage(InContext->GetSender(), InMessage.RequestId, ELogVerbosity::Log, TEXT("Asset change of template: %d to %s"), InMessage.TemplateId, *InMessage.AssetPath);
+		Rundown->GetOnPagesChanged().Broadcast(Rundown, Page, EAvaRundownPageChanges::Blueprint);
+	}
+	else
+	{
+		SendMessage(InContext->GetSender(), InMessage.RequestId, ELogVerbosity::Error,
+			TEXT("Asset change of template: %d to %s failed."), InMessage.TemplateId, *InMessage.AssetPath);	
+	}
 }
 
 void FAvaRundownServer::HandleGetPageDetails(const FAvaRundownGetPageDetails& InMessage, const TSharedRef<IMessageContext, ESPMode::ThreadSafe>& InContext)
@@ -1328,6 +1451,36 @@ void FAvaRundownServer::HandleChangePageChannel(const FAvaRundownPageChangeChann
 	SendMessage(InContext->GetSender(), InMessage.RequestId, ELogVerbosity::Log, TEXT("Channel changed"));
 }
 
+void FAvaRundownServer::HandleChangePageName(const FAvaRundownChangePageName& InMessage, const TSharedRef<IMessageContext, ESPMode::ThreadSafe>& InContext)
+{
+	UAvaRundown* Rundown = GetOrLoadRundownForEdit(InContext->GetSender(), InMessage.RequestId, InMessage.Rundown);
+	if (!Rundown)
+	{
+		return; // Response sent by GetOrLoadRundownForEdit.
+	}
+
+	FAvaRundownPage& Page = Rundown->GetPage(InMessage.PageId);
+	if (!Page.IsValidPage())
+	{
+		LogAndSendMessage(InContext->GetSender(), InMessage.RequestId, ELogVerbosity::Error, TEXT("PageId %d is invalid."), InMessage.PageId);
+		return;
+	}
+
+	if (InMessage.bSetFriendlyName)
+	{
+		Page.RenameFriendlyName(InMessage.PageName);
+		Rundown->GetOnPagesChanged().Broadcast(Rundown, Page, EAvaRundownPageChanges::FriendlyName);
+	}
+	else
+	{
+		Page.Rename(InMessage.PageName);
+		Rundown->GetOnPagesChanged().Broadcast(Rundown, Page, EAvaRundownPageChanges::Name);
+	}
+
+	SendMessage(InContext->GetSender(), InMessage.RequestId, ELogVerbosity::Log,
+		TEXT("Page %d Name Changed to \"%s\""), InMessage.PageId, *InMessage.PageName);
+}
+
 void FAvaRundownServer::HandleUpdatePageFromRCP(const FAvaRundownUpdatePageFromRCP& InMessage, const TSharedRef<IMessageContext, ESPMode::ThreadSafe>& InContext)
 {
 	// Note that this doesn't save the rundown.
@@ -2094,7 +2247,6 @@ void FAvaRundownServer::RegisterConsoleCommands()
 			FConsoleCommandWithArgsDelegate::CreateRaw(this, &FAvaRundownServer::ShowStatusCommand),
 			ECVF_Default
 			));
-
 }
 
 void FAvaRundownServer::ShowStatusCommand(const TArray<FString>& InArgs)
diff --git a/Engine/Plugins/Experimental/Avalanche/Source/AvalancheMedia/Private/Rundown/AvaRundownServer.h b/Engine/Plugins/Experimental/Avalanche/Source/AvalancheMedia/Private/Rundown/AvaRundownServer.h
--- a/Engine/Plugins/Experimental/Avalanche/Source/AvalancheMedia/Private/Rundown/AvaRundownServer.h
+++ b/Engine/Plugins/Experimental/Avalanche/Source/AvalancheMedia/Private/Rundown/AvaRundownServer.h
@@ -51,6 +51,7 @@ public:
 	void PageBlueprintChanged(const UAvaRundown* InRundown, const FAvaRundownPage& InPage, const FString& InBlueprintPath) const;
 	void PageStatusChanged(const UAvaRundown* InRundown, const FAvaRundownPage& InPage) const;
 	void PageChannelChanged(const UAvaRundown* InRundown, const FAvaRundownPage& InPage, const FString& InChannelName) const;
+	void PageNameChanged(const UAvaRundown* InRundown, const FAvaRundownPage& InPage, bool bInFriendlyName) const;
 	void PageAnimSettingsChanged(const UAvaRundown* InRundown, const FAvaRundownPage& InPage) const;
 	void OnBroadcastChannelListChanged(const FAvaBroadcastProfile& InProfile) const;
 	void OnBroadcastChannelChanged(const FAvaBroadcastOutputChannel& InChannel, EAvaBroadcastChannelChange InChange) const;
@@ -66,6 +67,7 @@ public:
 	// Rundown message handlers
 	void HandleRundownPing(const FAvaRundownPing& InMessage, const TSharedRef<IMessageContext, ESPMode::ThreadSafe>& InContext);
 	void HandleGetRundownServerInfo(const FAvaRundownGetServerInfo& InMessage, const TSharedRef<IMessageContext, ESPMode::ThreadSafe>& InContext);
+	void HandleGetPlayableAssets(const FAvaRundownGetPlayableAssets& InMessage, const TSharedRef<IMessageContext, ESPMode::ThreadSafe>& InContext);
 	void HandleGetRundowns(const FAvaRundownGetRundowns& InMessage, const TSharedRef<IMessageContext, ESPMode::ThreadSafe>& InContext);
 	void HandleLoadRundown(const FAvaRundownLoadRundown& InMessage, const TSharedRef<IMessageContext, ESPMode::ThreadSafe>& InContext);
 	void HandleCreateRundown(const FAvaRundownCreateRundown& InMessage, const TSharedRef<IMessageContext, ESPMode::ThreadSafe>& InContext);
@@ -76,11 +78,13 @@ public:
 	void HandleGetPages(const FAvaRundownGetPages& InMessage, const TSharedRef<IMessageContext, ESPMode::ThreadSafe>& InContext);
 	void HandleCreatePage(const FAvaRundownCreatePage& InMessage,const TSharedRef<IMessageContext, ESPMode::ThreadSafe>& InContext);
 	void HandleCreateTemplate(const FAvaRundownCreateTemplate& InMessage, const TSharedRef<IMessageContext, ESPMode::ThreadSafe>& InContext);
+	void HandleCreateComboTemplate(const FAvaRundownCreateComboTemplate& InMessage, const TSharedRef<IMessageContext, ESPMode::ThreadSafe>& InContext);
 	void HandleDeletePage(const FAvaRundownDeletePage& InMessage,const TSharedRef<IMessageContext, ESPMode::ThreadSafe>& InContext);
 	void HandleDeleteTemplate(const FAvaRundownDeleteTemplate& InMessage,const TSharedRef<IMessageContext, ESPMode::ThreadSafe>& InContext);
 	void HandleChangeTemplateBP(const FAvaRundownChangeTemplateBP& InMessage, const TSharedRef<IMessageContext, ESPMode::ThreadSafe>& InContext);
 	void HandleGetPageDetails(const FAvaRundownGetPageDetails& InMessage, const TSharedRef<IMessageContext, ESPMode::ThreadSafe>& InContext);
 	void HandleChangePageChannel(const FAvaRundownPageChangeChannel& InMessage, const TSharedRef<IMessageContext, ESPMode::ThreadSafe>& InContext);
+	void HandleChangePageName(const FAvaRundownChangePageName& InMessage, const TSharedRef<IMessageContext, ESPMode::ThreadSafe>& InContext);
 	void HandleUpdatePageFromRCP(const FAvaRundownUpdatePageFromRCP& InMessage, const TSharedRef<IMessageContext, ESPMode::ThreadSafe>& InContext);
 	void HandlePageAction(const FAvaRundownPageAction& InMessage, const TSharedRef<IMessageContext, ESPMode::ThreadSafe>& InContext);
 	void HandlePagePreviewAction(const FAvaRundownPagePreviewAction& InMessage, const TSharedRef<IMessageContext, ESPMode::ThreadSafe>& InContext);
diff --git a/Engine/Plugins/Experimental/Avalanche/Source/AvalancheMedia/Public/AvaMediaDefines.h b/Engine/Plugins/Experimental/Avalanche/Source/AvalancheMedia/Public/AvaMediaDefines.h
--- a/Engine/Plugins/Experimental/Avalanche/Source/AvalancheMedia/Public/AvaMediaDefines.h
+++ b/Engine/Plugins/Experimental/Avalanche/Source/AvalancheMedia/Public/AvaMediaDefines.h
@@ -250,6 +250,8 @@ enum class EAvaRundownPageChanges : uint8
 	Blueprint				= 1 << 2,
 	Channel					= 1 << 3,
 	Status					= 1 << 4,
+	Name					= 1 << 5,
+	FriendlyName			= 1 << 6,
 
 	All              = 0xFF,
 };
diff --git a/Engine/Plugins/Experimental/Avalanche/Source/AvalancheMedia/Public/Rundown/AvaRundown.h b/Engine/Plugins/Experimental/Avalanche/Source/AvalancheMedia/Public/Rundown/AvaRundown.h
--- a/Engine/Plugins/Experimental/Avalanche/Source/AvalancheMedia/Public/Rundown/AvaRundown.h
+++ b/Engine/Plugins/Experimental/Avalanche/Source/AvalancheMedia/Public/Rundown/AvaRundown.h
@@ -352,8 +352,24 @@ public:
 	/** Add empty template. */
 	int32 AddTemplate(const FAvaRundownPageIdGeneratorParams& InIdGeneratorParams = FAvaRundownPageIdGeneratorParams());
 
+	/**
+	 * Creates a combination template which combines other templates.
+	 * @remark This function does not validate the templates, use ValidateTemplateIdsForComboTemplate independently.
+	 */ 
 	int32 AddComboTemplate(const TArray<int32>& InTemplateIds, const FAvaRundownPageIdGeneratorParams& InIdGeneratorParams = FAvaRundownPageIdGeneratorParams());
 
+	/**
+	 * Combo Templates have specific requirement for the sub-templates.
+	 * This function validates a given list of templates to use as sub-templates for a combo template.
+	 * A combination template can only be created using transition logic templates that are on different
+	 * transition layers.
+	 * 
+	 * @param InTemplateIds Templates to validate.
+	 * @param InErrorLog Used to log the reasons for invalid templates. 
+	 * @return validated template Ids.
+	 */
+	TArray<int32> ValidateTemplateIdsForComboTemplate(const TArray<int32>& InTemplateIds, FOutputDevice& InErrorLog);
+
 	/**
 	 * @brief Add templates from existing source.
 	 * @param InSourceTemplates Source template to add.
diff --git a/Engine/Plugins/Experimental/Avalanche/Source/AvalancheMediaEditor/Private/Rundown/DetailsView/SAvaRundownPageDetails.cpp b/Engine/Plugins/Experimental/Avalanche/Source/AvalancheMediaEditor/Private/Rundown/DetailsView/SAvaRundownPageDetails.cpp
--- a/Engine/Plugins/Experimental/Avalanche/Source/AvalancheMediaEditor/Private/Rundown/DetailsView/SAvaRundownPageDetails.cpp
+++ b/Engine/Plugins/Experimental/Avalanche/Source/AvalancheMediaEditor/Private/Rundown/DetailsView/SAvaRundownPageDetails.cpp
@@ -233,24 +233,16 @@ void SAvaRundownPageDetails::OnManagedInstanceCacheEntryInvalidated(const FSoftO
 {
 	if (!bRefreshSelectedPageQueued)
 	{
-		if (TSharedPtr<FAvaRundownEditor> RundownEditor = RundownEditorWeak.Pin())
+		const FRundownContext RundownContext = GetRundownContext();
+		const FAvaRundownPage& SelectedPage = GetSelectedPage(RundownContext);
+		if (SelectedPage.IsValidPage())
 		{
-			UAvaRundown* Rundown = RundownEditor->GetRundown();
-
-			if (IsValid(Rundown))
+			if (SelectedPage.GetAssetPath(RundownContext.GetRundown()) == InAssetPath)
 			{
-				const FAvaRundownPage& SelectedPage = GetSelectedPage();
-
-				if (SelectedPage.IsValidPage())
-				{
-					if (SelectedPage.GetAssetPath(Rundown) == InAssetPath)
-					{
-						// Queue a refresh on next tick.
-						// We don't want to refresh immediately to avoid issues with
-						// cascading events within the managed instance cache.
-						QueueUpdateAndRefreshSelectedPage();
-					}
-				}
+				// Queue a refresh on next tick.
+				// We don't want to refresh immediately to avoid issues with
+				// cascading events within the managed instance cache.
+				QueueUpdateAndRefreshSelectedPage();
 			}
 		}
 	}
@@ -288,25 +280,38 @@ const FSlateBrush* SAvaRundownPageDetails::GetExposedPropertiesVisibilityBrush()
 	}
 }
 
-const FAvaRundownPage& SAvaRundownPageDetails::GetSelectedPage() const
+bool SAvaRundownPageDetails::FRundownContext::IsValid() const
 {
-	const TSharedPtr<FAvaRundownEditor> RundownEditor = RundownEditorWeak.Pin();
+	return RundownEditor && RundownEditor->IsRundownValid(); 
+}
 
-	if (HasSelectedPage() && RundownEditor && RundownEditor->IsRundownValid())
+UAvaRundown* SAvaRundownPageDetails::FRundownContext::GetRundown() const
+{
+	return IsValid() ? RundownEditor->GetRundown() : nullptr; 
+}
+
+SAvaRundownPageDetails::FRundownContext SAvaRundownPageDetails::GetRundownContext() const
+{
+	FRundownContext RundownContext;
+	RundownContext.RundownEditor = RundownEditorWeak.Pin();
+	return RundownContext;
+}
+
+const FAvaRundownPage& SAvaRundownPageDetails::GetSelectedPage(const FRundownContext& InRundownContext) const
+{
+	if (HasSelectedPage() && InRundownContext.IsValid())
 	{
-		return RundownEditor->GetRundown()->GetPage(ActivePageId);
+		return InRundownContext.GetRundown()->GetPage(ActivePageId);
 	}
 
 	return FAvaRundownPage::NullPage;
 }
 
-FAvaRundownPage& SAvaRundownPageDetails::GetMutableSelectedPage() const
+FAvaRundownPage& SAvaRundownPageDetails::GetMutableSelectedPage(const FRundownContext& InRundownContext) const
 {
-	const TSharedPtr<FAvaRundownEditor> RundownEditor = RundownEditorWeak.Pin();
-
-	if (HasSelectedPage() && RundownEditor && RundownEditor->IsRundownValid())
+	if (HasSelectedPage() && InRundownContext.IsValid())
 	{
-		return RundownEditor->GetRundown()->GetPage(ActivePageId);
+		return InRundownContext.GetRundown()->GetPage(ActivePageId);
 	}
 
 	return FAvaRundownPage::NullPage;
@@ -322,7 +327,8 @@ void SAvaRundownPageDetails::QueueRefreshSelectedPage()
 	
 	FTSTicker::GetCoreTicker().AddTicker(FTickerDelegate::CreateSPLambda(this, [this](float)
 	{
-		const FAvaRundownPage& SelectedPage = GetSelectedPage();
+		const FRundownContext RundownContext = GetRundownContext();
+		const FAvaRundownPage& SelectedPage = GetSelectedPage(RundownContext);
 		if (SelectedPage.IsValidPage())
 		{
 			RemoteControlProps->Refresh({SelectedPage.GetPageId()});
@@ -343,7 +349,8 @@ void SAvaRundownPageDetails::QueueUpdateAndRefreshSelectedPage()
 
 	FTSTicker::GetCoreTicker().AddTicker(FTickerDelegate::CreateSPLambda(this, [this](float)
 	{
-		const FAvaRundownPage& SelectedPage = GetSelectedPage();
+		const FRundownContext RundownContext = GetRundownContext();
+		const FAvaRundownPage& SelectedPage = GetSelectedPage(RundownContext);
 		if (SelectedPage.IsValidPage())
 		{
 			RemoteControlProps->UpdateDefaultValuesAndRefresh({SelectedPage.GetPageId()});
@@ -366,7 +373,8 @@ bool SAvaRundownPageDetails::HasSelectedPage() const
 
 FText SAvaRundownPageDetails::GetPageId() const
 {
-	const FAvaRundownPage& SelectedPage = GetSelectedPage();
+	const FRundownContext RundownContext = GetRundownContext();
+	const FAvaRundownPage& SelectedPage = GetSelectedPage(RundownContext);
 
 	if (SelectedPage.IsValidPage())
 	{
@@ -395,33 +403,31 @@ void SAvaRundownPageDetails::OnPageIdCommitted(const FText& InNewText, ETextComm
 		return;
 	}
 
-	FAvaRundownPage& SelectedPage = GetMutableSelectedPage();
+	const FRundownContext RundownContext = GetRundownContext();
 
-	if (!SelectedPage.IsValidPage()) // Not FAvaRundownPage::NullPage
+	if (!RundownContext.IsValid())
 	{
 		return;
 	}
 
-	int32 NewId = FCString::Atoi(*InNewText.ToString());
+	const FAvaRundownPage& SelectedPage = GetSelectedPage(RundownContext);
 
-	if (NewId == SelectedPage.GetPageId())
+	if (!SelectedPage.IsValidPage()) // Not FAvaRundownPage::NullPage
 	{
 		return;
 	}
 
-	const TSharedPtr<FAvaRundownEditor> RundownEditor = RundownEditorWeak.Pin();
+	const int32 NewId = FCString::Atoi(*InNewText.ToString());
 
-	if (!RundownEditor.IsValid() || !IsValid(RundownEditor->GetRundown()))
+	if (NewId == SelectedPage.GetPageId())
 	{
 		return;
 	}
 
-	const bool bSuccessful = RundownEditor->GetRundown()->RenumberPageId(SelectedPage.GetPageId(), NewId);
-
-	if (bSuccessful)
+	if (RundownContext.GetRundown()->RenumberPageId(SelectedPage.GetPageId(), NewId))
 	{
-		TSharedPtr<SAvaRundownInstancedPageList> PageList = RundownEditor->GetActiveListWidget();
-		
+		const TSharedPtr<SAvaRundownInstancedPageList> PageList = RundownContext.RundownEditor->GetActiveListWidget();
+
 		if (PageList.IsValid())
 		{
 			PageList->SelectPage(NewId);
@@ -431,7 +437,8 @@ void SAvaRundownPageDetails::OnPageIdCommitted(const FText& InNewText, ETextComm
 
 FText SAvaRundownPageDetails::GetPageDescription() const
 {
-	const FAvaRundownPage& SelectedPage = GetSelectedPage();
+	const FRundownContext RundownContext = GetRundownContext();
+	const FAvaRundownPage& SelectedPage = GetSelectedPage(RundownContext);
 
 	if (SelectedPage.IsValidPage())
 	{
@@ -443,38 +450,43 @@ FText SAvaRundownPageDetails::GetPageDescription() const
 
 void SAvaRundownPageDetails::OnPageNameChanged(const FText& InNewText)
 {
-	FAvaRundownPage& SelectedPage = GetMutableSelectedPage();
+	const FRundownContext RundownContext = GetRundownContext();
+	FAvaRundownPage& SelectedPage = GetMutableSelectedPage(RundownContext);
 
 	if (SelectedPage.IsValidPage()) // Not FAvaRundownPage::NullPage
 	{
 		SelectedPage.SetPageFriendlyName(InNewText);
+		
+		if (UAvaRundown* Rundown = RundownContext.GetRundown())
+		{
+			Rundown->GetOnPagesChanged().Broadcast(Rundown, SelectedPage, EAvaRundownPageChanges::FriendlyName);
+		}
 	}
 }
 
 FReply SAvaRundownPageDetails::DuplicateSelectedPage()
 {
-	FAvaRundownPage& SelectedPage = GetMutableSelectedPage();
-
-	if (!SelectedPage.IsValidPage()) // Not FAvaRundownPage::NullPage
+	const FRundownContext RundownContext = GetRundownContext();
+	if (!RundownContext.IsValid())
 	{
 		return FReply::Unhandled();
 	}
+	
+	const FAvaRundownPage& SelectedPage = GetSelectedPage(RundownContext);
 
-	const TSharedPtr<FAvaRundownEditor> RundownEditor = RundownEditorWeak.Pin();
-
-	if (!RundownEditor.IsValid())
+	if (!SelectedPage.IsValidPage()) // Not FAvaRundownPage::NullPage
 	{
 		return FReply::Unhandled();
 	}
 
-	TSharedPtr<SAvaRundownInstancedPageList> PageList = RundownEditor->GetActiveListWidget();
+	const TSharedPtr<SAvaRundownInstancedPageList> PageList = RundownContext.RundownEditor->GetActiveListWidget();
 
 	if (!PageList.IsValid())
 	{
 		return FReply::Unhandled();
 	}
 
-	TArray<int32> SelectedPages = TArray<int32>(PageList->GetSelectedPageIds());
+	const TArray<int32> SelectedPages = PageList->GetSelectedPageIds();
 	PageList->SelectPage(SelectedPage.GetPageId());
 	PageList->DuplicateSelectedPages();
 	PageList->SelectPages(SelectedPages);
diff --git a/Engine/Plugins/Experimental/Avalanche/Source/AvalancheMediaEditor/Private/Rundown/DetailsView/SAvaRundownPageDetails.h b/Engine/Plugins/Experimental/Avalanche/Source/AvalancheMediaEditor/Private/Rundown/DetailsView/SAvaRundownPageDetails.h
--- a/Engine/Plugins/Experimental/Avalanche/Source/AvalancheMediaEditor/Private/Rundown/DetailsView/SAvaRundownPageDetails.h
+++ b/Engine/Plugins/Experimental/Avalanche/Source/AvalancheMediaEditor/Private/Rundown/DetailsView/SAvaRundownPageDetails.h
@@ -39,8 +39,18 @@ protected:
 
 	const FSlateBrush* GetExposedPropertiesVisibilityBrush() const;
 
-	const FAvaRundownPage& GetSelectedPage() const;
-	FAvaRundownPage& GetMutableSelectedPage() const;
+	/** Encapsulate how the rundown is accessed for this widget to reduce duplicated code. */
+	struct FRundownContext
+	{
+		bool IsValid() const;
+		UAvaRundown* GetRundown() const;
+
+		TSharedPtr<FAvaRundownEditor> RundownEditor;
+	};
+	FRundownContext GetRundownContext() const;
+	
+	const FAvaRundownPage& GetSelectedPage(const FRundownContext& InRundownContext) const;
+	FAvaRundownPage& GetMutableSelectedPage(const FRundownContext& InRundownContext) const;
 
 	void QueueRefreshSelectedPage();
 	void QueueUpdateAndRefreshSelectedPage();
diff --git a/Engine/Plugins/Experimental/Avalanche/Source/AvalancheMediaEditor/Private/Rundown/Pages/PageViews/AvaRundownPageViewImpl.cpp b/Engine/Plugins/Experimental/Avalanche/Source/AvalancheMediaEditor/Private/Rundown/Pages/PageViews/AvaRundownPageViewImpl.cpp
--- a/Engine/Plugins/Experimental/Avalanche/Source/AvalancheMediaEditor/Private/Rundown/Pages/PageViews/AvaRundownPageViewImpl.cpp
+++ b/Engine/Plugins/Experimental/Avalanche/Source/AvalancheMediaEditor/Private/Rundown/Pages/PageViews/AvaRundownPageViewImpl.cpp
@@ -162,6 +162,7 @@ bool FAvaRundownPageViewImpl::Rename(const FText& InNewName)
 		Rundown->Modify();
 		
 		Page.Rename(InNewName.ToString());
+		Rundown->GetOnPagesChanged().Broadcast(Rundown, Page, EAvaRundownPageChanges::Name);
 		return true;
 	}
 	return false;
@@ -182,6 +183,7 @@ bool FAvaRundownPageViewImpl::RenameFriendlyName(const FText& InNewName)
 		Rundown->Modify();
 		
 		Page.RenameFriendlyName(InNewName.ToString());
+		Rundown->GetOnPagesChanged().Broadcast(Rundown, Page, EAvaRundownPageChanges::FriendlyName);
 		return true;
 	}
 	return false;
diff --git a/Engine/Plugins/Experimental/Avalanche/Source/AvalancheMediaEditor/Private/Rundown/Pages/Slate/SAvaRundownTemplatePageList.cpp b/Engine/Plugins/Experimental/Avalanche/Source/AvalancheMediaEditor/Private/Rundown/Pages/Slate/SAvaRundownTemplatePageList.cpp
--- a/Engine/Plugins/Experimental/Avalanche/Source/AvalancheMediaEditor/Private/Rundown/Pages/Slate/SAvaRundownTemplatePageList.cpp
+++ b/Engine/Plugins/Experimental/Avalanche/Source/AvalancheMediaEditor/Private/Rundown/Pages/Slate/SAvaRundownTemplatePageList.cpp
@@ -408,6 +408,24 @@ void SAvaRundownTemplatePageList::CreateInstance()
 	}
 }
 
+class FAvaRundownTemplatePageListErrorContext : public FOutputDevice
+{
+public:
+	const TCHAR* ContextName;
+
+	FAvaRundownTemplatePageListErrorContext(const TCHAR* InContextName)
+		: FOutputDevice()
+		, ContextName(InContextName)
+	{
+	}
+
+	virtual void Serialize(const TCHAR* InText, ELogVerbosity::Type InVerbosity, const class FName& InCategory) override
+	{
+		UE_LOG(LogAvaRundown, Error, TEXT("%s: %s"), ContextName, InText);
+	}
+};
+
+
 void SAvaRundownTemplatePageList::CreateComboTemplate()
 {
 	UAvaRundown* Rundown = GetValidRundown();
@@ -417,55 +435,8 @@ void SAvaRundownTemplatePageList::CreateComboTemplate()
 		return;
 	}
 
-	TSet<FAvaTagId> LayerIds;
-	TArray<int32> TemplateIds;
-	FAvaPlayableRemoteControlValues MergedValues;
-	
-	for (const int32 SelectedPageId : SelectedPageIds)
-	{
-		const FAvaRundownPage& Page = Rundown->GetPage(SelectedPageId);
-		if (!Page.IsValidPage())
-		{
-			UE_LOG(LogAvaRundown, Error, TEXT("CreateComboTemplate: Template %d is not valid."), SelectedPageId);
-			continue;
-		}
-		if (!Page.IsTemplate())
-		{
-			UE_LOG(LogAvaRundown, Error, TEXT("CreateComboTemplate: Page %d is not a template."), SelectedPageId);
-			continue;
-		}
-		if (Page.IsComboTemplate())
-		{
-			UE_LOG(LogAvaRundown, Error, TEXT("CreateComboTemplate: Template %d is already a combo template."), SelectedPageId);
-			continue;
-		}
-		if (!Page.HasTransitionLogic(Rundown))
-		{
-			UE_LOG(LogAvaRundown, Error, TEXT("CreateComboTemplate: Template %d doesn't have transition logic."), SelectedPageId);
-			continue;
-		}
-		if (!Page.GetTransitionLayer(Rundown).IsValid())
-		{
-			UE_LOG(LogAvaRundown, Error, TEXT("CreateComboTemplate: Template %d doesn't have a valid transition logic layer."), SelectedPageId);
-			continue;
-		}
-		if (LayerIds.Contains(Page.GetTransitionLayer(Rundown).TagId))
-		{
-			UE_LOG(LogAvaRundown, Error, TEXT("CreateComboTemplate: Template %d's layer %s is already in the selection."), SelectedPageId, *Page.GetTransitionLayer(Rundown).ToString());
-			continue;
-		}
-
-		// Make sure the RC values can merge correctly. If not, original template needs fixing.
-		if (MergedValues.HasIdCollisions(Page.GetRemoteControlValues()))
-		{
-			UE_LOG(LogAvaRundown, Error, TEXT("CreateComboTemplate: Template %d's RemoteControl values have Id collisions with other templates in the selection."), SelectedPageId);
-			continue;
-		}
-		
-		MergedValues.Merge(Page.GetRemoteControlValues());
-		LayerIds.Add(Page.GetTransitionLayer(Rundown).TagId);
-		TemplateIds.Add(Page.GetPageId());
-	}
+	FAvaRundownTemplatePageListErrorContext ErrorContext(TEXT("CreateComboTemplate"));
+	const TArray<int32> TemplateIds = Rundown->ValidateTemplateIdsForComboTemplate(SelectedPageIds, ErrorContext);
 
 	if (TemplateIds.Num() > 1)
 	{
@@ -473,7 +444,7 @@ void SAvaRundownTemplatePageList::CreateComboTemplate()
 	}
 	else
 	{
-		UE_LOG(LogAvaRundown, Error, TEXT("CreateComboTemplate: Need at least 2 suitable templates to create a combo template."));
+		ErrorContext.Log(TEXT("Need at least 2 suitable templates to create a combo template."));
 	}
 }
 
