[MotionDesign] Stagger Tool: Change validity check to enable auto-update toggling.

diff --git a/Engine/Plugins/Experimental/Avalanche/Source/AvalancheSequencer/Private/StaggerTool/Widgets/SAvaStaggerTool.cpp b/Engine/Plugins/Experimental/Avalanche/Source/AvalancheSequencer/Private/StaggerTool/Widgets/SAvaStaggerTool.cpp
--- a/Engine/Plugins/Experimental/Avalanche/Source/AvalancheSequencer/Private/StaggerTool/Widgets/SAvaStaggerTool.cpp
+++ b/Engine/Plugins/Experimental/Avalanche/Source/AvalancheSequencer/Private/StaggerTool/Widgets/SAvaStaggerTool.cpp
@@ -261,7 +261,7 @@ FReply SAvaStaggerTool::OnResetToDefaultsClick()
 
 void SAvaStaggerTool::OnToggleAutoUpdateClick(const ECheckBoxState InNewState)
 {
-	if (WeakTool.IsValid())
+	if (!WeakTool.IsValid())
 	{
 		return;
 	}
