Motion Design: Added support for child nodes within exposed properties to show in the Rundown Editor.
Additionally added support for functions to be shown/called in the Rundown page.

diff --git a/Engine/Plugins/Experimental/Avalanche/Source/AvalancheMediaEditor/Private/Rundown/DetailsView/RemoteControl/Properties/AvaRundownPagePropertyContextMenu.cpp b/Engine/Plugins/Experimental/Avalanche/Source/AvalancheMediaEditor/Private/Rundown/DetailsView/RemoteControl/Properties/AvaRundownPagePropertyContextMenu.cpp
--- a/Engine/Plugins/Experimental/Avalanche/Source/AvalancheMediaEditor/Private/Rundown/DetailsView/RemoteControl/Properties/AvaRundownPagePropertyContextMenu.cpp
+++ b/Engine/Plugins/Experimental/Avalanche/Source/AvalancheMediaEditor/Private/Rundown/DetailsView/RemoteControl/Properties/AvaRundownPagePropertyContextMenu.cpp
@@ -1,8 +1,8 @@
 // Copyright Epic Games, Inc. All Rights Reserved.
 
 #include "AvaRundownPagePropertyContextMenu.h"
-#include "AvaRundownRCPropertyItem.h"
 #include "IAvaMediaEditorModule.h"
+#include "Item/AvaRundownRCFieldItem.h"
 #include "Rundown/AvaRundownCommands.h"
 #include "Rundown/AvaRundownEditor.h"
 #include "Rundown/AvaRundownPage.h"
@@ -10,8 +10,8 @@
 #include "SAvaRundownPageRemoteControlProps.h"
 #include "ScopedTransaction.h"
 #include "ToolMenu.h"
-#include "ToolMenus.h"
 #include "ToolMenuSection.h"
+#include "ToolMenus.h"
 #include "UObject/WeakObjectPtrTemplates.h"
 
 #define LOCTEXT_NAMESPACE "AvaRundownPagePropertyContextMenu"
@@ -142,19 +142,19 @@ void FAvaRundownPagePropertyContextMenu::PopulatePageContextMenu(UToolMenu& InMe
 	}
 }
 
-TArray<TWeakPtr<FAvaRundownRCPropertyItem>> FAvaRundownPagePropertyContextMenu::GetSelectedPropertyItems() const
+TArray<TWeakPtr<FAvaRundownRCFieldItem>> FAvaRundownPagePropertyContextMenu::GetSelectedPropertyItems() const
 {
-	TArray<TWeakPtr<FAvaRundownRCPropertyItem>> OutItems;
+	TArray<TWeakPtr<FAvaRundownRCFieldItem>> OutItems;
 
 	if (CurrentContext.IsValid())
 	{
 		if (TSharedPtr<SAvaRundownPageRemoteControlProps> PropertyListWidget = CurrentContext->GetPropertyListWidget())
 		{
-			for (const FAvaRundownRCPropertyItemPtr& PropertyItem : PropertyListWidget->GetSelectedPropertyItems())
+			for (const TSharedPtr<FAvaRundownRCFieldItem>& FieldItem : PropertyListWidget->GetSelectedPropertyItems())
 			{
-				if (PropertyItem.IsValid())
+				if (FieldItem.IsValid())
 				{
-					OutItems.Add(PropertyItem);
+					OutItems.Add(FieldItem);
 				}
 			}
 		}
@@ -179,10 +179,10 @@ void FAvaRundownPagePropertyContextMenu::ResetValuesToDefaults(const bool bInUse
 	FScopedTransaction Transaction(LOCTEXT("ResetPropertyValuesTransaction", "Reset RC Property Values"));
 	Rundown->Modify();
 
-	const TArray<TWeakPtr<FAvaRundownRCPropertyItem>> SelectedPropertyItems = GetSelectedPropertyItems();
-	for (const TWeakPtr<FAvaRundownRCPropertyItem>& PropertyItemWeak : SelectedPropertyItems)
+	const TArray<TWeakPtr<FAvaRundownRCFieldItem>> SelectedPropertyItems = GetSelectedPropertyItems();
+	for (const TWeakPtr<FAvaRundownRCFieldItem>& PropertyItemWeak : SelectedPropertyItems)
 	{
-		if (const TSharedPtr<FAvaRundownRCPropertyItem> PropertyItem = PropertyItemWeak.Pin())
+		if (const TSharedPtr<FAvaRundownRCFieldItem> PropertyItem = PropertyItemWeak.Pin())
 		{
 			const TSharedPtr<FRemoteControlEntity> Entity = PropertyItem->GetEntity();
 			if (Entity.IsValid())
@@ -208,7 +208,7 @@ bool FAvaRundownPagePropertyContextMenu::CanResetValuesToDefaults(const bool bIn
 		return false;
 	}
 
-	const TArray<TWeakPtr<FAvaRundownRCPropertyItem>> SelectedItems = GetSelectedPropertyItems();
+	const TArray<TWeakPtr<FAvaRundownRCFieldItem>> SelectedItems = GetSelectedPropertyItems();
 	if (SelectedItems.Num() == 0)
 	{
 		return false;
@@ -216,9 +216,9 @@ bool FAvaRundownPagePropertyContextMenu::CanResetValuesToDefaults(const bool bIn
 
 	bool bContainsDifferentValues = false;
 
-	for (const TWeakPtr<FAvaRundownRCPropertyItem>& PropertyItemWeak : SelectedItems)
+	for (const TWeakPtr<FAvaRundownRCFieldItem>& PropertyItemWeak : SelectedItems)
 	{
-		const TSharedPtr<FAvaRundownRCPropertyItem> PropertyItem = PropertyItemWeak.Pin();
+		const TSharedPtr<FAvaRundownRCFieldItem> PropertyItem = PropertyItemWeak.Pin();
 		if (PropertyItem.IsValid())
 		{
 			if (PropertyItem->IsEntityControlled())
diff --git a/Engine/Plugins/Experimental/Avalanche/Source/AvalancheMediaEditor/Private/Rundown/DetailsView/RemoteControl/Properties/AvaRundownPagePropertyContextMenu.h b/Engine/Plugins/Experimental/Avalanche/Source/AvalancheMediaEditor/Private/Rundown/DetailsView/RemoteControl/Properties/AvaRundownPagePropertyContextMenu.h
--- a/Engine/Plugins/Experimental/Avalanche/Source/AvalancheMediaEditor/Private/Rundown/DetailsView/RemoteControl/Properties/AvaRundownPagePropertyContextMenu.h
+++ b/Engine/Plugins/Experimental/Avalanche/Source/AvalancheMediaEditor/Private/Rundown/DetailsView/RemoteControl/Properties/AvaRundownPagePropertyContextMenu.h
@@ -7,7 +7,7 @@
 #include "UObject/WeakObjectPtrTemplatesFwd.h"
 
 class FAvaRundownEditor;
-class FAvaRundownRCPropertyItem;
+class FAvaRundownRCFieldItem;
 class FUICommandList;
 class SAvaRundownPageRemoteControlProps;
 class SWidget;
@@ -34,7 +34,7 @@ private:
 	void ResetValuesToDefaults(const bool bInUseTemplateValues);
 	bool CanResetValuesToDefaults(const bool bInUseTemplateValues) const;
 
-	TArray<TWeakPtr<FAvaRundownRCPropertyItem>> GetSelectedPropertyItems() const;
+	TArray<TWeakPtr<FAvaRundownRCFieldItem>> GetSelectedPropertyItems() const;
 
 	UAvaRundown* GetContextRundown() const;
 
diff --git a/Engine/Plugins/Experimental/Avalanche/Source/AvalancheMediaEditor/Private/Rundown/DetailsView/RemoteControl/Properties/AvaRundownPageRemoteControlWidgetUtils.cpp b/Engine/Plugins/Experimental/Avalanche/Source/AvalancheMediaEditor/Private/Rundown/DetailsView/RemoteControl/Properties/AvaRundownPageRemoteControlWidgetUtils.cpp
--- a/Engine/Plugins/Experimental/Avalanche/Source/AvalancheMediaEditor/Private/Rundown/DetailsView/RemoteControl/Properties/AvaRundownPageRemoteControlWidgetUtils.cpp
+++ b/Engine/Plugins/Experimental/Avalanche/Source/AvalancheMediaEditor/Private/Rundown/DetailsView/RemoteControl/Properties/AvaRundownPageRemoteControlWidgetUtils.cpp
@@ -37,7 +37,7 @@ TSharedRef<SWidget> FAvaRundownPageRemoteControlWidgetUtils::CreateNodeValueWidg
 	return FieldWidget;
 }
 
-bool FAvaRundownPageRemoteControlWidgetUtils::FindPropertyHandleRecursive(const TSharedPtr<IPropertyHandle>& InPropertyHandle, const FString& InPropertyNameOrPath, EFindNodeMethod InFindMethod)
+bool FAvaRundownPageRemoteControlWidgetUtils::FindPropertyHandleRecursive(const TSharedPtr<IPropertyHandle>& InPropertyHandle, const FProperty* InProperty, EFindNodeMethod InFindMethod)
 {
 	if (InPropertyHandle && InPropertyHandle->IsValidHandle())
 	{
@@ -46,38 +46,28 @@ bool FAvaRundownPageRemoteControlWidgetUtils::FindPropertyHandleRecursive(const
 		for (uint32 Index = 0; Index < ChildrenCount; ++Index)
 		{
 			TSharedPtr<IPropertyHandle> ChildHandle = InPropertyHandle->GetChildHandle(Index);
-			if (FindPropertyHandleRecursive(ChildHandle, InPropertyNameOrPath, InFindMethod))
+			if (FindPropertyHandleRecursive(ChildHandle, InProperty, InFindMethod))
 			{
 				return true;
 			}
 		}
 
-		if (InPropertyHandle->GetProperty())
+		if (InPropertyHandle->GetProperty() == InProperty)
 		{
-			if (InFindMethod == EFindNodeMethod::Path)
-			{
-				if (InPropertyHandle->GeneratePathToProperty() == InPropertyNameOrPath)
-				{
-					return true;
-				}
-			}
-			else if (InPropertyHandle->GetProperty()->GetName() == InPropertyNameOrPath)
-			{
-				return true;
-			}
+			return true;
 		}
 	}
 
 	return false;
 }
 
-TSharedPtr<IDetailTreeNode> FAvaRundownPageRemoteControlWidgetUtils::FindTreeNodeRecursive(const TSharedRef<IDetailTreeNode>& InRootNode, const FString& InPropertyNameOrPath, EFindNodeMethod InFindMethod)
+TSharedPtr<IDetailTreeNode> FAvaRundownPageRemoteControlWidgetUtils::FindTreeNodeRecursive(const TSharedRef<IDetailTreeNode>& InRootNode, const FProperty* InProperty, EFindNodeMethod InFindMethod)
 {
 	TArray<TSharedRef<IDetailTreeNode>> Children;
 	InRootNode->GetChildren(Children);
 	for (TSharedRef<IDetailTreeNode>& Child : Children)
 	{
-		TSharedPtr<IDetailTreeNode> FoundNode = FindTreeNodeRecursive(Child, InPropertyNameOrPath, InFindMethod);
+		TSharedPtr<IDetailTreeNode> FoundNode = FindTreeNodeRecursive(Child, InProperty, InFindMethod);
 		if (FoundNode.IsValid())
 		{
 			return FoundNode;
@@ -85,7 +75,7 @@ TSharedPtr<IDetailTreeNode> FAvaRundownPageRemoteControlWidgetUtils::FindTreeNod
 	}
 
 	TSharedPtr<IPropertyHandle> Handle = InRootNode->CreatePropertyHandle();
-	if (FindPropertyHandleRecursive(Handle, InPropertyNameOrPath, InFindMethod))
+	if (FindPropertyHandleRecursive(Handle, InProperty, InFindMethod))
 	{
 		return InRootNode;
 	}
@@ -93,11 +83,11 @@ TSharedPtr<IDetailTreeNode> FAvaRundownPageRemoteControlWidgetUtils::FindTreeNod
 	return nullptr;
 }
 
-TSharedPtr<IDetailTreeNode> FAvaRundownPageRemoteControlWidgetUtils::FindNode(const TArray<TSharedRef<IDetailTreeNode>>& InRootNodes, const FString& InQualifiedPropertyName, EFindNodeMethod InFindMethod)
+TSharedPtr<IDetailTreeNode> FAvaRundownPageRemoteControlWidgetUtils::FindNode(const TArray<TSharedRef<IDetailTreeNode>>& InRootNodes, const FProperty* InProperty, EFindNodeMethod InFindMethod)
 {
 	for (const TSharedRef<IDetailTreeNode>& CategoryNode : InRootNodes)
 	{
-		TSharedPtr<IDetailTreeNode> FoundNode = FindTreeNodeRecursive(CategoryNode, InQualifiedPropertyName, InFindMethod);
+		TSharedPtr<IDetailTreeNode> FoundNode = FindTreeNodeRecursive(CategoryNode, InProperty, InFindMethod);
 		if (FoundNode.IsValid())
 		{
 			return FoundNode;
diff --git a/Engine/Plugins/Experimental/Avalanche/Source/AvalancheMediaEditor/Private/Rundown/DetailsView/RemoteControl/Properties/AvaRundownPageRemoteControlWidgetUtils.h b/Engine/Plugins/Experimental/Avalanche/Source/AvalancheMediaEditor/Private/Rundown/DetailsView/RemoteControl/Properties/AvaRundownPageRemoteControlWidgetUtils.h
--- a/Engine/Plugins/Experimental/Avalanche/Source/AvalancheMediaEditor/Private/Rundown/DetailsView/RemoteControl/Properties/AvaRundownPageRemoteControlWidgetUtils.h
+++ b/Engine/Plugins/Experimental/Avalanche/Source/AvalancheMediaEditor/Private/Rundown/DetailsView/RemoteControl/Properties/AvaRundownPageRemoteControlWidgetUtils.h
@@ -24,7 +24,7 @@ public:
 
 public:
 	static TSharedRef<SWidget> CreateNodeValueWidget(const TSharedPtr<IDetailTreeNode>& InNode);
-	static bool FindPropertyHandleRecursive(const TSharedPtr<IPropertyHandle>& InPropertyHandle, const FString& InPropertyNameOrPath, EFindNodeMethod InFindMethod);
-	static TSharedPtr<IDetailTreeNode> FindTreeNodeRecursive(const TSharedRef<IDetailTreeNode>& InRootNode, const FString& InPropertyNameOrPath, EFindNodeMethod InFindMethod);
-	static TSharedPtr<IDetailTreeNode> FindNode(const TArray<TSharedRef<IDetailTreeNode>>& InRootNodes, const FString& InQualifiedPropertyName, EFindNodeMethod InFindMethod);
+	static bool FindPropertyHandleRecursive(const TSharedPtr<IPropertyHandle>& InPropertyHandle, const FProperty* InProperty, EFindNodeMethod InFindMethod);
+	static TSharedPtr<IDetailTreeNode> FindTreeNodeRecursive(const TSharedRef<IDetailTreeNode>& InRootNode, const FProperty* InProperty, EFindNodeMethod InFindMethod);
+	static TSharedPtr<IDetailTreeNode> FindNode(const TArray<TSharedRef<IDetailTreeNode>>& InRootNodes, const FProperty* InProperty, EFindNodeMethod InFindMethod);
 };
diff --git a/Engine/Plugins/Experimental/Avalanche/Source/AvalancheMediaEditor/Private/Rundown/DetailsView/RemoteControl/Properties/AvaRundownRCPropertyItem.cpp b/Engine/Plugins/Experimental/Avalanche/Source/AvalancheMediaEditor/Private/Rundown/DetailsView/RemoteControl/Properties/AvaRundownRCPropertyItem.cpp
deleted file mode 100644
diff --git a/Engine/Plugins/Experimental/Avalanche/Source/AvalancheMediaEditor/Private/Rundown/DetailsView/RemoteControl/Properties/AvaRundownRCPropertyItem.h b/Engine/Plugins/Experimental/Avalanche/Source/AvalancheMediaEditor/Private/Rundown/DetailsView/RemoteControl/Properties/AvaRundownRCPropertyItem.h
deleted file mode 100644
diff --git a/Engine/Plugins/Experimental/Avalanche/Source/AvalancheMediaEditor/Private/Rundown/DetailsView/RemoteControl/Properties/Item/AvaRundownRCDetailTreeNodeItem.cpp b/Engine/Plugins/Experimental/Avalanche/Source/AvalancheMediaEditor/Private/Rundown/DetailsView/RemoteControl/Properties/Item/AvaRundownRCDetailTreeNodeItem.cpp
new file mode 100644
--- /dev/null
+++ b/Engine/Plugins/Experimental/Avalanche/Source/AvalancheMediaEditor/Private/Rundown/DetailsView/RemoteControl/Properties/Item/AvaRundownRCDetailTreeNodeItem.cpp
@@ -0,0 +1,85 @@
+// Copyright Epic Games, Inc. All Rights Reserved.
+
+#include "AvaRundownRCDetailTreeNodeItem.h"
+#include "IPropertyRowGenerator.h"
+#include "IRemoteControlPropertyHandle.h"
+#include "RemoteControlField.h"
+#include "Rundown/DetailsView/RemoteControl/Properties/AvaRundownPageRemoteControlWidgetUtils.h"
+#include "Rundown/DetailsView/RemoteControl/Properties/SAvaRundownPageRemoteControlProps.h"
+
+TSharedPtr<FAvaRundownRCDetailTreeNodeItem> FAvaRundownRCDetailTreeNodeItem::CreateItem(const TSharedRef<SAvaRundownPageRemoteControlProps>& InPropertyPanel
+	, const TSharedRef<FRemoteControlProperty>& InPropertyEntity
+	, bool bInControlled)
+{
+	TSharedPtr<IRemoteControlPropertyHandle> PropertyHandle = InPropertyEntity->GetPropertyHandle();
+	if (!PropertyHandle.IsValid())
+	{
+		return nullptr;
+	}
+
+	TArray<UObject*> Objects = InPropertyEntity->GetBoundObjects();
+	if (Objects.IsEmpty())
+	{
+		return nullptr;
+	}
+
+	TSharedRef<FAvaRundownRCDetailTreeNodeItem> PropertyItem = MakeShared<FAvaRundownRCDetailTreeNodeItem>();
+	PropertyItem->EntityOwnerWeak = InPropertyEntity;
+	PropertyItem->bEntityControlled = bInControlled;
+
+	TSharedRef<IPropertyRowGenerator> Generator = InPropertyPanel->FindOrAddGenerator(Objects[0]);
+
+	TSharedPtr<IDetailTreeNode> Node = FAvaRundownPageRemoteControlWidgetUtils::FindNode(Generator->GetRootTreeNodes()
+		, PropertyHandle->GetProperty()
+		, FAvaRundownPageRemoteControlWidgetUtils::EFindNodeMethod::Path);
+
+	PropertyItem->Initialize(Node);
+	return PropertyItem;
+}
+
+void FAvaRundownRCDetailTreeNodeItem::Initialize(TSharedPtr<IDetailTreeNode> InDetailTreeNode)
+{
+	DetailTreeNode = InDetailTreeNode;
+
+	if (DetailTreeNode.IsValid())
+	{
+		NodeWidgets = DetailTreeNode->CreateNodeWidgets();
+	}
+
+	RefreshChildren();
+}
+
+TSharedPtr<SWidget> FAvaRundownRCDetailTreeNodeItem::GetNameWidget() const
+{
+	return NodeWidgets.NameWidget;
+}
+
+TSharedPtr<SWidget> FAvaRundownRCDetailTreeNodeItem::GetValueWidget() const
+{
+	return NodeWidgets.WholeRowWidget.IsValid()
+		? NodeWidgets.WholeRowWidget
+		: NodeWidgets.ValueWidget;
+}
+
+void FAvaRundownRCDetailTreeNodeItem::RefreshChildren()
+{
+	Children.Reset();
+	if (!DetailTreeNode.IsValid())
+	{
+		return;
+	}
+
+	TArray<TSharedRef<IDetailTreeNode>> ChildNodes;
+	DetailTreeNode->GetChildren(ChildNodes);
+
+	Children.Reserve(ChildNodes.Num());
+
+	for (const TSharedRef<IDetailTreeNode>& ChildNode : ChildNodes)
+	{
+		TSharedRef<FAvaRundownRCDetailTreeNodeItem> ChildPropertyItem = MakeShared<FAvaRundownRCDetailTreeNodeItem>();
+		ChildPropertyItem->EntityOwnerWeak = EntityOwnerWeak;
+		ChildPropertyItem->bEntityControlled = bEntityControlled;
+		ChildPropertyItem->Initialize(ChildNode);
+		Children.Add(MoveTemp(ChildPropertyItem));
+	}
+}
diff --git a/Engine/Plugins/Experimental/Avalanche/Source/AvalancheMediaEditor/Private/Rundown/DetailsView/RemoteControl/Properties/Item/AvaRundownRCDetailTreeNodeItem.h b/Engine/Plugins/Experimental/Avalanche/Source/AvalancheMediaEditor/Private/Rundown/DetailsView/RemoteControl/Properties/Item/AvaRundownRCDetailTreeNodeItem.h
new file mode 100644
--- /dev/null
+++ b/Engine/Plugins/Experimental/Avalanche/Source/AvalancheMediaEditor/Private/Rundown/DetailsView/RemoteControl/Properties/Item/AvaRundownRCDetailTreeNodeItem.h
@@ -0,0 +1,32 @@
+// Copyright Epic Games, Inc. All Rights Reserved.
+
+#pragma once
+
+#include "AvaRundownRCFieldItem.h"
+#include "IDetailTreeNode.h"
+
+struct FRemoteControlProperty;
+
+class FAvaRundownRCDetailTreeNodeItem : public FAvaRundownRCFieldItem
+{
+	friend class FAvaRundownRCFieldItem;
+
+public:
+	static TSharedPtr<FAvaRundownRCDetailTreeNodeItem> CreateItem(const TSharedRef<SAvaRundownPageRemoteControlProps>& InPropertyPanel
+		, const TSharedRef<FRemoteControlProperty>& InPropertyEntity
+		, bool bInControlled);
+
+	void Initialize(TSharedPtr<IDetailTreeNode> InDetailTreeNode);
+
+	//~ Begin FAvaRundownRCFieldItem
+	virtual TSharedPtr<SWidget> GetNameWidget() const override;
+	virtual TSharedPtr<SWidget> GetValueWidget() const override;
+	//~ End FAvaRundownRCFieldItem
+
+private:
+	void RefreshChildren();
+
+	TSharedPtr<IDetailTreeNode> DetailTreeNode;
+
+	FNodeWidgets NodeWidgets;
+};
diff --git a/Engine/Plugins/Experimental/Avalanche/Source/AvalancheMediaEditor/Private/Rundown/DetailsView/RemoteControl/Properties/Item/AvaRundownRCFieldItem.cpp b/Engine/Plugins/Experimental/Avalanche/Source/AvalancheMediaEditor/Private/Rundown/DetailsView/RemoteControl/Properties/Item/AvaRundownRCFieldItem.cpp
new file mode 100644
--- /dev/null
+++ b/Engine/Plugins/Experimental/Avalanche/Source/AvalancheMediaEditor/Private/Rundown/DetailsView/RemoteControl/Properties/Item/AvaRundownRCFieldItem.cpp
@@ -0,0 +1,54 @@
+// Copyright Epic Games, Inc. All Rights Reserved.
+
+#include "AvaRundownRCFieldItem.h"
+#include "AvaRundownRCDetailTreeNodeItem.h"
+#include "AvaRundownRCFunctionItem.h"
+#include "RemoteControlEntity.h"
+#include "RemoteControlField.h"
+#include "Rundown/DetailsView/RemoteControl/Properties/SAvaRundownRCPropertyItemRow.h"
+
+TSharedPtr<FAvaRundownRCFieldItem> FAvaRundownRCFieldItem::CreateItem(const TSharedRef<SAvaRundownPageRemoteControlProps>& InPropertyPanel, const TSharedRef<FRemoteControlEntity>& InEntity, bool bInControlled)
+{
+	const UScriptStruct* EntityStruct = InEntity->GetStruct();
+	if (!EntityStruct || !EntityStruct->IsChildOf<FRemoteControlField>())
+	{
+		return nullptr;
+	}
+
+	TSharedRef<FRemoteControlField> FieldEntity = StaticCastSharedRef<FRemoteControlField>(InEntity);
+
+	if (FieldEntity->FieldType == EExposedFieldType::Function)
+	{
+		TSharedRef<FAvaRundownRCFunctionItem> FunctionItem = MakeShared<FAvaRundownRCFunctionItem>();
+		FunctionItem->EntityOwnerWeak = FieldEntity;
+		FunctionItem->bEntityControlled = bInControlled;
+		return FunctionItem;
+	}
+
+	if (FieldEntity->FieldType == EExposedFieldType::Property)
+	{
+		TSharedPtr<FAvaRundownRCDetailTreeNodeItem> DetailTreeNodeItem = FAvaRundownRCDetailTreeNodeItem::CreateItem(InPropertyPanel
+			, StaticCastSharedRef<FRemoteControlProperty>(InEntity)
+			, bInControlled);
+
+		// Reset Name Widget so that it defaults to the entity widget
+		if (DetailTreeNodeItem.IsValid())
+		{
+			DetailTreeNodeItem->NodeWidgets.NameWidget.Reset();
+		}
+
+		return DetailTreeNodeItem;
+	}
+
+	return nullptr;
+}
+
+TSharedRef<ITableRow> FAvaRundownRCFieldItem::CreateWidget(const TSharedRef<SAvaRundownPageRemoteControlProps>& InPropertyPanel, const TSharedRef<STableViewBase>& InOwnerTable)
+{
+	return SNew(SAvaRundownRCPropertyItemRow, InPropertyPanel, InOwnerTable, SharedThis(this));
+}
+
+TSharedPtr<SWidget> FAvaRundownRCFieldItem::GetNameWidget() const
+{
+	return nullptr;
+}
diff --git a/Engine/Plugins/Experimental/Avalanche/Source/AvalancheMediaEditor/Private/Rundown/DetailsView/RemoteControl/Properties/Item/AvaRundownRCFieldItem.h b/Engine/Plugins/Experimental/Avalanche/Source/AvalancheMediaEditor/Private/Rundown/DetailsView/RemoteControl/Properties/Item/AvaRundownRCFieldItem.h
new file mode 100644
--- /dev/null
+++ b/Engine/Plugins/Experimental/Avalanche/Source/AvalancheMediaEditor/Private/Rundown/DetailsView/RemoteControl/Properties/Item/AvaRundownRCFieldItem.h
@@ -0,0 +1,50 @@
+// Copyright Epic Games, Inc. All Rights Reserved.
+
+#pragma once
+
+#include "Internationalization/Text.h"
+#include "Templates/SharedPointer.h"
+
+class ITableRow;
+class SAvaRundownPageRemoteControlProps;
+class STableViewBase;
+class SWidget;
+struct FRemoteControlEntity;
+
+class FAvaRundownRCFieldItem : public TSharedFromThis<FAvaRundownRCFieldItem>
+{
+public:
+	static TSharedPtr<FAvaRundownRCFieldItem> CreateItem(const TSharedRef<SAvaRundownPageRemoteControlProps>& InPropertyPanel
+		, const TSharedRef<FRemoteControlEntity>& InEntity
+		, bool bInControlled);
+
+	virtual ~FAvaRundownRCFieldItem() = default;
+
+	TSharedPtr<FRemoteControlEntity> GetEntity() const
+	{
+		return EntityOwnerWeak.Pin();
+	}
+
+	bool IsEntityControlled() const
+	{
+		return bEntityControlled;
+	}
+
+	TConstArrayView<TSharedPtr<FAvaRundownRCFieldItem>> GetChildren() const
+	{
+		return Children;
+	}
+
+	TSharedRef<ITableRow> CreateWidget(const TSharedRef<SAvaRundownPageRemoteControlProps>& InPropertyPanel, const TSharedRef<STableViewBase>& InOwnerTable);
+
+	virtual TSharedPtr<SWidget> GetNameWidget() const;
+
+	virtual TSharedPtr<SWidget> GetValueWidget() const = 0;
+
+protected:
+	TWeakPtr<FRemoteControlEntity> EntityOwnerWeak;
+
+	bool bEntityControlled = false;
+
+	TArray<TSharedPtr<FAvaRundownRCFieldItem>> Children;
+};
diff --git a/Engine/Plugins/Experimental/Avalanche/Source/AvalancheMediaEditor/Private/Rundown/DetailsView/RemoteControl/Properties/Item/AvaRundownRCFunctionItem.cpp b/Engine/Plugins/Experimental/Avalanche/Source/AvalancheMediaEditor/Private/Rundown/DetailsView/RemoteControl/Properties/Item/AvaRundownRCFunctionItem.cpp
new file mode 100644
--- /dev/null
+++ b/Engine/Plugins/Experimental/Avalanche/Source/AvalancheMediaEditor/Private/Rundown/DetailsView/RemoteControl/Properties/Item/AvaRundownRCFunctionItem.cpp
@@ -0,0 +1,66 @@
+// Copyright Epic Games, Inc. All Rights Reserved.
+
+#include "AvaRundownRCFunctionItem.h"
+#include "RemoteControlField.h"
+#include "ScopedTransaction.h"
+#include "UObject/Script.h"
+#include "Widgets/Input/SButton.h"
+#include "Widgets/Text/STextBlock.h"
+
+#define LOCTEXT_NAMESPACE "AvaRundownRCFunctionItem"
+
+TSharedPtr<SWidget> FAvaRundownRCFunctionItem::GetValueWidget() const
+{
+	return SNew(SButton)
+		.VAlign(VAlign_Center)
+		.OnClicked(this, &FAvaRundownRCFunctionItem::OnFunctionButtonClicked)
+		[
+			SNew(STextBlock)
+			.Text(LOCTEXT("CallFunctionLabel", "Call Function"))
+		];
+}
+
+FReply FAvaRundownRCFunctionItem::OnFunctionButtonClicked() const
+{
+	TSharedPtr<FRemoteControlEntity> Entity = GetEntity();
+	if (!Entity.IsValid())
+	{
+		return FReply::Unhandled();
+	}
+
+	const UScriptStruct* EntityStruct = Entity->GetStruct();
+	if (!EntityStruct || !EntityStruct->IsChildOf<FRemoteControlFunction>())
+	{
+		return FReply::Unhandled();
+	}
+
+	TSharedRef<FRemoteControlFunction> FunctionEntity = StaticCastSharedRef<FRemoteControlFunction>(Entity.ToSharedRef());
+
+	FScopedTransaction Transaction(LOCTEXT("CallExposedFunction", "Called a function through Rundown."));
+	FEditorScriptExecutionGuard ScriptGuard;
+
+	bool bObjectsModified = false;
+
+	for (UObject* Object : FunctionEntity->GetBoundObjects())
+	{
+		if (FunctionEntity->FunctionArguments && FunctionEntity->FunctionArguments->IsValid())
+		{
+			Object->Modify();
+			Object->ProcessEvent(FunctionEntity->GetFunction(), FunctionEntity->FunctionArguments->GetStructMemory());
+			bObjectsModified = true;
+		}
+		else
+		{
+			ensureAlwaysMsgf(false, TEXT("Function default arguments could not be resolved."));
+		}
+	}
+
+	if (!bObjectsModified)
+	{
+		Transaction.Cancel();
+	}
+
+	return FReply::Handled();
+}
+
+#undef LOCTEXT_NAMESPACE
diff --git a/Engine/Plugins/Experimental/Avalanche/Source/AvalancheMediaEditor/Private/Rundown/DetailsView/RemoteControl/Properties/Item/AvaRundownRCFunctionItem.h b/Engine/Plugins/Experimental/Avalanche/Source/AvalancheMediaEditor/Private/Rundown/DetailsView/RemoteControl/Properties/Item/AvaRundownRCFunctionItem.h
new file mode 100644
--- /dev/null
+++ b/Engine/Plugins/Experimental/Avalanche/Source/AvalancheMediaEditor/Private/Rundown/DetailsView/RemoteControl/Properties/Item/AvaRundownRCFunctionItem.h
@@ -0,0 +1,18 @@
+// Copyright Epic Games, Inc. All Rights Reserved.
+
+#pragma once
+
+#include "AvaRundownRCFieldItem.h"
+
+class FReply;
+
+class FAvaRundownRCFunctionItem : public FAvaRundownRCFieldItem
+{
+public:
+	//~ Begin FAvaRundownRCFieldItem
+	virtual TSharedPtr<SWidget> GetValueWidget() const override;
+	//~ End FAvaRundownRCFieldItem
+
+private:
+	FReply OnFunctionButtonClicked() const;
+};
diff --git a/Engine/Plugins/Experimental/Avalanche/Source/AvalancheMediaEditor/Private/Rundown/DetailsView/RemoteControl/Properties/SAvaRundownPageRemoteControlProps.cpp b/Engine/Plugins/Experimental/Avalanche/Source/AvalancheMediaEditor/Private/Rundown/DetailsView/RemoteControl/Properties/SAvaRundownPageRemoteControlProps.cpp
--- a/Engine/Plugins/Experimental/Avalanche/Source/AvalancheMediaEditor/Private/Rundown/DetailsView/RemoteControl/Properties/SAvaRundownPageRemoteControlProps.cpp
+++ b/Engine/Plugins/Experimental/Avalanche/Source/AvalancheMediaEditor/Private/Rundown/DetailsView/RemoteControl/Properties/SAvaRundownPageRemoteControlProps.cpp
@@ -1,9 +1,11 @@
 // Copyright Epic Games, Inc. All Rights Reserved.
 
 #include "SAvaRundownPageRemoteControlProps.h"
-#include "AvaRundownRCPropertyItem.h"
+#include "Item/AvaRundownRCFieldItem.h"
+#include "Modules/ModuleManager.h"
 #include "Playable/AvaPlayableRemoteControl.h"
 #include "Playable/AvaPlayableRemoteControlValues.h"
+#include "PropertyEditorModule.h"
 #include "RemoteControlEntity.h"
 #include "RemoteControlPreset.h"
 #include "Rundown/AvaRundownEditor.h"
@@ -15,8 +17,8 @@
 #include "SlateOptMacros.h"
 #include "UObject/NameTypes.h"
 #include "Widgets/Views/ITableRow.h"
-#include "Widgets/Views/SListView.h"
 #include "Widgets/Views/STableViewBase.h"
+#include "Widgets/Views/STreeView.h"
 
 #define LOCTEXT_NAMESPACE "SAvaRundownPageRemoteControlProps"
 
@@ -84,10 +86,10 @@ void SAvaRundownPageRemoteControlProps::Construct(const FArguments& InArgs, TSha
 		SNew(SHeaderRow)
 		+ SHeaderRow::Column(PropertyColumnName)
 		.DefaultLabel(LOCTEXT("Property", "Property"))
-		.FixedWidth(150.f)
+		.FillWidth(0.2f)
 		+ SHeaderRow::Column(ValueColumnName)
 		.DefaultLabel(LOCTEXT("Value", "Value"))
-		.FillWidth(1.f);
+		.FillWidth(0.8f);
 
 	HeaderRowExtensionDelegate.Broadcast(SharedThis(this), HeaderRow);
 
@@ -97,11 +99,12 @@ void SAvaRundownPageRemoteControlProps::Construct(const FArguments& InArgs, TSha
 		.Padding(8.0f)
 		.BorderImage(FAppStyle::GetBrush("ToolPanel.GroupBorder"))
 		[
-			SAssignNew(PropertyContainer, SListView<FAvaRundownRCPropertyItemPtr>)
-			.ListItemsSource(&PropertyItems)
+			SAssignNew(FieldContainer, STreeView<TSharedPtr<FAvaRundownRCFieldItem>>)
+			.TreeItemsSource(&FieldItems)
 			.SelectionMode(ESelectionMode::Multi)
 			.OnContextMenuOpening(this, &SAvaRundownPageRemoteControlProps::GetContextMenuContent)
 			.OnGenerateRow(this, &SAvaRundownPageRemoteControlProps::OnGenerateControllerRow)
+			.OnGetChildren(this, &SAvaRundownPageRemoteControlProps::OnGetEntityChildren)
 			.HeaderRow(HeaderRow)
 		]
 	];
@@ -130,19 +133,27 @@ void SAvaRundownPageRemoteControlProps::UpdateDefaultValuesAndRefresh(const TArr
 	Refresh(InSelectedPageIds);
 }
 
-TSharedRef<ITableRow> SAvaRundownPageRemoteControlProps::OnGenerateControllerRow(FAvaRundownRCPropertyItemPtr InItem, const TSharedRef<STableViewBase>& InOwnerTable)
+TSharedRef<ITableRow> SAvaRundownPageRemoteControlProps::OnGenerateControllerRow(TSharedPtr<FAvaRundownRCFieldItem> InItem, const TSharedRef<STableViewBase>& InOwnerTable)
 {
 	return InItem->CreateWidget(SharedThis(this), InOwnerTable);
 }
 
+void SAvaRundownPageRemoteControlProps::OnGetEntityChildren(TSharedPtr<FAvaRundownRCFieldItem> InItem, TArray<TSharedPtr<FAvaRundownRCFieldItem>>& OutChildren)
+{
+	if (InItem.IsValid())
+	{
+		OutChildren.Append(InItem->GetChildren());
+	}
+}
+
 void SAvaRundownPageRemoteControlProps::RefreshTable(const TSet<FGuid>& InEntityIds)
 {
-	for (const FAvaRundownRCPropertyItemPtr& PropertyItem : PropertyItems)
+	for (const TSharedPtr<FAvaRundownRCFieldItem>& PropertyItem : FieldItems)
 	{
 		const TSharedPtr<FRemoteControlEntity> Entity = PropertyItem->GetEntity();
 		if (Entity && (InEntityIds.IsEmpty() || InEntityIds.Contains(Entity->GetId())))
 		{
-			TSharedPtr<ITableRow> TableRow = PropertyContainer->WidgetFromItem(PropertyItem);
+			TSharedPtr<ITableRow> TableRow = FieldContainer->WidgetFromItem(PropertyItem);
 
 			if (TableRow.IsValid())
 			{
@@ -155,7 +166,7 @@ void SAvaRundownPageRemoteControlProps::RefreshTable(const TSet<FGuid>& InEntity
 
 void SAvaRundownPageRemoteControlProps::Refresh(const TArray<int32>& InSelectedPageIds)
 {
-	if (!PropertyContainer.IsValid())
+	if (!FieldContainer.IsValid())
 	{
 		return;
 	}
@@ -217,19 +228,37 @@ void SAvaRundownPageRemoteControlProps::Refresh(const TArray<int32>& InSelectedP
 			
 			for (const TWeakPtr<FRemoteControlEntity>& EntityWeakPtr : ExposedEntities)
 			{
-				if (const TSharedPtr<FRemoteControlEntity> Entity = EntityWeakPtr.Pin())
+				const TSharedPtr<FRemoteControlEntity> Entity = EntityWeakPtr.Pin();
+				if (!Entity.IsValid())
+				{
+					continue;
+				}
+
+				const UScriptStruct* EntityStruct = Entity->GetStruct();
+				if (!EntityStruct || !EntityStruct->IsChildOf<FRemoteControlField>())
+				{
+					continue;
+				}
+
+				TSharedRef<FRemoteControlField> FieldEntity = StaticCastSharedRef<FRemoteControlField>(Entity.ToSharedRef());
+
+				if (FieldEntity->FieldType == EExposedFieldType::Function)
 				{
-					const FAvaPlayableRemoteControlValue* EntityValue = GetSelectedPageEntityValue(Entity);
+					NewItems.Add({ FieldEntity, false });
+				}
+				else if (FieldEntity->FieldType == EExposedFieldType::Property)
+				{
+					const FAvaPlayableRemoteControlValue* EntityValue = GetSelectedPageEntityValue(FieldEntity);
 
 					if (!EntityValue)
 					{
 						// If the page doesn't already have a value, we get it from the template's default values.
-						const FAvaPlayableRemoteControlValue* const DefaultEntityValue = ManagedInstance->GetDefaultRemoteControlValues().GetEntityValue(Entity->GetId()); 
+						const FAvaPlayableRemoteControlValue* const DefaultEntityValue = ManagedInstance->GetDefaultRemoteControlValues().GetEntityValue(FieldEntity->GetId()); 
 						
 						if (!DefaultEntityValue)
 						{
 							UE_LOG(LogAvaPlayableRemoteControl, Error, TEXT("Entity \"%s\" (id:%s) doesn't have a template default value."),
-								*Entity->GetLabel().ToString(), *Entity->GetId().ToString());
+								*FieldEntity->GetLabel().ToString(), *FieldEntity->GetId().ToString());
 							// TODO: UX improvement: instead of skipping, could add empty element, with error mark (and error message in tooltip).
 							continue;
 						}
@@ -239,10 +268,10 @@ void SAvaRundownPageRemoteControlProps::Refresh(const TArray<int32>& InSelectedP
 						
 						// WYSIWYG (Solution):
 						// Capture the default value (flagged as default) in the current page to ensure all values will be applied to runtime RCP.
-						if (!SetSelectedPageEntityValue(Entity, *DefaultEntityValue))
+						if (!SetSelectedPageEntityValue(FieldEntity, *DefaultEntityValue))
 						{
 							UE_LOG(LogAvaPlayableRemoteControl, Error, TEXT("Entity \"%s\" (id:%s): failed to set value in currently selected page."),
-								*Entity->GetLabel().ToString(), *Entity->GetId().ToString());
+								*FieldEntity->GetLabel().ToString(), *FieldEntity->GetId().ToString());
 						}
 
 						EntityValue = DefaultEntityValue;
@@ -250,48 +279,45 @@ void SAvaRundownPageRemoteControlProps::Refresh(const TArray<int32>& InSelectedP
 					
 					// Update Exposed entity value with value from page.
 					using namespace UE::AvaPlayableRemoteControl;
-					if (const EAvaPlayableRemoteControlResult Result = SetValueOfEntity(Entity, EntityValue->Value); Failed(Result))
+					if (const EAvaPlayableRemoteControlResult Result = SetValueOfEntity(FieldEntity, EntityValue->Value); Failed(Result))
 					{
 						UE_LOG(LogAvaPlayableRemoteControl, Error, TEXT("Entity \"%s\" (id:%s): failed to set entity value: %s."),
-							*Entity->GetLabel().ToString(), *Entity->GetId().ToString(), *EnumToString(Result));
+							*FieldEntity->GetLabel().ToString(), *FieldEntity->GetId().ToString(), *EnumToString(Result));
 						// TODO: UX improvement: instead of skipping, could add empty element, with error mark (and error message in tooltip).
 						continue;
 					}
 
-					{
-						const bool bEntityControlled = ManagedInstance->GetDefaultRemoteControlValues().EntitiesControlledByController.Contains(Entity->GetId());
-						NewItems.Add({Entity.ToSharedRef(), bEntityControlled});
-					}
+					const bool bEntityControlled = ManagedInstance->GetDefaultRemoteControlValues().EntitiesControlledByController.Contains(FieldEntity->GetId());
+					NewItems.Add({ FieldEntity, bEntityControlled });
 				}
 			}
 		}
 
-		PropertyContainer->RebuildList();
+		FieldContainer->RebuildList();
 
-		bool bRecreateList = NewItems.Num() != PropertyItems.Num();
+		bool bRecreateList = NewItems.Num() != FieldItems.Num();
 
 		if (!bRecreateList)
 		{
 			for (int32 PropertyIdx = 0; PropertyIdx < NewItems.Num(); ++PropertyIdx)
 			{
-				const FAvaRundownRCPropertyItemPtr& PropertyItem = PropertyItems[PropertyIdx];
+				const TSharedPtr<FAvaRundownRCFieldItem>& FieldItem = FieldItems[PropertyIdx];
 				const FAvaPropertyDetails& NewItem = NewItems[PropertyIdx];
 
-				if (!PropertyItem.IsValid())
+				if (!FieldItem.IsValid())
 				{
 					bRecreateList = true;
 					break;
 				}
 
-				TSharedPtr<FRemoteControlEntity> PropertyEntity = PropertyItem->GetEntity();
-
-				if (!PropertyEntity.IsValid())
+				TSharedPtr<FRemoteControlEntity> FieldEntity = FieldItem->GetEntity();
+				if (!FieldEntity.IsValid())
 				{
 					bRecreateList = true;
 					break;
 				}
 
-				if (PropertyEntity.Get() != NewItem.Entity.ToSharedPtr().Get())
+				if (FieldEntity.Get() != NewItem.Entity.ToSharedPtr().Get())
 				{
 					bRecreateList = true;
 					break;
@@ -305,11 +331,15 @@ void SAvaRundownPageRemoteControlProps::Refresh(const TArray<int32>& InSelectedP
 			return;
 		}
 
-		PropertyItems.Empty();
+		FieldItems.Empty(NewItems.Num());
 
+		TSharedRef<SAvaRundownPageRemoteControlProps> This = SharedThis(this);
 		for (const FAvaPropertyDetails& NewItem : NewItems)
 		{
-			PropertyItems.Add(MakeShared<FAvaRundownRCPropertyItem>(NewItem.Entity, NewItem.bEntityControlled));
+			if (TSharedPtr<FAvaRundownRCFieldItem> Item =  FAvaRundownRCFieldItem::CreateItem(This, NewItem.Entity, NewItem.bEntityControlled))
+			{
+				FieldItems.Add(MoveTemp(Item));
+			}
 		}
 	}
 }
@@ -471,9 +501,9 @@ bool SAvaRundownPageRemoteControlProps::SetSelectedPageEntityValue(const TShared
 	return false;
 }
 
-const TArray<FAvaRundownRCPropertyItemPtr> SAvaRundownPageRemoteControlProps::GetSelectedPropertyItems() const
+const TArray<TSharedPtr<FAvaRundownRCFieldItem>> SAvaRundownPageRemoteControlProps::GetSelectedPropertyItems() const
 {
-	return PropertyContainer->GetSelectedItems();
+	return FieldContainer->GetSelectedItems();
 }
 
 bool SAvaRundownPageRemoteControlProps::SaveRemoteControlEntitiesToPage(const URemoteControlPreset* InPreset, const TSet<FGuid>& InPropertyIds, UAvaRundown* InRundown, int32 InPageId)
@@ -526,9 +556,22 @@ TSharedPtr<FAvaRundownPageRCPropsNotifyHook> SAvaRundownPageRemoteControlProps::
 	return NotifyHook;
 }
 
+TSharedRef<IPropertyRowGenerator> SAvaRundownPageRemoteControlProps::FindOrAddGenerator(UObject* InObject)
+{
+	TSharedPtr<IPropertyRowGenerator>& Generator = ObjectGeneratorMap.FindOrAdd(InObject);
+	if (!Generator.IsValid())
+	{
+		FPropertyRowGeneratorArgs Args;
+		Args.NotifyHook = NotifyHook.Get();
+		Generator = FModuleManager::GetModuleChecked<FPropertyEditorModule>("PropertyEditor").CreatePropertyRowGenerator(Args);
+		Generator->SetObjects({ InObject });
+	}
+	return Generator.ToSharedRef();
+}
+
 TSharedPtr<SWidget> SAvaRundownPageRemoteControlProps::GetContextMenuContent()
 {
-	TArray<FAvaRundownRCPropertyItemPtr> SelectedItems = GetSelectedPropertyItems();
+	TArray<TSharedPtr<FAvaRundownRCFieldItem>> SelectedItems = GetSelectedPropertyItems();
 	if (SelectedItems.Num() > 0)
 	{
 		return ContextMenu->GeneratePageContextMenuWidget(RundownEditorWeak, *GetActivePage(), SharedThis(this));
diff --git a/Engine/Plugins/Experimental/Avalanche/Source/AvalancheMediaEditor/Private/Rundown/DetailsView/RemoteControl/Properties/SAvaRundownPageRemoteControlProps.h b/Engine/Plugins/Experimental/Avalanche/Source/AvalancheMediaEditor/Private/Rundown/DetailsView/RemoteControl/Properties/SAvaRundownPageRemoteControlProps.h
--- a/Engine/Plugins/Experimental/Avalanche/Source/AvalancheMediaEditor/Private/Rundown/DetailsView/RemoteControl/Properties/SAvaRundownPageRemoteControlProps.h
+++ b/Engine/Plugins/Experimental/Avalanche/Source/AvalancheMediaEditor/Private/Rundown/DetailsView/RemoteControl/Properties/SAvaRundownPageRemoteControlProps.h
@@ -7,6 +7,7 @@
 #include "Containers/Set.h"
 #include "Misc/NotifyHook.h"
 #include "Templates/SharedPointer.h"
+#include "UObject/ObjectKey.h"
 #include "Widgets/SCompoundWidget.h"
 #include "Widgets/Views/ITableRow.h"
 #include "Widgets/Views/SListView.h"
@@ -14,9 +15,11 @@
 class FAvaRundownEditor;
 class FAvaRundownManagedInstance;
 class FAvaRundownPagePropertyContextMenu;
-class FAvaRundownRCPropertyItem;
+class FAvaRundownRCDetailTreeNodeItem;
+class FAvaRundownRCFieldItem;
 class FName;
 class FUICommandList;
+class IPropertyRowGenerator;
 class ITableRow;
 class SAvaRundownPageRemoteControlProps;
 class STableViewBase;
@@ -28,12 +31,10 @@ struct FAvaRundownPage;
 struct FGuid;
 struct FRemoteControlEntity;
 
-using FAvaRundownRCPropertyItemPtr = TSharedPtr<FAvaRundownRCPropertyItem>;
-
 DECLARE_MULTICAST_DELEGATE_TwoParams(FAvaRundownRCPropertyHeaderRowExtensionDelegate, TSharedRef<SAvaRundownPageRemoteControlProps> Panel,
 	TSharedRef<SHeaderRow>& HeaderRow)
 DECLARE_DELEGATE_ThreeParams(FAvaRundownRCPropertyTableRowExtensionDelegate, TSharedRef<SAvaRundownPageRemoteControlProps> Panel,
-	TSharedRef<const FAvaRundownRCPropertyItem> ItemPtr, TSharedPtr<SWidget>& CurrentWidget)
+	TSharedRef<const FAvaRundownRCFieldItem> ItemPtr, TSharedPtr<SWidget>& CurrentWidget)
 
 /**
  * Note:
@@ -73,10 +74,12 @@ public:
 	/** Refreshes the content of this widget. */
 	void Refresh(const TArray<int32>& InSelectedPageIds);
 
-	const TArray<FAvaRundownRCPropertyItemPtr> GetSelectedPropertyItems() const;
+	const TArray<TSharedPtr<FAvaRundownRCFieldItem>> GetSelectedPropertyItems() const;
 	
 	TSharedPtr<FAvaRundownPageRCPropsNotifyHook> GetNotifyHook() const;
 
+	TSharedRef<IPropertyRowGenerator> FindOrAddGenerator(UObject* InObject);
+
 	/**
 	 * Utility to save the given Preset's Entities to the given page's property.
 	 * @param InPreset Source Preset the entity values are read from.
@@ -120,7 +123,9 @@ private:
 	 */
 	bool SetSelectedPageEntityValue(const TSharedPtr<FRemoteControlEntity>& InRemoteControlEntity, const FAvaPlayableRemoteControlValue& InValue) const;
 
-	TSharedRef<ITableRow> OnGenerateControllerRow(FAvaRundownRCPropertyItemPtr InItem, const TSharedRef<STableViewBase>& InOwnerTable);
+	TSharedRef<ITableRow> OnGenerateControllerRow(TSharedPtr<FAvaRundownRCFieldItem> InItem, const TSharedRef<STableViewBase>& InOwnerTable);
+
+	void OnGetEntityChildren(TSharedPtr<FAvaRundownRCFieldItem> InItem, TArray<TSharedPtr<FAvaRundownRCFieldItem>>& OutChildren);
 
 	void RefreshTable(const TSet<FGuid>& InEntityIds = TSet<FGuid>());
 
@@ -131,10 +136,10 @@ private:
 	TArray<TSharedPtr<FAvaRundownManagedInstance>> ManagedInstances;
 
 	/** The widget that lists the property rows. */
-	TSharedPtr<SListView<FAvaRundownRCPropertyItemPtr>> PropertyContainer;
+	TSharedPtr<SListView<TSharedPtr<FAvaRundownRCFieldItem>>> FieldContainer;
 	
 	/** The data used to back the properties container list view. */
-	TArray<FAvaRundownRCPropertyItemPtr> PropertyItems;
+	TArray<TSharedPtr<FAvaRundownRCFieldItem>> FieldItems;
 
 	int32 ActivePageId = -1;
 
@@ -144,5 +149,7 @@ private:
 
 	TSharedPtr<FAvaRundownPageRCPropsNotifyHook> NotifyHook;
 
+	TMap<FObjectKey, TSharedPtr<IPropertyRowGenerator>> ObjectGeneratorMap;
+
 	friend class FAvaRundownPageRCPropsNotifyHookImpl;
 };
\ No newline at end of file
diff --git a/Engine/Plugins/Experimental/Avalanche/Source/AvalancheMediaEditor/Private/Rundown/DetailsView/RemoteControl/Properties/SAvaRundownRCPropertyItemRow.cpp b/Engine/Plugins/Experimental/Avalanche/Source/AvalancheMediaEditor/Private/Rundown/DetailsView/RemoteControl/Properties/SAvaRundownRCPropertyItemRow.cpp
--- a/Engine/Plugins/Experimental/Avalanche/Source/AvalancheMediaEditor/Private/Rundown/DetailsView/RemoteControl/Properties/SAvaRundownRCPropertyItemRow.cpp
+++ b/Engine/Plugins/Experimental/Avalanche/Source/AvalancheMediaEditor/Private/Rundown/DetailsView/RemoteControl/Properties/SAvaRundownRCPropertyItemRow.cpp
@@ -2,11 +2,11 @@
 
 #include "SAvaRundownRCPropertyItemRow.h"
 #include "AvaRundownPageRemoteControlWidgetUtils.h"
-#include "AvaRundownRCPropertyItem.h"
 #include "GameFramework/Actor.h"
 #include "IDetailTreeNode.h"
-#include "Internationalization/Text.h"
 #include "IPropertyRowGenerator.h"
+#include "Internationalization/Text.h"
+#include "Item/AvaRundownRCFieldItem.h"
 #include "Modules/ModuleManager.h"
 #include "PropertyEditorModule.h"
 #include "RemoteControlEntity.h"
@@ -17,31 +17,39 @@
 
 #define LOCTEXT_NAMESPACE "SAvaRundownRCPropertyItemRow"
 
-void SAvaRundownRCPropertyItemRow::Construct(const FArguments& InArgs, TSharedRef<SAvaRundownPageRemoteControlProps> InPropertyPanel,
-	const TSharedRef<STableViewBase>& InOwnerTableView, const TSharedPtr<const FAvaRundownRCPropertyItem>& InRowItem)
+void SAvaRundownRCPropertyItemRow::Construct(const FArguments& InArgs, TSharedRef<SAvaRundownPageRemoteControlProps> InPropertyPanel,const TSharedRef<STableViewBase>& InOwnerTableView, const TSharedPtr<const FAvaRundownRCFieldItem>& InRowItem)
 {
-	ItemPtrWeak = InRowItem;
+	FieldItemWeak = InRowItem;
 	PropertyPanelWeak = InPropertyPanel;
 	NotifyHook = InPropertyPanel->GetNotifyHook();
-	Generator = nullptr;
 	ValueContainer = nullptr;
-	ValueWidget = nullptr;
 
-	SMultiColumnTableRow<FAvaRundownRCPropertyItemPtr>::Construct(FSuperRowType::FArguments(), InOwnerTableView);
+	SMultiColumnTableRow<TSharedPtr<FAvaRundownRCFieldItem>>::Construct(FSuperRowType::FArguments(), InOwnerTableView);
 }
 
 TSharedRef<SWidget> SAvaRundownRCPropertyItemRow::GenerateWidgetForColumn(const FName& InColumnName)
 {
-	TSharedPtr<const FAvaRundownRCPropertyItem> ItemPtr = ItemPtrWeak.Pin();
+	TSharedPtr<const FAvaRundownRCFieldItem> ItemPtr = FieldItemWeak.Pin();
 
 	if (ItemPtr.IsValid())
 	{
 		if (InColumnName == SAvaRundownPageRemoteControlProps::PropertyColumnName)
 		{
-			return SNew(STextBlock)
-				.Margin(FMargin(8.f, 2.f, 0.f, 2.f))
-				.Text(GetFieldLabel())
-				.ToolTipText(this, &SAvaRundownRCPropertyItemRow::GetPropertyTooltipText);
+			return SNew(SHorizontalBox)
+				+ SHorizontalBox::Slot()
+				.AutoWidth()
+				.VAlign(VAlign_Fill)
+				[
+					SNew(SExpanderArrow, SharedThis(this))
+					.IndentAmount(12)
+					.ShouldDrawWires(true)
+				]
+				+ SHorizontalBox::Slot()
+				.HAlign(HAlign_Left)
+				.VAlign(VAlign_Center)
+				[
+					CreateName()
+				];
 		}
 		else if (InColumnName == SAvaRundownPageRemoteControlProps::ValueColumnName)
 		{
@@ -83,74 +91,69 @@ void SAvaRundownRCPropertyItemRow::UpdateValue()
 	}
 }
 
-FText SAvaRundownRCPropertyItemRow::GetFieldLabel() const
+TSharedRef<SWidget> SAvaRundownRCPropertyItemRow::CreateName()
 {
-	if (TSharedPtr<const FAvaRundownRCPropertyItem> ItemPtr = ItemPtrWeak.Pin())
+	TSharedPtr<const FAvaRundownRCFieldItem> FieldItem = FieldItemWeak.Pin();
+	if (!FieldItem.IsValid())
+	{
+		return SNullWidget::NullWidget;
+	}
+
+	TSharedPtr<SWidget> NameWidget = FieldItem->GetNameWidget();
+	if (!NameWidget.IsValid())
 	{
-		if (TSharedPtr<FRemoteControlEntity> EntityPtr = ItemPtr->GetEntity())
+		if (TSharedPtr<FRemoteControlEntity> Entity = FieldItem->GetEntity())
 		{
-			TSharedPtr<FRemoteControlField> FieldPtr = StaticCastSharedPtr<FRemoteControlField>(EntityPtr);
-			return FText::FromName(FieldPtr->FieldName);
+			TSharedPtr<FRemoteControlField> FieldEntity = StaticCastSharedPtr<FRemoteControlField>(Entity);
+			return SNew(STextBlock)
+				.Margin(FMargin(8.f, 2.f, 0.f, 2.f))
+				.Text(FText::FromName(FieldEntity->GetLabel()))
+				.ToolTipText(this, &SAvaRundownRCPropertyItemRow::GetPropertyTooltipText);
 		}
 	}
 
-	return FText::GetEmpty();
+	NameWidget->SetToolTipText(TAttribute<FText>::CreateSP(this, &SAvaRundownRCPropertyItemRow::GetPropertyTooltipText));
+	return NameWidget.ToSharedRef();
 }
 
 TSharedRef<SWidget> SAvaRundownRCPropertyItemRow::CreateValue()
 {
-	if (TSharedPtr<const FAvaRundownRCPropertyItem> ItemPtr = ItemPtrWeak.Pin())
+	TSharedPtr<const FAvaRundownRCFieldItem> FieldItem = FieldItemWeak.Pin();
+	if (!FieldItem.IsValid())
 	{
-		if (TSharedPtr<FRemoteControlEntity> EntityPtr = ItemPtr->GetEntity())
-		{
-			TSharedPtr<FRemoteControlField> FieldPtr = StaticCastSharedPtr<FRemoteControlField>(EntityPtr);
-
-			// For the moment, just use the first object.
-			TArray<UObject*> Objects = FieldPtr->GetBoundObjects();
+		return SNullWidget::NullWidget;
+	}
 
-			if ((FieldPtr->FieldType == EExposedFieldType::Property) && (Objects.Num() > 0))
-			{
-				FPropertyRowGeneratorArgs Args;
-				Args.NotifyHook = NotifyHook.Get();
-				Generator = FModuleManager::GetModuleChecked<FPropertyEditorModule>("PropertyEditor").CreatePropertyRowGenerator(Args);
-				Generator->SetObjects({Objects[0]});
+	TSharedPtr<SWidget> ValueWidget = FieldItem->GetValueWidget();
+	if (!ValueWidget.IsValid())
+	{
+		return SNullWidget::NullWidget;
+	}
 
-				if (TSharedPtr<IDetailTreeNode> Node = FAvaRundownPageRemoteControlWidgetUtils::FindNode(Generator->GetRootTreeNodes(), FieldPtr->FieldPathInfo.ToPathPropertyString(), FAvaRundownPageRemoteControlWidgetUtils::EFindNodeMethod::Path))
-				{
-					const FNodeWidgets NodeWidgets = Node->CreateNodeWidgets();
-					ValueWidget = NodeWidgets.WholeRowWidget.IsValid() ? NodeWidgets.WholeRowWidget : NodeWidgets.ValueWidget;
-
-					if (ItemPtr->IsEntityControlled())
-					{
-						ValueWidget = SNew(SHorizontalBox)
-							+ SHorizontalBox::Slot()
-							.VAlign(EVerticalAlignment::VAlign_Center)
-							[
-								ValueWidget.ToSharedRef()
-							]
-							+ SHorizontalBox::Slot()
-							.VAlign(EVerticalAlignment::VAlign_Center)
-							.Padding(3.f, 0.f, 0.f, 0.f)
-							[
-								SNew(STextBlock)
-								.Text(LOCTEXT("Controlled", "(Controlled)"))
-							];
-
-						ValueWidget->SetEnabled(false);
-					}
-
-					return ValueWidget.ToSharedRef();
-				}
-			}
-		}
+	if (FieldItem->IsEntityControlled())
+	{
+		ValueWidget->SetEnabled(false);
+		ValueWidget = SNew(SHorizontalBox)
+			+ SHorizontalBox::Slot()
+			.VAlign(EVerticalAlignment::VAlign_Center)
+			[
+				ValueWidget.ToSharedRef()
+			]
+			+ SHorizontalBox::Slot()
+			.VAlign(EVerticalAlignment::VAlign_Center)
+			.Padding(3.f, 0.f, 0.f, 0.f)
+			[
+				SNew(STextBlock)
+				.Text(LOCTEXT("Controlled", "(Controlled)"))
+			];
 	}
 
-	return SNullWidget::NullWidget;
+	return ValueWidget.ToSharedRef();
 }
 
 FText SAvaRundownRCPropertyItemRow::GetPropertyTooltipText() const
 {
-	TSharedPtr<const FAvaRundownRCPropertyItem> ItemPtr = ItemPtrWeak.Pin();
+	TSharedPtr<const FAvaRundownRCFieldItem> ItemPtr = FieldItemWeak.Pin();
 
 	FText OwnerText = LOCTEXT("InvalidOwnerText", "(Invalid)");
 	FText SubobjectPathText = LOCTEXT("InvalidSubobjectPathText", "(Invalid)");
diff --git a/Engine/Plugins/Experimental/Avalanche/Source/AvalancheMediaEditor/Private/Rundown/DetailsView/RemoteControl/Properties/SAvaRundownRCPropertyItemRow.h b/Engine/Plugins/Experimental/Avalanche/Source/AvalancheMediaEditor/Private/Rundown/DetailsView/RemoteControl/Properties/SAvaRundownRCPropertyItemRow.h
--- a/Engine/Plugins/Experimental/Avalanche/Source/AvalancheMediaEditor/Private/Rundown/DetailsView/RemoteControl/Properties/SAvaRundownRCPropertyItemRow.h
+++ b/Engine/Plugins/Experimental/Avalanche/Source/AvalancheMediaEditor/Private/Rundown/DetailsView/RemoteControl/Properties/SAvaRundownRCPropertyItemRow.h
@@ -2,42 +2,38 @@
 
 #pragma once
 
-#include "AvaRundownRCPropertyItem.h"
 #include "IPropertyRowGenerator.h"
-#include "Rundown/DetailsView/RemoteControl/Properties/AvaRundownRCPropertyItem.h"
 #include "SAvaRundownPageRemoteControlProps.h"
 #include "Templates/SharedPointer.h"
-#include "Widgets/Views/STableRow.h"
 #include "Widgets/Layout/SBox.h"
 #include "Widgets/SWidget.h"
+#include "Widgets/Views/STableRow.h"
 
+class FAvaRundownRCFieldItem;
 class FText;
 
-class SAvaRundownRCPropertyItemRow : public SMultiColumnTableRow<FAvaRundownRCPropertyItemPtr>
+class SAvaRundownRCPropertyItemRow : public SMultiColumnTableRow<TSharedPtr<FAvaRundownRCFieldItem>>
 {
 public:
 	SLATE_BEGIN_ARGS(SAvaRundownRCPropertyItemRow) {}
 	SLATE_END_ARGS()
 
 	void Construct(const FArguments& InArgs, TSharedRef<SAvaRundownPageRemoteControlProps> InPropertyPanel,
-		const TSharedRef<STableViewBase>& InOwnerTableView, const TSharedPtr<const FAvaRundownRCPropertyItem>& InRowItem);
+		const TSharedRef<STableViewBase>& InOwnerTableView, const TSharedPtr<const FAvaRundownRCFieldItem>& InRowItem);
 
 	virtual TSharedRef<SWidget> GenerateWidgetForColumn(const FName& InColumnName) override;
 
 	void UpdateValue();
 
 protected:
-	/** Get this field's label. */
-	FText GetFieldLabel() const;
+	TSharedRef<SWidget> CreateName();
 
 	TSharedRef<SWidget> CreateValue();
 
 	FText GetPropertyTooltipText() const;
 
-	TWeakPtr<const FAvaRundownRCPropertyItem> ItemPtrWeak;
+	TWeakPtr<const FAvaRundownRCFieldItem> FieldItemWeak;
 	TWeakPtr<SAvaRundownPageRemoteControlProps> PropertyPanelWeak;
 	TSharedPtr<FAvaRundownPageRCPropsNotifyHook> NotifyHook;
-	TSharedPtr<IPropertyRowGenerator> Generator;
 	TSharedPtr<SBox> ValueContainer;
-	TSharedPtr<SWidget> ValueWidget;
 };
diff --git a/Engine/Plugins/Experimental/Avalanche/Source/AvalancheMediaEditor/Private/Rundown/DetailsView/SAvaRundownPageDetails.cpp b/Engine/Plugins/Experimental/Avalanche/Source/AvalancheMediaEditor/Private/Rundown/DetailsView/SAvaRundownPageDetails.cpp
--- a/Engine/Plugins/Experimental/Avalanche/Source/AvalancheMediaEditor/Private/Rundown/DetailsView/SAvaRundownPageDetails.cpp
+++ b/Engine/Plugins/Experimental/Avalanche/Source/AvalancheMediaEditor/Private/Rundown/DetailsView/SAvaRundownPageDetails.cpp
@@ -50,28 +50,6 @@ void SAvaRundownPageDetails::Construct(const FArguments& InArgs, const TSharedPt
 		Rundown->GetOnPageListChanged().AddSP(this, &SAvaRundownPageDetails::OnPageListChanged);
 	}
 
-	TSharedRef<SHorizontalBox> AnimationHeader = SNew(SHorizontalBox);
-	{
-		static const TArray<FText> HeaderTitles
-    	{
-    		LOCTEXT("HeaderTitleAnimationName", "Animation Name"),
-    		LOCTEXT("HeaderTitleLoopsToPlay", "Loops to Play"),
-    		LOCTEXT("HeaderTitlePlaybackSpeed", "Playback Speed"),
-    		LOCTEXT("HeaderTitlePlayMode", "Play Mode")
-    	};
-
-    	for (const FText& Title : HeaderTitles)
-    	{
-    		AnimationHeader->AddSlot()
-    			.FillWidth(1.f)
-    			[
-    				SNew(STextBlock)
-    				.Text(Title)
-    				.Justification(ETextJustify::Center)
-    			];
-    	}
-	}
-
 	ChildSlot
 	[
 		SNew(SScrollBox)
