Motion Design: fixed crash when removing a color entry

diff --git a/Engine/Plugins/Experimental/Avalanche/Source/AvalancheLevelViewport/Private/Widgets/SAvaLevelViewportActorColorMenu.cpp b/Engine/Plugins/Experimental/Avalanche/Source/AvalancheLevelViewport/Private/Widgets/SAvaLevelViewportActorColorMenu.cpp
--- a/Engine/Plugins/Experimental/Avalanche/Source/AvalancheLevelViewport/Private/Widgets/SAvaLevelViewportActorColorMenu.cpp
+++ b/Engine/Plugins/Experimental/Avalanche/Source/AvalancheLevelViewport/Private/Widgets/SAvaLevelViewportActorColorMenu.cpp
@@ -1292,7 +1292,7 @@ TSharedRef<SWidget> SAvaLevelViewportActorColorMenu::BuildThemesMenu()
 
 						SubMenuBuilder.EndSection();
 
-						if (ColorThemes[ThemeIdx].GetColors().Num() > 0)
+						if (ColorThemes.IsValidIndex(ThemeIdx) && ColorThemes[ThemeIdx].GetColors().Num() > 0)
 						{
 							SubMenuBuilder.BeginSection("Colors", LOCTEXT("Colors", "Colors"));
 
