Motion Design: Fix perf issue with Outliner unregistering and re-registering the same material item causing refresh and long stalls.

diff --git a/Engine/Plugins/Experimental/Avalanche/Source/AvalancheEditor/Private/Outliner/AvaOutlinerMaterialDesignerProxy.cpp b/Engine/Plugins/Experimental/Avalanche/Source/AvalancheEditor/Private/Outliner/AvaOutlinerMaterialDesignerProxy.cpp
--- a/Engine/Plugins/Experimental/Avalanche/Source/AvalancheEditor/Private/Outliner/AvaOutlinerMaterialDesignerProxy.cpp
+++ b/Engine/Plugins/Experimental/Avalanche/Source/AvalancheEditor/Private/Outliner/AvaOutlinerMaterialDesignerProxy.cpp
@@ -29,11 +29,15 @@ void FAvaOutlinerMaterialDesignerProxy::GetProxiedItems(const TSharedRef<IAvaOut
 			for (int32 Index = 0; Index < Materials.Num(); ++Index)
 			{
 				UMaterialInterface* const Material = Materials[Index];
+				if (!Material || !Material->IsA<UDynamicMaterialInstance>())
+				{
+					continue;
+				}
 
 				// Note that we pass in the parent of this proxy rather than the proxy itself because
 				// the parent (component item) is what's referencing this material
 				// Also, the Material Item already has a way to get the proxy (via GetParent())
-				FAvaOutlinerItemPtr MaterialItem = Outliner.FindOrAdd<FAvaOutlinerMaterialDesigner>(Material, InParent, Index);
+				FAvaOutlinerItemPtr MaterialItem = Outliner.FindOrAdd<FAvaOutlinerMaterial>(Material, InParent, Index);
 
 				// Upgrade it to a material designer item
 				if (!MaterialItem->IsExactlyA<FAvaOutlinerMaterialDesigner>())
