[Avalanche]
* Fixed compile errors when forceinline->inline

diff --git a/Engine/Plugins/Experimental/Avalanche/Source/AvalancheMedia/Private/Broadcast/AvaBroadcast.cpp b/Engine/Plugins/Experimental/Avalanche/Source/AvalancheMedia/Private/Broadcast/AvaBroadcast.cpp
--- a/Engine/Plugins/Experimental/Avalanche/Source/AvalancheMedia/Private/Broadcast/AvaBroadcast.cpp
+++ b/Engine/Plugins/Experimental/Avalanche/Source/AvalancheMedia/Private/Broadcast/AvaBroadcast.cpp
@@ -689,10 +689,9 @@ void UAvaBroadcast::EnsureValidCurrentProfile()
 {
 	if (CurrentProfile == NAME_None || !Profiles.Contains(CurrentProfile))
 	{
-		for (TPair<FName, FAvaBroadcastProfile>& Profile : Profiles)
+		if (!Profiles.IsEmpty())
 		{
-			SetCurrentProfile(Profile.Key);
-			break;
+			SetCurrentProfile(Profiles.begin()->Key);
 		}
 	}
 }
diff --git a/Engine/Plugins/Experimental/Avalanche/Source/AvalancheMedia/Private/Playback/AvaPlaybackServer.cpp b/Engine/Plugins/Experimental/Avalanche/Source/AvalancheMedia/Private/Playback/AvaPlaybackServer.cpp
--- a/Engine/Plugins/Experimental/Avalanche/Source/AvalancheMedia/Private/Playback/AvaPlaybackServer.cpp
+++ b/Engine/Plugins/Experimental/Avalanche/Source/AvalancheMedia/Private/Playback/AvaPlaybackServer.cpp
@@ -399,9 +399,9 @@ const IAvaBroadcastSettings* FAvaPlaybackServer::GetBroadcastSettings() const
 {
 	// Returns the first client we have.
 	// Todo: In case we have multiple clients, we will need a smarter way to handle this.
-	for (const TPair<FString, TSharedPtr<FClientInfo>>& Client : Clients)
+	if (!Clients.IsEmpty())
 	{
-		return &Client.Value->BroadcastSettings;
+		return &Clients.begin()->Value->BroadcastSettings;
 	}
 	return nullptr;
 }
@@ -410,9 +410,9 @@ const FAvaInstanceSettings* FAvaPlaybackServer::GetAvaInstanceSettings() const
 {
 	// Returns the first client we have.
 	// Todo: In case we have multiple clients, we will need a smarter way to handle this.
-	for (const TPair<FString, TSharedPtr<FClientInfo>>& Client : Clients)
+	if (!Clients.IsEmpty())
 	{
-		return &Client.Value->AvaInstanceSettings;
+		return &Clients.begin()->Value->AvaInstanceSettings;
 	}
 	return nullptr;
 }
@@ -421,9 +421,9 @@ const FAvaPlayableSettings* FAvaPlaybackServer::GetPlayableSettings() const
 {
 	// Returns the first client we have.
 	// Todo: In case we have multiple clients, we will need a smarter way to handle this.
-	for (const TPair<FString, TSharedPtr<FClientInfo>>& Client : Clients)
+	if (!Clients.IsEmpty())
 	{
-		return &Client.Value->PlayableSettings;
+		return &Clients.begin()->Value->PlayableSettings;
 	}
 	return nullptr;
 }
