[MotionDesign] Stagger Tool:
- Fix random not working as expect
- Remove decrement option
- Add "Playback Range" and "Selection Range" starting point options

diff --git a/Engine/Plugins/Experimental/Avalanche/Source/AvalancheSequencer/Private/StaggerTool/AvaStaggerTool.cpp b/Engine/Plugins/Experimental/Avalanche/Source/AvalancheSequencer/Private/StaggerTool/AvaStaggerTool.cpp
--- a/Engine/Plugins/Experimental/Avalanche/Source/AvalancheSequencer/Private/StaggerTool/AvaStaggerTool.cpp
+++ b/Engine/Plugins/Experimental/Avalanche/Source/AvalancheSequencer/Private/StaggerTool/AvaStaggerTool.cpp
@@ -166,7 +166,12 @@ void FAvaStaggerTool::OnSequencerSelectionChanged()
 {
 	CacheOriginalElements();
 
-	if (Settings->bAutoApply)
+	if (!HasValidSelection())
+	{
+		// Turn off auto apply if we've lost a valid selection
+		Settings->bAutoApply = false;
+	}
+	else if (Settings->bAutoApply)
 	{
 		OnApply();
 	}
@@ -245,13 +250,8 @@ FFrameNumber FAvaStaggerTool::GetInterval(const FFrameNumber InRangeSize, const
 	{
 	default:
 	case EAvaSequencerStaggerDistribution::Increment:
-	case EAvaSequencerStaggerDistribution::Decrement:
 		{
-			FFrameNumber Interval = ConvertTime(Settings->ToolOptions.Interval);
-			if (Settings->ToolOptions.Distribution == EAvaSequencerStaggerDistribution::Decrement)
-			{
-				Interval = Interval * -1;
-			}
+			const FFrameNumber Interval = ConvertTime(Settings->ToolOptions.Interval);
 			return Settings->ToolOptions.bUseCurve && InElementSize.Value != 0 ? InElementSize : Interval;
 		}
 
@@ -259,12 +259,8 @@ FFrameNumber FAvaStaggerTool::GetInterval(const FFrameNumber InRangeSize, const
 		return InRangeSize.Value / InElementCount;
 
 	case EAvaSequencerStaggerDistribution::Random:
-		FRandomStream Stream(Settings->ToolOptions.RandomSeed);
-		if (InElementSize > InRangeSize.Value)
-		{
-			return Stream.RandRange(0, InRangeSize.Value);
-		}
-		return Stream.RandRange(0, (InRangeSize - InElementSize).Value);
+		const int32 AdjustedMax = FMath::Clamp(InRangeSize.Value - InElementSize.Value, 0, InRangeSize.Value);
+		return RandomStream.RandRange(0, AdjustedMax);
 	}
 }
 
@@ -316,7 +312,8 @@ void FAvaStaggerTool::CacheOriginalElements()
 	}
 }
 
-FFrameNumber FAvaStaggerTool::FindFirstStaggerPoint(const TArrayView<FStaggerLayerBarElement>& InStaggerElements, const TRange<FFrameNumber>& InRange) const
+FFrameNumber FAvaStaggerTool::FindFirstStaggerPoint(const TArrayView<FStaggerLayerBarElement>& InStaggerElements
+	, const TRange<FFrameNumber>& InRange) const
 {
 	FFrameNumber OutFrameNumber;
 
@@ -325,9 +322,6 @@ FFrameNumber FAvaStaggerTool::FindFirstStaggerPoint(const TArrayView<FStaggerLay
 	case EAvaSequencerStaggerDistribution::Increment:
 		break;
 
-	case EAvaSequencerStaggerDistribution::Decrement:
-		break;
-
 	case EAvaSequencerStaggerDistribution::Range:
 		switch (Settings->ToolOptions.Range)
 		{
@@ -348,14 +342,10 @@ FFrameNumber FAvaStaggerTool::FindFirstStaggerPoint(const TArrayView<FStaggerLay
 
 	case EAvaSequencerStaggerDistribution::Random:
 		const int32 OriginalElementCount = OriginalLayerBarElements.Num();
-		if (OriginalElementCount > 1)
-		{
-			const TRange<FFrameNumber> Range = GetOperationRange(Settings->ToolOptions.Range);
-			const FFrameNumber RangeSize = Range.Size<FFrameNumber>();
-			const FFrameNumber ElementSize = OriginalLayerBarElements[0].LayerBarModel->ComputeRange().Size<FFrameNumber>();
-			OutFrameNumber = GetInterval(RangeSize, OriginalElementCount, ElementSize);
-		}
-		break;
+		const TRange<FFrameNumber> Range = GetOperationRange(Settings->ToolOptions.Range);
+		const FFrameNumber RangeSize = Range.Size<FFrameNumber>();
+		const FFrameNumber ElementSize = OriginalLayerBarElements[0].LayerBarModel->ComputeRange().Size<FFrameNumber>();
+		return Range.GetLowerBoundValue() + GetInterval(RangeSize, OriginalElementCount, ElementSize);
 	}
 
 	switch (Settings->ToolOptions.StartPosition)
@@ -372,24 +362,16 @@ FFrameNumber FAvaStaggerTool::FindFirstStaggerPoint(const TArrayView<FStaggerLay
 		break;
 
 	case EAvaSequencerStaggerStartPosition::Playhead:
-		const TSharedRef<ISequencer> Sequencer = Owner.GetSequencer();
-		OutFrameNumber = Sequencer->GetGlobalTime().Time.FrameNumber;
+		OutFrameNumber = Owner.GetSequencer()->GetGlobalTime().Time.FrameNumber;
 		break;
-	}
 
-	if (IsLayerBarSelection())
-	{
-		if (Settings->ToolOptions.Distribution == EAvaSequencerStaggerDistribution::Random)
-		{
-			const int32 OriginalElementCount = OriginalLayerBarElements.Num();
-			if (OriginalElementCount > 1)
-			{
-				const TRange<FFrameNumber> Range = GetOperationRange(Settings->ToolOptions.Range);
-				const FFrameNumber RangeSize = Range.Size<FFrameNumber>();
-				const FFrameNumber ElementSize = OriginalLayerBarElements[0].LayerBarModel->ComputeRange().Size<FFrameNumber>();
-				OutFrameNumber += GetInterval(RangeSize, OriginalElementCount, ElementSize);
-			}
-		}
+	case EAvaSequencerStaggerStartPosition::PlaybackRange:
+		OutFrameNumber = Owner.GetSequencer()->GetPlaybackRange().GetLowerBoundValue();
+		break;
+
+	case EAvaSequencerStaggerStartPosition::SelectionRange:
+		OutFrameNumber = Owner.GetSequencer()->GetSelectionRange().GetLowerBoundValue();
+		break;
 	}
 
 	OutFrameNumber += ConvertTime(Settings->ToolOptions.Shift);
@@ -397,7 +379,8 @@ FFrameNumber FAvaStaggerTool::FindFirstStaggerPoint(const TArrayView<FStaggerLay
 	return MoveTemp(OutFrameNumber);
 }
 
-FFrameNumber FAvaStaggerTool::FindFirstStaggerPoint(const TArrayView<FStaggerKeyFrameElement>& InStaggerElements, const TRange<FFrameNumber>& InRange) const
+FFrameNumber FAvaStaggerTool::FindFirstStaggerPoint(const TArrayView<FStaggerKeyFrameElement>& InStaggerElements
+	, const TRange<FFrameNumber>& InRange) const
 {
 	FFrameNumber OutFrameNumber;
 
@@ -415,8 +398,15 @@ FFrameNumber FAvaStaggerTool::FindFirstStaggerPoint(const TArrayView<FStaggerKey
 		break;
 
 	case EAvaSequencerStaggerStartPosition::Playhead:
-		const TSharedRef<ISequencer> Sequencer = Owner.GetSequencer();
-		OutFrameNumber = Sequencer->GetGlobalTime().Time.FrameNumber;
+		OutFrameNumber = Owner.GetSequencer()->GetGlobalTime().Time.FrameNumber;
+		break;
+
+	case EAvaSequencerStaggerStartPosition::PlaybackRange:
+		OutFrameNumber = Owner.GetSequencer()->GetPlaybackRange().GetLowerBoundValue();
+		break;
+
+	case EAvaSequencerStaggerStartPosition::SelectionRange:
+		OutFrameNumber = Owner.GetSequencer()->GetSelectionRange().GetLowerBoundValue();
 		break;
 	}
 
@@ -618,8 +608,7 @@ FFrameNumber FAvaStaggerTool::FindNextStaggerLocation(const FStaggerLayerBarElem
 		const FFrameNumber ElementFrameSpan = CalculateLayerBarFrameSpan(OriginalLayerBarElements, InElementIndex + 1);
 		const float ElementStartPercentOfRange = static_cast<float>(ElementFrameSpan.Value) / static_cast<float>(FullFrameSpan.Value);
 
-		if (Settings->ToolOptions.Distribution == EAvaSequencerStaggerDistribution::Increment
-			|| Settings->ToolOptions.Distribution == EAvaSequencerStaggerDistribution::Decrement)
+		if (Settings->ToolOptions.Distribution == EAvaSequencerStaggerDistribution::Increment)
 		{
 			const FFrameNumber LocalCurveOffset = CalculateLocalCurveOffset(ElementStartPercentOfRange, TRange<FFrameNumber>(0, FullFrameSpan));
 			return InFirstFrame + LocalCurveOffset + Interval;
@@ -634,11 +623,11 @@ FFrameNumber FAvaStaggerTool::FindNextStaggerLocation(const FStaggerLayerBarElem
 
 	if (Settings->ToolOptions.Distribution == EAvaSequencerStaggerDistribution::Random)
 	{
-		return InFirstFrame + Interval;
+		return Range.GetLowerBoundValue() + Interval;
 	}
 
-	const FFrameNumber LayerBarOperationOffset = GetLayerBarOperationOffset(InElement.LayerBarModel);
-	return InCurrentFrame + LayerBarOperationOffset + Interval;
+	const FFrameNumber OperationOffset = GetLayerBarOperationOffset(InElement.LayerBarModel);
+	return InCurrentFrame + OperationOffset + Interval;
 }
 
 FFrameNumber FAvaStaggerTool::FindNextStaggerLocation(const FStaggerKeyFrameElement& InElement
@@ -667,7 +656,7 @@ FFrameNumber FAvaStaggerTool::FindNextStaggerLocation(const FStaggerKeyFrameElem
 
 	if (Settings->ToolOptions.Distribution == EAvaSequencerStaggerDistribution::Random)
 	{
-		return InFirstFrame + Interval;
+		return Range.GetLowerBoundValue() + Interval;
 	}
 
 	return InCurrentFrame + Interval;
@@ -675,6 +664,11 @@ FFrameNumber FAvaStaggerTool::FindNextStaggerLocation(const FStaggerKeyFrameElem
 
 void FAvaStaggerTool::Stagger()
 {
+	if (Settings->ToolOptions.Distribution == EAvaSequencerStaggerDistribution::Random)
+	{
+		RandomStream = FRandomStream(Settings->ToolOptions.RandomSeed);
+	}
+
 	if (IsLayerBarSelection())
 	{
 		// Reset layer bar offsets that were last applied
diff --git a/Engine/Plugins/Experimental/Avalanche/Source/AvalancheSequencer/Private/StaggerTool/AvaStaggerTool.h b/Engine/Plugins/Experimental/Avalanche/Source/AvalancheSequencer/Private/StaggerTool/AvaStaggerTool.h
--- a/Engine/Plugins/Experimental/Avalanche/Source/AvalancheSequencer/Private/StaggerTool/AvaStaggerTool.h
+++ b/Engine/Plugins/Experimental/Avalanche/Source/AvalancheSequencer/Private/StaggerTool/AvaStaggerTool.h
@@ -89,8 +89,10 @@ private:
 
 	void CacheOriginalElements();
 
-	FFrameNumber FindFirstStaggerPoint(const TArrayView<FStaggerLayerBarElement>& InStaggerElements, const TRange<FFrameNumber>& InRange) const;
-	FFrameNumber FindFirstStaggerPoint(const TArrayView<FStaggerKeyFrameElement>& InStaggerElements, const TRange<FFrameNumber>& InRange) const;
+	FFrameNumber FindFirstStaggerPoint(const TArrayView<FStaggerLayerBarElement>& InStaggerElements
+		, const TRange<FFrameNumber>& InRange) const;
+	FFrameNumber FindFirstStaggerPoint(const TArrayView<FStaggerKeyFrameElement>& InStaggerElements
+		, const TRange<FFrameNumber>& InRange) const;
 
 	TArray<UE::Sequencer::TViewModelPtr<UE::Sequencer::FLayerBarModel>> GatherAllSelectedLayerBars() const;
 
@@ -136,4 +138,6 @@ private:
 	// Cached values at the start of auto-updating
 	TArray<FStaggerLayerBarElement> OriginalLayerBarElements;
 	TArray<FStaggerKeyFrameElement> OriginalKeyFrameElements;
+
+	FRandomStream RandomStream;
 };
diff --git a/Engine/Plugins/Experimental/Avalanche/Source/AvalancheSequencer/Private/StaggerTool/AvaStaggerToolSettings.h b/Engine/Plugins/Experimental/Avalanche/Source/AvalancheSequencer/Private/StaggerTool/AvaStaggerToolSettings.h
--- a/Engine/Plugins/Experimental/Avalanche/Source/AvalancheSequencer/Private/StaggerTool/AvaStaggerToolSettings.h
+++ b/Engine/Plugins/Experimental/Avalanche/Source/AvalancheSequencer/Private/StaggerTool/AvaStaggerToolSettings.h
@@ -16,6 +16,11 @@ enum class EAvaSequencerStaggerStartPosition : uint8
 	/** Sequence placement will begin at the current position of the playhead */
 	Playhead = 2,
 
+	/** Sequence placement will begin at the start of the playback range */
+	PlaybackRange = 3,
+	/** Sequence placement will begin at the start of the selection range */
+	SelectionRange = 4,
+
 	Default = FirstSelected UMETA(hidden)
 };
 
@@ -24,10 +29,8 @@ enum class EAvaSequencerStaggerDistribution : uint8
 {
 	/** Sequences placed one after the other */
 	Increment = 0,
-	/** Sequences placed one before the other */
-	Decrement = 1,
 	/** Sequences spaced out between a range */
-	Range = 2,
+	Range = 1,
 	/** Sequences randomly distributed in range */
 	Random = 255,
 
diff --git a/Engine/Plugins/Experimental/Avalanche/Source/AvalancheSequencer/Private/StaggerTool/Customization/AvaStaggerToolSettingsDetailsCustomization.cpp b/Engine/Plugins/Experimental/Avalanche/Source/AvalancheSequencer/Private/StaggerTool/Customization/AvaStaggerToolSettingsDetailsCustomization.cpp
--- a/Engine/Plugins/Experimental/Avalanche/Source/AvalancheSequencer/Private/StaggerTool/Customization/AvaStaggerToolSettingsDetailsCustomization.cpp
+++ b/Engine/Plugins/Experimental/Avalanche/Source/AvalancheSequencer/Private/StaggerTool/Customization/AvaStaggerToolSettingsDetailsCustomization.cpp
@@ -45,41 +45,27 @@ void FAvaStaggerToolSettingsDetailsCustomization::CustomizeDetails(IDetailLayout
 	AddDistributionRow(DistributionPropertyRow, DistributionProperty.ToSharedRef());
 
 	IDetailPropertyRow& RandomSeedPropertyRow = OptionsCategory.AddProperty(RandomSeedProperty)
-		.Visibility(
-			TAttribute<EVisibility>::Create(
-				TAttribute<EVisibility>::FGetter::CreateSP(this, &FAvaStaggerToolSettingsDetailsCustomization::GetRandomSeedPropertyVisibility)
-			));
+		.Visibility(TAttribute<EVisibility>::CreateSP(this, &FAvaStaggerToolSettingsDetailsCustomization::GetRandomSeedPropertyVisibility));
 	AddRandomSeedRow(RandomSeedPropertyRow, RandomSeedProperty.ToSharedRef());
 
 	IDetailPropertyRow& RangePropertyRow = OptionsCategory.AddProperty(RangeProperty)
-		.Visibility(
-			TAttribute<EVisibility>::Create(
-				TAttribute<EVisibility>::FGetter::CreateSP(this, &FAvaStaggerToolSettingsDetailsCustomization::GetRangePropertyVisibility)
-			));
+		.Visibility(TAttribute<EVisibility>::CreateSP(this, &FAvaStaggerToolSettingsDetailsCustomization::GetRangePropertyVisibility));
 	AddRangeRow(RangePropertyRow, RangeProperty.ToSharedRef());
 
 	IDetailPropertyRow& CustomRangePropertyRow = OptionsCategory.AddProperty(CustomRangeProperty)
-		.Visibility(
-			TAttribute<EVisibility>::Create(
-				TAttribute<EVisibility>::FGetter::CreateSP(this, &FAvaStaggerToolSettingsDetailsCustomization::GetCustomRangePropertyVisibility)
-			));
+		.Visibility(TAttribute<EVisibility>::CreateSP(this, &FAvaStaggerToolSettingsDetailsCustomization::GetCustomRangePropertyVisibility));
 	AddCustomIntSpinBoxRow(CustomRangePropertyRow, CustomRangeProperty.ToSharedRef());
 
 	IDetailPropertyRow& StartPositionPropertyRow = OptionsCategory.AddProperty(StartPositionProperty)
-		.Visibility(
-			TAttribute<EVisibility>::Create(
-				TAttribute<EVisibility>::FGetter::CreateSP(this, &FAvaStaggerToolSettingsDetailsCustomization::GetStartPositionPropertyVisibility)
-			));
+		.Visibility(TAttribute<EVisibility>::CreateSP(this, &FAvaStaggerToolSettingsDetailsCustomization::GetStartPositionPropertyVisibility));
 	AddStartPositionRow(StartPositionPropertyRow, StartPositionProperty.ToSharedRef());
 
-	IDetailPropertyRow& OperationPointPropertyRow = OptionsCategory.AddProperty(OperationPointProperty);
+	IDetailPropertyRow& OperationPointPropertyRow = OptionsCategory.AddProperty(OperationPointProperty)
+		.Visibility(TAttribute<EVisibility>::CreateSP(this, &FAvaStaggerToolSettingsDetailsCustomization::GetOperationPointPropertyVisibility));
 	AddOperationPointRow(OperationPointPropertyRow, OperationPointProperty.ToSharedRef());
 
 	IDetailPropertyRow& IntervalPropertyRow = OptionsCategory.AddProperty(IntervalProperty)
-		.Visibility(
-			TAttribute<EVisibility>::Create(
-				TAttribute<EVisibility>::FGetter::CreateSP(this, &FAvaStaggerToolSettingsDetailsCustomization::GetIntervalPropertyVisibility)
-			));
+		.Visibility(TAttribute<EVisibility>::CreateSP(this, &FAvaStaggerToolSettingsDetailsCustomization::GetIntervalPropertyVisibility));
 	AddCustomIntSpinBoxRow(IntervalPropertyRow, IntervalProperty.ToSharedRef());
 
 	IDetailPropertyRow& ShiftPropertyRow = OptionsCategory.AddProperty(ShiftProperty);
@@ -91,16 +77,10 @@ void FAvaStaggerToolSettingsDetailsCustomization::CustomizeDetails(IDetailLayout
 	OptionsCategory.AddProperty(UseCurveProperty);
 
 	OptionsCategory.AddProperty(CurveProperty)
-		.Visibility(
-			TAttribute<EVisibility>::Create(
-				TAttribute<EVisibility>::FGetter::CreateSP(this, &FAvaStaggerToolSettingsDetailsCustomization::GetCurvePropertyVisibility)
-			));
+		.Visibility(TAttribute<EVisibility>::CreateSP(this, &FAvaStaggerToolSettingsDetailsCustomization::GetCurvePropertyVisibility));
 
 	IDetailPropertyRow& CurveOffsetPropertyRow = OptionsCategory.AddProperty(CurveOffsetProperty)
-		.Visibility(
-			TAttribute<EVisibility>::Create(
-				TAttribute<EVisibility>::FGetter::CreateSP(this, &FAvaStaggerToolSettingsDetailsCustomization::GetCurveOffsetPropertyVisibility)
-			));
+		.Visibility(TAttribute<EVisibility>::CreateSP(this, &FAvaStaggerToolSettingsDetailsCustomization::GetCurveOffsetPropertyVisibility));
 	AddCustomFloatSpinBoxRow(CurveOffsetPropertyRow, CurveOffsetProperty.ToSharedRef());
 }
 
@@ -183,14 +163,25 @@ EVisibility FAvaStaggerToolSettingsDetailsCustomization::GetCustomRangePropertyV
 EVisibility FAvaStaggerToolSettingsDetailsCustomization::GetStartPositionPropertyVisibility() const
 {
 	const EAvaSequencerStaggerDistribution TypeValue = GetDistributionPropertyValue();
-	return (TypeValue != EAvaSequencerStaggerDistribution::Range)
+	return (TypeValue != EAvaSequencerStaggerDistribution::Range && TypeValue != EAvaSequencerStaggerDistribution::Random)
 		? EVisibility::Visible : EVisibility::Collapsed;
 }
 
 EVisibility FAvaStaggerToolSettingsDetailsCustomization::GetOperationPointPropertyVisibility() const
 {
-	return (WeakTool.IsValid() && WeakTool.Pin()->IsLayerBarSelection())
-		? EVisibility::Visible : EVisibility::Collapsed;
+	// Keyframes have no range to use
+	if (WeakTool.IsValid() && !WeakTool.Pin()->IsLayerBarSelection())
+	{
+		return EVisibility::Collapsed;
+	}
+
+	const EAvaSequencerStaggerDistribution DistributionValue = GetDistributionPropertyValue();
+	if (DistributionValue == EAvaSequencerStaggerDistribution::Random)
+	{
+		return EVisibility::Collapsed;
+	}
+
+	return EVisibility::Visible;
 }
 
 template<typename InValueType>
diff --git a/Engine/Plugins/Experimental/Avalanche/Source/AvalancheSequencer/Private/StaggerTool/Widgets/SAvaStaggerSettingsRadioGroup.h b/Engine/Plugins/Experimental/Avalanche/Source/AvalancheSequencer/Private/StaggerTool/Widgets/SAvaStaggerSettingsRadioGroup.h
--- a/Engine/Plugins/Experimental/Avalanche/Source/AvalancheSequencer/Private/StaggerTool/Widgets/SAvaStaggerSettingsRadioGroup.h
+++ b/Engine/Plugins/Experimental/Avalanche/Source/AvalancheSequencer/Private/StaggerTool/Widgets/SAvaStaggerSettingsRadioGroup.h
@@ -5,6 +5,7 @@
 #include "DetailLayoutBuilder.h"
 #include "Widgets/Input/SCheckBox.h"
 #include "Widgets/Layout/SBox.h"
+#include "Widgets/Layout/SUniformWrapPanel.h"
 #include "Widgets/SCompoundWidget.h"
 #include "Widgets/Text/STextBlock.h"
 
@@ -27,7 +28,10 @@ public:
 
 		OnValueChanged = InArgs._OnValueChanged;
 
-		const TSharedRef<SHorizontalBox> HorizontalBox = SNew(SHorizontalBox);
+		const TSharedRef<SUniformWrapPanel> WrapPanel = SNew(SUniformWrapPanel)
+			.HAlign(HAlign_Fill)
+			.SlotPadding(FMargin(0.f, 0.f, 1.f, 1.f))
+			.NumColumnsOverride(3);
 
 		for (int32 Index = 0; Index < Enum->NumEnums() - 1; ++Index)
 		{
@@ -39,9 +43,8 @@ public:
 					, Enum->GetToolTipTextByIndex(Index));
 				Buttons.Add(NewButton);
 
-				HorizontalBox->AddSlot()
+				WrapPanel->AddSlot()
 					.VAlign(VAlign_Center)
-					.Padding(FMargin(1.f, 0.f))
 					[
 						NewButton
 					];
@@ -49,8 +52,10 @@ public:
 		}
 
 		ChildSlot
+			.HAlign(HAlign_Fill)
+			.Padding(FMargin(0.f, 3.f))
 			[
-				HorizontalBox
+				WrapPanel
 			];
 	}
 
@@ -79,7 +84,7 @@ protected:
 				})
 			[
 				SNew(SBox)
-				.WidthOverride(80.f)
+				.WidthOverride(90.f)
 				.HAlign(HAlign_Center)
 				.VAlign(VAlign_Center)
 				[
diff --git a/Engine/Plugins/Experimental/Avalanche/Source/AvalancheSequencer/Private/StaggerTool/Widgets/SAvaStaggerTool.cpp b/Engine/Plugins/Experimental/Avalanche/Source/AvalancheSequencer/Private/StaggerTool/Widgets/SAvaStaggerTool.cpp
--- a/Engine/Plugins/Experimental/Avalanche/Source/AvalancheSequencer/Private/StaggerTool/Widgets/SAvaStaggerTool.cpp
+++ b/Engine/Plugins/Experimental/Avalanche/Source/AvalancheSequencer/Private/StaggerTool/Widgets/SAvaStaggerTool.cpp
@@ -42,7 +42,7 @@ void SAvaStaggerTool::Construct(const FArguments& InArgs, const TWeakPtr<FAvaSta
 		[
 			SNew(SSeparator)
 			.ColorAndOpacity(FStyleColors::Dropdown)
-			.SeparatorImage(FAppStyle::Get().GetBrush("WhiteBrush"))
+			.SeparatorImage(FAppStyle::Get().GetBrush(TEXT("WhiteBrush")))
 		]
 		+ SVerticalBox::Slot()
 		.AutoHeight()
@@ -112,7 +112,11 @@ TSharedRef<SWidget> SAvaStaggerTool::ConstructApplyRow()
 			.OnClicked(this, &SAvaStaggerTool::OnApplyButtonClick)
 			.IsEnabled_Lambda([this]()
 				{
-					return WeakTool.IsValid() && !WeakTool.Pin()->IsAutoAppying();
+					if (const TSharedPtr<FAvaStaggerTool> Tool = WeakTool.Pin())
+					{
+						return Tool->HasValidSelection() && !Tool->IsAutoAppying();
+					}
+					return false;
 				})
 		]
 		+ SHorizontalBox::Slot()
@@ -237,7 +241,8 @@ FText SAvaStaggerTool::GetSelectionToolTipText() const
 	return LOCTEXT("InvalidSelectionTooltip",
 		"Select at least two Sequencer outliner tracks containing a layer bar or at least two layer bars directly\n\n"
 		" - OR -\n\n"
-		"Select at least two Sequencer key frames");
+		"Select at least two Sequencer key frames\n\n"
+		"*NOTE* Order of selection matters! Tracks will be staggered in the order they were selected");
 }
 
 FSlateColor SAvaStaggerTool::GetSelectionTextColor() const
