Motion Design: Fix for color picker theme menu
- The color picker theme menu will now correctly close when selecting an option.

diff --git a/Engine/Plugins/Experimental/Avalanche/Source/AvalancheLevelViewport/Private/Widgets/SAvaLevelViewportActorColorMenu.cpp b/Engine/Plugins/Experimental/Avalanche/Source/AvalancheLevelViewport/Private/Widgets/SAvaLevelViewportActorColorMenu.cpp
--- a/Engine/Plugins/Experimental/Avalanche/Source/AvalancheLevelViewport/Private/Widgets/SAvaLevelViewportActorColorMenu.cpp
+++ b/Engine/Plugins/Experimental/Avalanche/Source/AvalancheLevelViewport/Private/Widgets/SAvaLevelViewportActorColorMenu.cpp
@@ -1142,6 +1142,7 @@ TSharedRef<SWidget> SAvaLevelViewportActorColorMenu::BuildRenameColorMenu(int32
 	FLinearColor CurrentColor = ThemeColors[InColorIdx].Color;
 
 	constexpr bool bInShouldCloseWindowAfterMenuSelection = true;
+
 	FMenuBuilder ColorMenuBuilder(bInShouldCloseWindowAfterMenuSelection, nullptr);
 	{
 		ColorMenuBuilder.BeginSection("ColorOptions", LOCTEXT("ColorOptions", "Color Options"));
@@ -1189,6 +1190,8 @@ TSharedRef<SWidget> SAvaLevelViewportActorColorMenu::BuildRenameColorMenu(int32
 					SaveColorThemesToIni();
 
 					UpdateColorPalette();
+
+					FSlateApplication::Get().DismissAllMenus();
 				})));
 
 		ColorMenuBuilder.EndSection();
@@ -1199,7 +1202,14 @@ TSharedRef<SWidget> SAvaLevelViewportActorColorMenu::BuildRenameColorMenu(int32
 
 TSharedRef<SWidget> SAvaLevelViewportActorColorMenu::BuildThemesMenu()
 {
-	FMenuBuilder ThemeMenuBuilder(/*bInShouldCloseWindowAfterMenuSelection*/false, nullptr);
+	/*
+	 * Causes the action to not execute if this is true. The widget (from the SP Lambda) is destroyed
+	 * prior to the action executing. Using FSlateApplication to dismiss the menu manually.
+	 */
+	constexpr bool bInShouldCloseWindowAfterMenuSelection = false;
+	constexpr bool bInOpenSubMenuOnClick = false;
+
+	FMenuBuilder ThemeMenuBuilder(bInShouldCloseWindowAfterMenuSelection, nullptr);
 	{
 		static FText RemoveThemeLabelText = LOCTEXT("RemoveThemeLabel", "Remove");
 		static FText RemoveThemeTooltipText = LOCTEXT("RemoveThemeTooltip", "Remove this theme.");
@@ -1249,8 +1259,11 @@ TSharedRef<SWidget> SAvaLevelViewportActorColorMenu::BuildThemesMenu()
 
 						SaveColorThemesToIni();
 						UpdateColorPalette();
+
+						FSlateApplication::Get().DismissAllMenus();
 					}),
-					FCanExecuteAction::CreateSPLambda(this, [this]() { return ColorThemes.Num() > 1; }))
+					FCanExecuteAction::CreateSPLambda(this, [this]() { return ColorThemes.Num() > 1; })
+				)
 			);
 
 			ThemeMenuBuilder.EndSection();
@@ -1286,6 +1299,8 @@ TSharedRef<SWidget> SAvaLevelViewportActorColorMenu::BuildThemesMenu()
 
 									SaveColorThemesToIni();
 									UpdateColorPalette();
+
+									FSlateApplication::Get().DismissAllMenus();
 								}),
 								FCanExecuteAction::CreateSPLambda(this, [this]() { return ColorThemes.Num() > 1; }))
 						);
@@ -1312,6 +1327,8 @@ TSharedRef<SWidget> SAvaLevelViewportActorColorMenu::BuildThemesMenu()
 											}
 
 											SetActiveColorRGB(ColorThemes[ThemeIdx].GetColors()[ColorIdx].Color);
+
+											FSlateApplication::Get().DismissAllMenus();
 										})										
 									),
 									SNew(SHorizontalBox)
@@ -1352,6 +1369,8 @@ TSharedRef<SWidget> SAvaLevelViewportActorColorMenu::BuildThemesMenu()
 														ColorThemes[ThemeIdx].GetColors().RemoveAt(ColorIdx);
 
 														SaveColorThemesToIni();
+
+														FSlateApplication::Get().DismissAllMenus();
 													})
 											));
 										})
@@ -1372,6 +1391,8 @@ TSharedRef<SWidget> SAvaLevelViewportActorColorMenu::BuildThemesMenu()
 
 							ActiveThemeIndex = ThemeIdx;
 							UpdateColorPalette();
+
+							FSlateApplication::Get().DismissAllMenus();
 						}),
 					FCanExecuteAction::CreateSPLambda(this, [this, ThemeIdx]()
 						{
@@ -1384,9 +1405,9 @@ TSharedRef<SWidget> SAvaLevelViewportActorColorMenu::BuildThemesMenu()
 				),
 				FName(ColorThemes[ThemeIdx].GetName()),
 				EUserInterfaceActionType::RadioButton,
-				/*bInOpenSubMenuOnClick*/false,
+				bInOpenSubMenuOnClick,
 				FSlateIcon(),
-				/*bInShouldCloseWindowAfterMenuSelection*/false
+				bInShouldCloseWindowAfterMenuSelection
 			);
 		}
 
