Motion Design: fixed issue with RC properties not properly displaying in Rundown or having issues updating/refreshing on change

diff --git a/Engine/Plugins/Experimental/Avalanche/Source/AvalancheMediaEditor/Private/Rundown/DetailsView/RemoteControl/Properties/AvaRundownPageRCObject.cpp b/Engine/Plugins/Experimental/Avalanche/Source/AvalancheMediaEditor/Private/Rundown/DetailsView/RemoteControl/Properties/AvaRundownPageRCObject.cpp
new file mode 100644
--- /dev/null
+++ b/Engine/Plugins/Experimental/Avalanche/Source/AvalancheMediaEditor/Private/Rundown/DetailsView/RemoteControl/Properties/AvaRundownPageRCObject.cpp
@@ -0,0 +1,64 @@
+// Copyright Epic Games, Inc. All Rights Reserved.
+
+#include "AvaRundownPageRCObject.h"
+#include "IDetailTreeNode.h"
+#include "IPropertyRowGenerator.h"
+#include "Item/AvaRundownRCDetailTreeNodeItem.h"
+#include "Modules/ModuleManager.h"
+#include "PropertyEditorModule.h"
+#include "PropertyHandle.h"
+
+FAvaRundownPageRCObject::FAvaRundownPageRCObject(UObject* InObject)
+	: ObjectKey(InObject)
+{
+}
+
+void FAvaRundownPageRCObject::Initialize(FNotifyHook* InNotifyHook)
+{
+	UObject* const ResolvedObject = ObjectKey.ResolveObjectPtr();
+
+	FPropertyRowGeneratorArgs Args;
+	Args.NotifyHook = InNotifyHook;
+	PropertyRowGenerator = FModuleManager::GetModuleChecked<FPropertyEditorModule>("PropertyEditor").CreatePropertyRowGenerator(Args);
+	PropertyRowGenerator->SetObjects({ ResolvedObject });
+
+	CacheTreeNodes();
+}
+
+TSharedPtr<IDetailTreeNode> FAvaRundownPageRCObject::FindTreeNode(const FRCFieldPathInfo& InPathInfo) const
+{
+	if (const TSharedRef<IDetailTreeNode>* FoundTreeNode = TreeNodeMap.Find(InPathInfo.ToString()))
+	{
+		return *FoundTreeNode;
+	}
+	return nullptr;
+}
+
+void FAvaRundownPageRCObject::CacheTreeNodes()
+{
+	TreeNodeMap.Reset();
+
+	if (!PropertyRowGenerator.IsValid())
+	{
+		return;
+	}
+
+	TArray<TSharedRef<IDetailTreeNode>> TreeNodes = PropertyRowGenerator->GetRootTreeNodes();
+	TArray<TSharedRef<IDetailTreeNode>> ChildTreeNodes;
+
+	while (!TreeNodes.IsEmpty())
+	{
+		TSharedRef<IDetailTreeNode> TreeNode = TreeNodes.Pop(EAllowShrinking::No);
+
+		if (TSharedPtr<IPropertyHandle> PropertyHandle = TreeNode->CreatePropertyHandle())
+		{
+			constexpr bool bCleanDuplicates = true;
+			const FRCFieldPathInfo FieldPathInfo(PropertyHandle->GeneratePathToProperty(), bCleanDuplicates);
+			TreeNodeMap.Add(FieldPathInfo.ToString(), TreeNode);
+		}
+
+		ChildTreeNodes.Reset();
+		TreeNode->GetChildren(ChildTreeNodes);
+		TreeNodes.Append(MoveTemp(ChildTreeNodes));
+	}
+}
diff --git a/Engine/Plugins/Experimental/Avalanche/Source/AvalancheMediaEditor/Private/Rundown/DetailsView/RemoteControl/Properties/AvaRundownPageRCObject.h b/Engine/Plugins/Experimental/Avalanche/Source/AvalancheMediaEditor/Private/Rundown/DetailsView/RemoteControl/Properties/AvaRundownPageRCObject.h
new file mode 100644
--- /dev/null
+++ b/Engine/Plugins/Experimental/Avalanche/Source/AvalancheMediaEditor/Private/Rundown/DetailsView/RemoteControl/Properties/AvaRundownPageRCObject.h
@@ -0,0 +1,45 @@
+// Copyright Epic Games, Inc. All Rights Reserved.
+
+#pragma once
+
+#include "Containers/Map.h"
+#include "RemoteControlFieldPath.h"
+#include "Templates/SharedPointer.h"
+#include "UObject/ObjectKey.h"
+
+class FAvaRundownRCDetailTreeNodeItem;
+class FNotifyHook;
+class IDetailTreeNode;
+class IPropertyRowGenerator;
+class UObject;
+class URemoteControlPreset;
+
+/** A class containing the information of a UObject exposed to RC */
+class FAvaRundownPageRCObject
+{
+public:
+	explicit FAvaRundownPageRCObject(UObject* InObject);
+
+	void Initialize(FNotifyHook* InNotifyHook);
+
+	TSharedPtr<IDetailTreeNode> FindTreeNode(const FRCFieldPathInfo& InPathInfo) const;
+
+	bool operator==(const FAvaRundownPageRCObject& InOther) const
+	{
+		return ObjectKey == InOther.ObjectKey;
+	}
+
+	friend uint32 GetTypeHash(const FAvaRundownPageRCObject& InObject)
+	{
+		return GetTypeHash(InObject.ObjectKey);
+	}
+
+private:
+	void CacheTreeNodes();
+
+	FObjectKey ObjectKey;
+
+	TSharedPtr<IPropertyRowGenerator> PropertyRowGenerator;
+
+	TMap<FString, TSharedRef<IDetailTreeNode>> TreeNodeMap;
+};
diff --git a/Engine/Plugins/Experimental/Avalanche/Source/AvalancheMediaEditor/Private/Rundown/DetailsView/RemoteControl/Properties/AvaRundownPageRemoteControlWidgetUtils.cpp b/Engine/Plugins/Experimental/Avalanche/Source/AvalancheMediaEditor/Private/Rundown/DetailsView/RemoteControl/Properties/AvaRundownPageRemoteControlWidgetUtils.cpp
deleted file mode 100644
diff --git a/Engine/Plugins/Experimental/Avalanche/Source/AvalancheMediaEditor/Private/Rundown/DetailsView/RemoteControl/Properties/AvaRundownPageRemoteControlWidgetUtils.h b/Engine/Plugins/Experimental/Avalanche/Source/AvalancheMediaEditor/Private/Rundown/DetailsView/RemoteControl/Properties/AvaRundownPageRemoteControlWidgetUtils.h
deleted file mode 100644
diff --git a/Engine/Plugins/Experimental/Avalanche/Source/AvalancheMediaEditor/Private/Rundown/DetailsView/RemoteControl/Properties/Item/AvaRundownRCDetailTreeNodeItem.cpp b/Engine/Plugins/Experimental/Avalanche/Source/AvalancheMediaEditor/Private/Rundown/DetailsView/RemoteControl/Properties/Item/AvaRundownRCDetailTreeNodeItem.cpp
--- a/Engine/Plugins/Experimental/Avalanche/Source/AvalancheMediaEditor/Private/Rundown/DetailsView/RemoteControl/Properties/Item/AvaRundownRCDetailTreeNodeItem.cpp
+++ b/Engine/Plugins/Experimental/Avalanche/Source/AvalancheMediaEditor/Private/Rundown/DetailsView/RemoteControl/Properties/Item/AvaRundownRCDetailTreeNodeItem.cpp
@@ -1,24 +1,26 @@
 // Copyright Epic Games, Inc. All Rights Reserved.
 
 #include "AvaRundownRCDetailTreeNodeItem.h"
-#include "IPropertyRowGenerator.h"
+#include "DetailLayoutBuilder.h"
 #include "IRemoteControlPropertyHandle.h"
 #include "RemoteControlField.h"
-#include "Rundown/DetailsView/RemoteControl/Properties/AvaRundownPageRemoteControlWidgetUtils.h"
+#include "Rundown/DetailsView/RemoteControl/Properties/AvaRundownPageRCObject.h"
 #include "Rundown/DetailsView/RemoteControl/Properties/SAvaRundownPageRemoteControlProps.h"
 
 TSharedPtr<FAvaRundownRCDetailTreeNodeItem> FAvaRundownRCDetailTreeNodeItem::CreateItem(const TSharedRef<SAvaRundownPageRemoteControlProps>& InPropertyPanel
 	, const TSharedRef<FRemoteControlProperty>& InPropertyEntity
 	, bool bInControlled)
 {
-	TSharedPtr<IRemoteControlPropertyHandle> PropertyHandle = InPropertyEntity->GetPropertyHandle();
-	if (!PropertyHandle.IsValid())
+	TArray<UObject*> Objects = InPropertyEntity->GetBoundObjects();
+	if (Objects.IsEmpty())
 	{
 		return nullptr;
 	}
 
-	TArray<UObject*> Objects = InPropertyEntity->GetBoundObjects();
-	if (Objects.IsEmpty())
+	FAvaRundownPageRCObject& PageRCObject = InPropertyPanel->FindOrAddPageRCObject(Objects[0]);
+
+	TSharedPtr<IDetailTreeNode> PropertyNode = PageRCObject.FindTreeNode(InPropertyEntity->FieldPathInfo);
+	if (!PropertyNode.IsValid())
 	{
 		return nullptr;
 	}
@@ -26,41 +28,23 @@ TSharedPtr<FAvaRundownRCDetailTreeNodeItem> FAvaRundownRCDetailTreeNodeItem::Cre
 	TSharedRef<FAvaRundownRCDetailTreeNodeItem> PropertyItem = MakeShared<FAvaRundownRCDetailTreeNodeItem>();
 	PropertyItem->EntityOwnerWeak = InPropertyEntity;
 	PropertyItem->bEntityControlled = bInControlled;
+	PropertyItem->Initialize(PropertyNode.ToSharedRef());
 
-	TSharedRef<IPropertyRowGenerator> Generator = InPropertyPanel->FindOrAddGenerator(Objects[0]);
-
-	TSharedPtr<IDetailTreeNode> Node = FAvaRundownPageRemoteControlWidgetUtils::FindNode(Generator->GetRootTreeNodes()
-		, PropertyHandle->GetProperty()
-		, FAvaRundownPageRemoteControlWidgetUtils::EFindNodeMethod::Path);
+	PropertyItem->NodeWidgets.NameWidget = SNew(STextBlock)
+		.Margin(FMargin(8.f, 2.f, 0.f, 2.f))
+		.Text(FText::FromName(InPropertyEntity->GetLabel()))
+		.Font(IDetailLayoutBuilder::GetDetailFont());
 
-	PropertyItem->Initialize(Node);
 	return PropertyItem;
 }
 
-void FAvaRundownRCDetailTreeNodeItem::Initialize(TSharedPtr<IDetailTreeNode> InDetailTreeNode)
+void FAvaRundownRCDetailTreeNodeItem::Initialize(TSharedRef<IDetailTreeNode> InDetailTreeNode)
 {
 	DetailTreeNode = InDetailTreeNode;
-
-	if (DetailTreeNode.IsValid())
-	{
-		NodeWidgets = DetailTreeNode->CreateNodeWidgets();
-	}
-
+	NodeWidgets = InDetailTreeNode->CreateNodeWidgets();
 	RefreshChildren();
 }
 
-TSharedPtr<SWidget> FAvaRundownRCDetailTreeNodeItem::GetNameWidget() const
-{
-	return NodeWidgets.NameWidget;
-}
-
-TSharedPtr<SWidget> FAvaRundownRCDetailTreeNodeItem::GetValueWidget() const
-{
-	return NodeWidgets.WholeRowWidget.IsValid()
-		? NodeWidgets.WholeRowWidget
-		: NodeWidgets.ValueWidget;
-}
-
 void FAvaRundownRCDetailTreeNodeItem::RefreshChildren()
 {
 	Children.Reset();
diff --git a/Engine/Plugins/Experimental/Avalanche/Source/AvalancheMediaEditor/Private/Rundown/DetailsView/RemoteControl/Properties/Item/AvaRundownRCDetailTreeNodeItem.h b/Engine/Plugins/Experimental/Avalanche/Source/AvalancheMediaEditor/Private/Rundown/DetailsView/RemoteControl/Properties/Item/AvaRundownRCDetailTreeNodeItem.h
--- a/Engine/Plugins/Experimental/Avalanche/Source/AvalancheMediaEditor/Private/Rundown/DetailsView/RemoteControl/Properties/Item/AvaRundownRCDetailTreeNodeItem.h
+++ b/Engine/Plugins/Experimental/Avalanche/Source/AvalancheMediaEditor/Private/Rundown/DetailsView/RemoteControl/Properties/Item/AvaRundownRCDetailTreeNodeItem.h
@@ -16,17 +16,12 @@ public:
 		, const TSharedRef<FRemoteControlProperty>& InPropertyEntity
 		, bool bInControlled);
 
-	void Initialize(TSharedPtr<IDetailTreeNode> InDetailTreeNode);
+	void Initialize(TSharedRef<IDetailTreeNode> InDetailTreeNode);
 
 	//~ Begin FAvaRundownRCFieldItem
-	virtual TSharedPtr<SWidget> GetNameWidget() const override;
-	virtual TSharedPtr<SWidget> GetValueWidget() const override;
+	virtual void RefreshChildren() override;
 	//~ End FAvaRundownRCFieldItem
 
 private:
-	void RefreshChildren();
-
 	TSharedPtr<IDetailTreeNode> DetailTreeNode;
-
-	FNodeWidgets NodeWidgets;
 };
diff --git a/Engine/Plugins/Experimental/Avalanche/Source/AvalancheMediaEditor/Private/Rundown/DetailsView/RemoteControl/Properties/Item/AvaRundownRCFieldItem.cpp b/Engine/Plugins/Experimental/Avalanche/Source/AvalancheMediaEditor/Private/Rundown/DetailsView/RemoteControl/Properties/Item/AvaRundownRCFieldItem.cpp
--- a/Engine/Plugins/Experimental/Avalanche/Source/AvalancheMediaEditor/Private/Rundown/DetailsView/RemoteControl/Properties/Item/AvaRundownRCFieldItem.cpp
+++ b/Engine/Plugins/Experimental/Avalanche/Source/AvalancheMediaEditor/Private/Rundown/DetailsView/RemoteControl/Properties/Item/AvaRundownRCFieldItem.cpp
@@ -19,36 +19,30 @@ TSharedPtr<FAvaRundownRCFieldItem> FAvaRundownRCFieldItem::CreateItem(const TSha
 
 	if (FieldEntity->FieldType == EExposedFieldType::Function)
 	{
-		TSharedRef<FAvaRundownRCFunctionItem> FunctionItem = MakeShared<FAvaRundownRCFunctionItem>();
-		FunctionItem->EntityOwnerWeak = FieldEntity;
-		FunctionItem->bEntityControlled = bInControlled;
-		return FunctionItem;
+		return FAvaRundownRCFunctionItem::CreateItem(InPropertyPanel
+			, StaticCastSharedRef<FRemoteControlFunction>(InEntity)
+			, bInControlled);
 	}
 
 	if (FieldEntity->FieldType == EExposedFieldType::Property)
 	{
-		TSharedPtr<FAvaRundownRCDetailTreeNodeItem> DetailTreeNodeItem = FAvaRundownRCDetailTreeNodeItem::CreateItem(InPropertyPanel
+		return FAvaRundownRCDetailTreeNodeItem::CreateItem(InPropertyPanel
 			, StaticCastSharedRef<FRemoteControlProperty>(InEntity)
 			, bInControlled);
-
-		// Reset Name Widget so that it defaults to the entity widget
-		if (DetailTreeNodeItem.IsValid())
-		{
-			DetailTreeNodeItem->NodeWidgets.NameWidget.Reset();
-		}
-
-		return DetailTreeNodeItem;
 	}
 
 	return nullptr;
 }
 
-TSharedRef<ITableRow> FAvaRundownRCFieldItem::CreateWidget(const TSharedRef<SAvaRundownPageRemoteControlProps>& InPropertyPanel, const TSharedRef<STableViewBase>& InOwnerTable)
+FAvaRundownRCFieldItem::FAvaRundownRCFieldItem()
 {
-	return SNew(SAvaRundownRCPropertyItemRow, InPropertyPanel, InOwnerTable, SharedThis(this));
+	NodeWidgets.NameWidgetLayoutData.HorizontalAlignment = HAlign_Fill;
+	NodeWidgets.NameWidgetLayoutData.VerticalAlignment = VAlign_Fill;
+	NodeWidgets.ValueWidgetLayoutData.HorizontalAlignment = HAlign_Fill;
+	NodeWidgets.ValueWidgetLayoutData.VerticalAlignment = VAlign_Fill;
 }
 
-TSharedPtr<SWidget> FAvaRundownRCFieldItem::GetNameWidget() const
+TSharedRef<ITableRow> FAvaRundownRCFieldItem::CreateWidget(const TSharedRef<SAvaRundownPageRemoteControlProps>& InPropertyPanel, const TSharedRef<STableViewBase>& InOwnerTable)
 {
-	return nullptr;
+	return SNew(SAvaRundownRCPropertyItemRow, InPropertyPanel, InOwnerTable, SharedThis(this));
 }
diff --git a/Engine/Plugins/Experimental/Avalanche/Source/AvalancheMediaEditor/Private/Rundown/DetailsView/RemoteControl/Properties/Item/AvaRundownRCFieldItem.h b/Engine/Plugins/Experimental/Avalanche/Source/AvalancheMediaEditor/Private/Rundown/DetailsView/RemoteControl/Properties/Item/AvaRundownRCFieldItem.h
--- a/Engine/Plugins/Experimental/Avalanche/Source/AvalancheMediaEditor/Private/Rundown/DetailsView/RemoteControl/Properties/Item/AvaRundownRCFieldItem.h
+++ b/Engine/Plugins/Experimental/Avalanche/Source/AvalancheMediaEditor/Private/Rundown/DetailsView/RemoteControl/Properties/Item/AvaRundownRCFieldItem.h
@@ -2,6 +2,7 @@
 
 #pragma once
 
+#include "IDetailTreeNode.h"
 #include "Internationalization/Text.h"
 #include "Templates/SharedPointer.h"
 
@@ -18,6 +19,8 @@ public:
 		, const TSharedRef<FRemoteControlEntity>& InEntity
 		, bool bInControlled);
 
+	FAvaRundownRCFieldItem();
+
 	virtual ~FAvaRundownRCFieldItem() = default;
 
 	TSharedPtr<FRemoteControlEntity> GetEntity() const
@@ -30,6 +33,8 @@ public:
 		return bEntityControlled;
 	}
 
+	virtual void RefreshChildren() {}
+
 	TConstArrayView<TSharedPtr<FAvaRundownRCFieldItem>> GetChildren() const
 	{
 		return Children;
@@ -37,9 +42,10 @@ public:
 
 	TSharedRef<ITableRow> CreateWidget(const TSharedRef<SAvaRundownPageRemoteControlProps>& InPropertyPanel, const TSharedRef<STableViewBase>& InOwnerTable);
 
-	virtual TSharedPtr<SWidget> GetNameWidget() const;
-
-	virtual TSharedPtr<SWidget> GetValueWidget() const = 0;
+	const FNodeWidgets& GetNodeWidgets() const
+	{
+		return NodeWidgets;
+	}
 
 protected:
 	TWeakPtr<FRemoteControlEntity> EntityOwnerWeak;
@@ -47,4 +53,6 @@ protected:
 	bool bEntityControlled = false;
 
 	TArray<TSharedPtr<FAvaRundownRCFieldItem>> Children;
+
+	FNodeWidgets NodeWidgets;
 };
diff --git a/Engine/Plugins/Experimental/Avalanche/Source/AvalancheMediaEditor/Private/Rundown/DetailsView/RemoteControl/Properties/Item/AvaRundownRCFunctionItem.cpp b/Engine/Plugins/Experimental/Avalanche/Source/AvalancheMediaEditor/Private/Rundown/DetailsView/RemoteControl/Properties/Item/AvaRundownRCFunctionItem.cpp
--- a/Engine/Plugins/Experimental/Avalanche/Source/AvalancheMediaEditor/Private/Rundown/DetailsView/RemoteControl/Properties/Item/AvaRundownRCFunctionItem.cpp
+++ b/Engine/Plugins/Experimental/Avalanche/Source/AvalancheMediaEditor/Private/Rundown/DetailsView/RemoteControl/Properties/Item/AvaRundownRCFunctionItem.cpp
@@ -5,18 +5,37 @@
 #include "ScopedTransaction.h"
 #include "UObject/Script.h"
 #include "Widgets/Input/SButton.h"
+#include "Widgets/Layout/SBox.h"
 #include "Widgets/Text/STextBlock.h"
 
 #define LOCTEXT_NAMESPACE "AvaRundownRCFunctionItem"
 
-TSharedPtr<SWidget> FAvaRundownRCFunctionItem::GetValueWidget() const
+TSharedPtr<FAvaRundownRCFunctionItem> FAvaRundownRCFunctionItem::CreateItem(const TSharedRef<SAvaRundownPageRemoteControlProps>& InPropertyPanel
+	, const TSharedRef<FRemoteControlFunction>& InFunctionEntity
+	, bool bInControlled)
 {
-	return SNew(SButton)
-		.VAlign(VAlign_Center)
-		.OnClicked(this, &FAvaRundownRCFunctionItem::OnFunctionButtonClicked)
+	TSharedRef<FAvaRundownRCFunctionItem> FunctionItem = MakeShared<FAvaRundownRCFunctionItem>();
+	FunctionItem->EntityOwnerWeak = InFunctionEntity;
+	FunctionItem->bEntityControlled = bInControlled;
+	FunctionItem->Initialize();
+	return FunctionItem;
+}
+
+void FAvaRundownRCFunctionItem::Initialize()
+{
+	NodeWidgets.NameWidgetLayoutData.VerticalAlignment = VAlign_Center;
+	NodeWidgets.ValueWidgetLayoutData.HorizontalAlignment = HAlign_Left;
+
+	NodeWidgets.ValueWidget = SNew(SBox)
+		.Padding(0.f, 3.f)
 		[
-			SNew(STextBlock)
-			.Text(LOCTEXT("CallFunctionLabel", "Call Function"))
+			SNew(SButton)
+			.VAlign(VAlign_Center)
+			.OnClicked(this, &FAvaRundownRCFunctionItem::OnFunctionButtonClicked)
+			[
+				SNew(STextBlock)
+				.Text(LOCTEXT("CallFunctionLabel", "Call Function"))
+			]
 		];
 }
 
diff --git a/Engine/Plugins/Experimental/Avalanche/Source/AvalancheMediaEditor/Private/Rundown/DetailsView/RemoteControl/Properties/Item/AvaRundownRCFunctionItem.h b/Engine/Plugins/Experimental/Avalanche/Source/AvalancheMediaEditor/Private/Rundown/DetailsView/RemoteControl/Properties/Item/AvaRundownRCFunctionItem.h
--- a/Engine/Plugins/Experimental/Avalanche/Source/AvalancheMediaEditor/Private/Rundown/DetailsView/RemoteControl/Properties/Item/AvaRundownRCFunctionItem.h
+++ b/Engine/Plugins/Experimental/Avalanche/Source/AvalancheMediaEditor/Private/Rundown/DetailsView/RemoteControl/Properties/Item/AvaRundownRCFunctionItem.h
@@ -4,15 +4,18 @@
 
 #include "AvaRundownRCFieldItem.h"
 
+struct FRemoteControlFunction;
 class FReply;
 
 class FAvaRundownRCFunctionItem : public FAvaRundownRCFieldItem
 {
 public:
-	//~ Begin FAvaRundownRCFieldItem
-	virtual TSharedPtr<SWidget> GetValueWidget() const override;
-	//~ End FAvaRundownRCFieldItem
+	static TSharedPtr<FAvaRundownRCFunctionItem> CreateItem(const TSharedRef<SAvaRundownPageRemoteControlProps>& InPropertyPanel
+		, const TSharedRef<FRemoteControlFunction>& InFunctionEntity
+		, bool bInControlled);
 
 private:
+	void Initialize();
+
 	FReply OnFunctionButtonClicked() const;
 };
diff --git a/Engine/Plugins/Experimental/Avalanche/Source/AvalancheMediaEditor/Private/Rundown/DetailsView/RemoteControl/Properties/SAvaRundownPageRemoteControlProps.cpp b/Engine/Plugins/Experimental/Avalanche/Source/AvalancheMediaEditor/Private/Rundown/DetailsView/RemoteControl/Properties/SAvaRundownPageRemoteControlProps.cpp
--- a/Engine/Plugins/Experimental/Avalanche/Source/AvalancheMediaEditor/Private/Rundown/DetailsView/RemoteControl/Properties/SAvaRundownPageRemoteControlProps.cpp
+++ b/Engine/Plugins/Experimental/Avalanche/Source/AvalancheMediaEditor/Private/Rundown/DetailsView/RemoteControl/Properties/SAvaRundownPageRemoteControlProps.cpp
@@ -56,7 +56,7 @@ public:
 				OngoingPropertyChanges.Add(InPropertyThatChanged);
 				RundownEditor->BeginModify();
 			}
-		
+
 			// Apply change to page immediately to capture it in the transaction.
 			Panel->OnPostPropertyChanged(InPropertyThatChanged);
 
@@ -213,7 +213,6 @@ void SAvaRundownPageRemoteControlProps::Refresh(const TArray<int32>& InSelectedP
 		};
 
 		TArray<FAvaPropertyDetails> NewItems;
-		int32 TotalNumExposedEntities = 0;
 
 		for (const TSharedPtr<FAvaRundownManagedInstance>& ManagedInstance : ManagedInstances)
 		{
@@ -222,10 +221,12 @@ void SAvaRundownPageRemoteControlProps::Refresh(const TArray<int32>& InSelectedP
 			{
 				continue;
 			}
-			
+
 			const TArray<TWeakPtr<FRemoteControlEntity>> ExposedEntities = RemoteControlPreset->GetExposedEntities<FRemoteControlEntity>();
-			NewItems.Reserve(TotalNumExposedEntities += ExposedEntities.Num());
-			
+
+			TArray<FAvaPropertyDetails> ItemsToAdd;
+			ItemsToAdd.Reserve(ExposedEntities.Num());
+
 			for (const TWeakPtr<FRemoteControlEntity>& EntityWeakPtr : ExposedEntities)
 			{
 				const TSharedPtr<FRemoteControlEntity> Entity = EntityWeakPtr.Pin();
@@ -244,7 +245,7 @@ void SAvaRundownPageRemoteControlProps::Refresh(const TArray<int32>& InSelectedP
 
 				if (FieldEntity->FieldType == EExposedFieldType::Function)
 				{
-					NewItems.Add({ FieldEntity, false });
+					ItemsToAdd.Add({ FieldEntity, false });
 				}
 				else if (FieldEntity->FieldType == EExposedFieldType::Property)
 				{
@@ -288,9 +289,22 @@ void SAvaRundownPageRemoteControlProps::Refresh(const TArray<int32>& InSelectedP
 					}
 
 					const bool bEntityControlled = ManagedInstance->GetDefaultRemoteControlValues().EntitiesControlledByController.Contains(FieldEntity->GetId());
-					NewItems.Add({ FieldEntity, bEntityControlled });
+					ItemsToAdd.Add({ FieldEntity, bEntityControlled });
 				}
 			}
+
+			const TArray<FGuid>& DefaultGroupOrder = RemoteControlPreset->Layout.GetDefaultGroupOrder();
+
+			// Before adding to new Items, sort the entities based on the Default Group Order.
+			ItemsToAdd.Sort(
+				[&DefaultGroupOrder](const FAvaPropertyDetails& InLeftItem, const FAvaPropertyDetails& InRightItem)
+				{
+					const int32 LeftIndex = DefaultGroupOrder.IndexOfByKey(InLeftItem.Entity->GetId());
+					const int32 RightIndex = DefaultGroupOrder.IndexOfByKey(InRightItem.Entity->GetId());
+					return LeftIndex < RightIndex;
+				});
+
+			NewItems.Append(MoveTemp(ItemsToAdd));
 		}
 
 		FieldContainer->RebuildList();
@@ -344,6 +358,20 @@ void SAvaRundownPageRemoteControlProps::Refresh(const TArray<int32>& InSelectedP
 	}
 }
 
+void SAvaRundownPageRemoteControlProps::Tick(const FGeometry& InAllottedGeometry, const double InCurrentTime, const float InDeltaTime)
+{
+	// This tick happens before Tree View tick
+	// Ensure all the items have their hierarchies up-to-date prior to Tree Refresh
+	if (bRefreshRequested)
+	{
+		for (const TSharedPtr<FAvaRundownRCFieldItem>& FieldItem : FieldItems)
+		{
+			FieldItem->RefreshChildren();
+		}
+		bRefreshRequested = false;
+	}
+}
+
 // Remarks:
 // This is called by URemoteControlPreset::OnEndFrame() as a result of an entity being modified.
 // However, it doesn't seem to be called (or not always) if the entity is modified by a controller action. 
@@ -362,6 +390,8 @@ void SAvaRundownPageRemoteControlProps::OnRemoteControlExposedPropertiesModified
 			RundownEditor->MarkAsModified();
 		}
 	}
+
+	RequestRefresh();
 }
 
 void SAvaRundownPageRemoteControlProps::OnRemoteControlControllerModified(URemoteControlPreset* InPreset, const TSet<FGuid>& InModifiedControllerIds)
@@ -385,6 +415,8 @@ void SAvaRundownPageRemoteControlProps::OnRemoteControlControllerModified(URemot
 	// It seems OnPropertyChangedDelegate (OnExposedPropertiesModified()) is not called when properties are
 	// changed by controllers. Ensure the values are saved by calling our handler directly.
 	OnRemoteControlExposedPropertiesModified(InPreset, EntityIds);
+
+	RequestRefresh();
 }
 
 void SAvaRundownPageRemoteControlProps::OnPostPropertyChanged(FProperty* InPropertyThatChanged)
@@ -406,6 +438,17 @@ void SAvaRundownPageRemoteControlProps::OnPostPropertyChanged(FProperty* InPrope
 			}
 		}
 	}
+
+	RequestRefresh();
+}
+
+void SAvaRundownPageRemoteControlProps::RequestRefresh()
+{
+	if (FieldContainer.IsValid())
+	{
+		FieldContainer->RequestTreeRefresh();
+	}
+	bRefreshRequested = true;
 }
 
 void SAvaRundownPageRemoteControlProps::BindRemoteControlDelegates(URemoteControlPreset* InPreset)
@@ -556,17 +599,18 @@ TSharedPtr<FAvaRundownPageRCPropsNotifyHook> SAvaRundownPageRemoteControlProps::
 	return NotifyHook;
 }
 
-TSharedRef<IPropertyRowGenerator> SAvaRundownPageRemoteControlProps::FindOrAddGenerator(UObject* InObject)
+FAvaRundownPageRCObject& SAvaRundownPageRemoteControlProps::FindOrAddPageRCObject(UObject* InObject)
 {
-	TSharedPtr<IPropertyRowGenerator>& Generator = ObjectGeneratorMap.FindOrAdd(InObject);
-	if (!Generator.IsValid())
+	bool bIsAlreadyInSet;
+	FAvaRundownPageRCObject& PageRCObject = PageRCObjectMap.FindOrAdd(FAvaRundownPageRCObject(InObject), &bIsAlreadyInSet);
+
+	// Initialize if new element in set
+	if (!bIsAlreadyInSet)
 	{
-		FPropertyRowGeneratorArgs Args;
-		Args.NotifyHook = NotifyHook.Get();
-		Generator = FModuleManager::GetModuleChecked<FPropertyEditorModule>("PropertyEditor").CreatePropertyRowGenerator(Args);
-		Generator->SetObjects({ InObject });
+		PageRCObject.Initialize(NotifyHook.Get());
 	}
-	return Generator.ToSharedRef();
+
+	return PageRCObject;
 }
 
 TSharedPtr<SWidget> SAvaRundownPageRemoteControlProps::GetContextMenuContent()
diff --git a/Engine/Plugins/Experimental/Avalanche/Source/AvalancheMediaEditor/Private/Rundown/DetailsView/RemoteControl/Properties/SAvaRundownPageRemoteControlProps.h b/Engine/Plugins/Experimental/Avalanche/Source/AvalancheMediaEditor/Private/Rundown/DetailsView/RemoteControl/Properties/SAvaRundownPageRemoteControlProps.h
--- a/Engine/Plugins/Experimental/Avalanche/Source/AvalancheMediaEditor/Private/Rundown/DetailsView/RemoteControl/Properties/SAvaRundownPageRemoteControlProps.h
+++ b/Engine/Plugins/Experimental/Avalanche/Source/AvalancheMediaEditor/Private/Rundown/DetailsView/RemoteControl/Properties/SAvaRundownPageRemoteControlProps.h
@@ -2,6 +2,7 @@
 
 #pragma once
 
+#include "AvaRundownPageRCObject.h"
 #include "Containers/Array.h"
 #include "Containers/Map.h"
 #include "Containers/Set.h"
@@ -19,7 +20,6 @@ class FAvaRundownRCDetailTreeNodeItem;
 class FAvaRundownRCFieldItem;
 class FName;
 class FUICommandList;
-class IPropertyRowGenerator;
 class ITableRow;
 class SAvaRundownPageRemoteControlProps;
 class STableViewBase;
@@ -30,6 +30,7 @@ struct FAvaPlayableRemoteControlValue;
 struct FAvaRundownPage;
 struct FGuid;
 struct FRemoteControlEntity;
+template<typename ItemType> class STreeView;
 
 DECLARE_MULTICAST_DELEGATE_TwoParams(FAvaRundownRCPropertyHeaderRowExtensionDelegate, TSharedRef<SAvaRundownPageRemoteControlProps> Panel,
 	TSharedRef<SHeaderRow>& HeaderRow)
@@ -74,11 +75,15 @@ public:
 	/** Refreshes the content of this widget. */
 	void Refresh(const TArray<int32>& InSelectedPageIds);
 
+	//~ Begin SWidget
+	virtual void Tick(const FGeometry& InAllottedGeometry, const double InCurrentTime, const float InDeltaTime) override;
+	//~ End SWidget
+
 	const TArray<TSharedPtr<FAvaRundownRCFieldItem>> GetSelectedPropertyItems() const;
 	
 	TSharedPtr<FAvaRundownPageRCPropsNotifyHook> GetNotifyHook() const;
 
-	TSharedRef<IPropertyRowGenerator> FindOrAddGenerator(UObject* InObject);
+	FAvaRundownPageRCObject& FindOrAddPageRCObject(UObject* InObject);
 
 	/**
 	 * Utility to save the given Preset's Entities to the given page's property.
@@ -101,6 +106,8 @@ private:
 
 	void OnPostPropertyChanged(FProperty* InPropertyThatChanged);
 
+	void RequestRefresh();
+
 	void BindRemoteControlDelegates(URemoteControlPreset* InPreset);
 
 	bool HasRemoteControlPreset(const URemoteControlPreset* InPreset) const;
@@ -136,8 +143,8 @@ private:
 	TArray<TSharedPtr<FAvaRundownManagedInstance>> ManagedInstances;
 
 	/** The widget that lists the property rows. */
-	TSharedPtr<SListView<TSharedPtr<FAvaRundownRCFieldItem>>> FieldContainer;
-	
+	TSharedPtr<STreeView<TSharedPtr<FAvaRundownRCFieldItem>>> FieldContainer;
+
 	/** The data used to back the properties container list view. */
 	TArray<TSharedPtr<FAvaRundownRCFieldItem>> FieldItems;
 
@@ -149,7 +156,9 @@ private:
 
 	TSharedPtr<FAvaRundownPageRCPropsNotifyHook> NotifyHook;
 
-	TMap<FObjectKey, TSharedPtr<IPropertyRowGenerator>> ObjectGeneratorMap;
+	TSet<FAvaRundownPageRCObject> PageRCObjectMap;
 
 	friend class FAvaRundownPageRCPropsNotifyHookImpl;
+
+	bool bRefreshRequested = false;
 };
\ No newline at end of file
diff --git a/Engine/Plugins/Experimental/Avalanche/Source/AvalancheMediaEditor/Private/Rundown/DetailsView/RemoteControl/Properties/SAvaRundownRCPropertyItemRow.cpp b/Engine/Plugins/Experimental/Avalanche/Source/AvalancheMediaEditor/Private/Rundown/DetailsView/RemoteControl/Properties/SAvaRundownRCPropertyItemRow.cpp
--- a/Engine/Plugins/Experimental/Avalanche/Source/AvalancheMediaEditor/Private/Rundown/DetailsView/RemoteControl/Properties/SAvaRundownRCPropertyItemRow.cpp
+++ b/Engine/Plugins/Experimental/Avalanche/Source/AvalancheMediaEditor/Private/Rundown/DetailsView/RemoteControl/Properties/SAvaRundownRCPropertyItemRow.cpp
@@ -1,7 +1,6 @@
 // Copyright Epic Games, Inc. All Rights Reserved.
 
 #include "SAvaRundownRCPropertyItemRow.h"
-#include "AvaRundownPageRemoteControlWidgetUtils.h"
 #include "GameFramework/Actor.h"
 #include "IDetailTreeNode.h"
 #include "IPropertyRowGenerator.h"
@@ -38,15 +37,12 @@ TSharedRef<SWidget> SAvaRundownRCPropertyItemRow::GenerateWidgetForColumn(const
 			return SNew(SHorizontalBox)
 				+ SHorizontalBox::Slot()
 				.AutoWidth()
-				.VAlign(VAlign_Fill)
 				[
 					SNew(SExpanderArrow, SharedThis(this))
 					.IndentAmount(12)
 					.ShouldDrawWires(true)
 				]
 				+ SHorizontalBox::Slot()
-				.HAlign(HAlign_Left)
-				.VAlign(VAlign_Center)
 				[
 					CreateName()
 				];
@@ -54,6 +50,7 @@ TSharedRef<SWidget> SAvaRundownRCPropertyItemRow::GenerateWidgetForColumn(const
 		else if (InColumnName == SAvaRundownPageRemoteControlProps::ValueColumnName)
 		{
 			return SAssignNew(ValueContainer, SBox)
+				.MinDesiredHeight(26.f)
 				[
 					CreateValue()
 				];
@@ -99,21 +96,22 @@ TSharedRef<SWidget> SAvaRundownRCPropertyItemRow::CreateName()
 		return SNullWidget::NullWidget;
 	}
 
-	TSharedPtr<SWidget> NameWidget = FieldItem->GetNameWidget();
+	const FNodeWidgets& NodeWidgets = FieldItem->GetNodeWidgets();
+
+	TSharedPtr<SWidget> NameWidget = NodeWidgets.NameWidget;
 	if (!NameWidget.IsValid())
 	{
-		if (TSharedPtr<FRemoteControlEntity> Entity = FieldItem->GetEntity())
-		{
-			TSharedPtr<FRemoteControlField> FieldEntity = StaticCastSharedPtr<FRemoteControlField>(Entity);
-			return SNew(STextBlock)
-				.Margin(FMargin(8.f, 2.f, 0.f, 2.f))
-				.Text(FText::FromName(FieldEntity->GetLabel()))
-				.ToolTipText(this, &SAvaRundownRCPropertyItemRow::GetPropertyTooltipText);
-		}
+		return SNullWidget::NullWidget;
 	}
 
 	NameWidget->SetToolTipText(TAttribute<FText>::CreateSP(this, &SAvaRundownRCPropertyItemRow::GetPropertyTooltipText));
-	return NameWidget.ToSharedRef();
+
+	return SNew(SBox)
+		.HAlign(NodeWidgets.NameWidgetLayoutData.HorizontalAlignment)
+		.VAlign(NodeWidgets.NameWidgetLayoutData.VerticalAlignment)
+		[
+			NameWidget.ToSharedRef()
+		];
 }
 
 TSharedRef<SWidget> SAvaRundownRCPropertyItemRow::CreateValue()
@@ -124,22 +122,24 @@ TSharedRef<SWidget> SAvaRundownRCPropertyItemRow::CreateValue()
 		return SNullWidget::NullWidget;
 	}
 
-	TSharedPtr<SWidget> ValueWidget = FieldItem->GetValueWidget();
-	if (!ValueWidget.IsValid())
+	const FNodeWidgets& NodeWidgets = FieldItem->GetNodeWidgets();
+	if (!NodeWidgets.ValueWidget.IsValid())
 	{
 		return SNullWidget::NullWidget;
 	}
 
+	TSharedRef<SHorizontalBox> ValueRow = SNew(SHorizontalBox);
+
+	ValueRow->AddSlot()
+		.HAlign(NodeWidgets.ValueWidgetLayoutData.HorizontalAlignment)
+		.VAlign(NodeWidgets.ValueWidgetLayoutData.VerticalAlignment)
+		[
+			NodeWidgets.ValueWidget.ToSharedRef()
+		];
+
 	if (FieldItem->IsEntityControlled())
 	{
-		ValueWidget->SetEnabled(false);
-		ValueWidget = SNew(SHorizontalBox)
-			+ SHorizontalBox::Slot()
-			.VAlign(EVerticalAlignment::VAlign_Center)
-			[
-				ValueWidget.ToSharedRef()
-			]
-			+ SHorizontalBox::Slot()
+		ValueRow->AddSlot()
 			.VAlign(EVerticalAlignment::VAlign_Center)
 			.Padding(3.f, 0.f, 0.f, 0.f)
 			[
@@ -148,7 +148,7 @@ TSharedRef<SWidget> SAvaRundownRCPropertyItemRow::CreateValue()
 			];
 	}
 
-	return ValueWidget.ToSharedRef();
+	return ValueRow;
 }
 
 FText SAvaRundownRCPropertyItemRow::GetPropertyTooltipText() const
