MotionDesign : Shapes
- Fix shape spawn at 0,0,0 when default action is used even in MotionDesignViewport

diff --git a/Engine/Plugins/Experimental/Avalanche/Source/AvalancheShapesEditor/Private/Tools/AvaShapesEditorShapeToolBase.cpp b/Engine/Plugins/Experimental/Avalanche/Source/AvalancheShapesEditor/Private/Tools/AvaShapesEditorShapeToolBase.cpp
--- a/Engine/Plugins/Experimental/Avalanche/Source/AvalancheShapesEditor/Private/Tools/AvaShapesEditorShapeToolBase.cpp
+++ b/Engine/Plugins/Experimental/Avalanche/Source/AvalancheShapesEditor/Private/Tools/AvaShapesEditorShapeToolBase.cpp
@@ -94,10 +94,5 @@ void UAvaShapesEditorShapeToolBase::SetShapeSize(AAvaShapeActor* InShapeActor, c
 
 bool UAvaShapesEditorShapeToolBase::UseIdentityLocation() const
 {
-	if (!bPerformingDefaultAction)
-	{
-		return false;
-	}
-
-	return !IsMotionDesignViewport();
+	return bPerformingDefaultAction;
 }
