Motion Design: Additional fixes to rundown properties to keep expansion states, refresh the root item detail tree nodes too (where before it would only refresh the existing children nodes)

diff --git a/Engine/Plugins/Experimental/Avalanche/Source/AvalancheMediaEditor/Private/Rundown/DetailsView/RemoteControl/Properties/AvaRundownPageRCObject.cpp b/Engine/Plugins/Experimental/Avalanche/Source/AvalancheMediaEditor/Private/Rundown/DetailsView/RemoteControl/Properties/AvaRundownPageRCObject.cpp
--- a/Engine/Plugins/Experimental/Avalanche/Source/AvalancheMediaEditor/Private/Rundown/DetailsView/RemoteControl/Properties/AvaRundownPageRCObject.cpp
+++ b/Engine/Plugins/Experimental/Avalanche/Source/AvalancheMediaEditor/Private/Rundown/DetailsView/RemoteControl/Properties/AvaRundownPageRCObject.cpp
@@ -25,9 +25,9 @@ void FAvaRundownPageRCObject::Initialize(FNotifyHook* InNotifyHook)
 	CacheTreeNodes();
 }
 
-TSharedPtr<IDetailTreeNode> FAvaRundownPageRCObject::FindTreeNode(const FRCFieldPathInfo& InPathInfo) const
+TSharedPtr<IDetailTreeNode> FAvaRundownPageRCObject::FindTreeNode(const FString& InPathInfo) const
 {
-	if (const TSharedRef<IDetailTreeNode>* FoundTreeNode = TreeNodeMap.Find(InPathInfo.ToString()))
+	if (const TSharedRef<IDetailTreeNode>* FoundTreeNode = TreeNodeMap.Find(InPathInfo))
 	{
 		return *FoundTreeNode;
 	}
diff --git a/Engine/Plugins/Experimental/Avalanche/Source/AvalancheMediaEditor/Private/Rundown/DetailsView/RemoteControl/Properties/AvaRundownPageRCObject.h b/Engine/Plugins/Experimental/Avalanche/Source/AvalancheMediaEditor/Private/Rundown/DetailsView/RemoteControl/Properties/AvaRundownPageRCObject.h
--- a/Engine/Plugins/Experimental/Avalanche/Source/AvalancheMediaEditor/Private/Rundown/DetailsView/RemoteControl/Properties/AvaRundownPageRCObject.h
+++ b/Engine/Plugins/Experimental/Avalanche/Source/AvalancheMediaEditor/Private/Rundown/DetailsView/RemoteControl/Properties/AvaRundownPageRCObject.h
@@ -22,7 +22,7 @@ public:
 
 	void Initialize(FNotifyHook* InNotifyHook);
 
-	TSharedPtr<IDetailTreeNode> FindTreeNode(const FRCFieldPathInfo& InPathInfo) const;
+	TSharedPtr<IDetailTreeNode> FindTreeNode(const FString& InPathInfo) const;
 
 	bool operator==(const FAvaRundownPageRCObject& InOther) const
 	{
@@ -34,9 +34,9 @@ public:
 		return GetTypeHash(InObject.ObjectKey);
 	}
 
-private:
 	void CacheTreeNodes();
 
+private:
 	FObjectKey ObjectKey;
 
 	TSharedPtr<IPropertyRowGenerator> PropertyRowGenerator;
diff --git a/Engine/Plugins/Experimental/Avalanche/Source/AvalancheMediaEditor/Private/Rundown/DetailsView/RemoteControl/Properties/Item/AvaRundownRCDetailTreeNodeItem.cpp b/Engine/Plugins/Experimental/Avalanche/Source/AvalancheMediaEditor/Private/Rundown/DetailsView/RemoteControl/Properties/Item/AvaRundownRCDetailTreeNodeItem.cpp
--- a/Engine/Plugins/Experimental/Avalanche/Source/AvalancheMediaEditor/Private/Rundown/DetailsView/RemoteControl/Properties/Item/AvaRundownRCDetailTreeNodeItem.cpp
+++ b/Engine/Plugins/Experimental/Avalanche/Source/AvalancheMediaEditor/Private/Rundown/DetailsView/RemoteControl/Properties/Item/AvaRundownRCDetailTreeNodeItem.cpp
@@ -11,24 +11,11 @@ TSharedPtr<FAvaRundownRCDetailTreeNodeItem> FAvaRundownRCDetailTreeNodeItem::Cre
 	, const TSharedRef<FRemoteControlProperty>& InPropertyEntity
 	, bool bInControlled)
 {
-	TArray<UObject*> Objects = InPropertyEntity->GetBoundObjects();
-	if (Objects.IsEmpty())
-	{
-		return nullptr;
-	}
-
-	FAvaRundownPageRCObject& PageRCObject = InPropertyPanel->FindOrAddPageRCObject(Objects[0]);
-
-	TSharedPtr<IDetailTreeNode> PropertyNode = PageRCObject.FindTreeNode(InPropertyEntity->FieldPathInfo);
-	if (!PropertyNode.IsValid())
-	{
-		return nullptr;
-	}
-
 	TSharedRef<FAvaRundownRCDetailTreeNodeItem> PropertyItem = MakeShared<FAvaRundownRCDetailTreeNodeItem>();
+	PropertyItem->PropertyEntityWeak = InPropertyEntity;
 	PropertyItem->EntityOwnerWeak = InPropertyEntity;
 	PropertyItem->bEntityControlled = bInControlled;
-	PropertyItem->Initialize(PropertyNode.ToSharedRef());
+	PropertyItem->Refresh(InPropertyPanel);
 
 	PropertyItem->NodeWidgets.NameWidget = SNew(STextBlock)
 		.Margin(FMargin(8.f, 2.f, 0.f, 2.f))
@@ -45,6 +32,39 @@ void FAvaRundownRCDetailTreeNodeItem::Initialize(TSharedRef<IDetailTreeNode> InD
 	RefreshChildren();
 }
 
+FStringView FAvaRundownRCDetailTreeNodeItem::GetPath() const
+{
+	return FieldPath;
+}
+
+void FAvaRundownRCDetailTreeNodeItem::Refresh(const TSharedRef<SAvaRundownPageRemoteControlProps>& InPropertyPanel)
+{
+	FieldPath.Reset();
+
+	TSharedPtr<FRemoteControlProperty> PropertyEntity = PropertyEntityWeak.Pin();
+	if (!PropertyEntity.IsValid())
+	{
+		return;
+	}
+
+	TArray<UObject*> Objects = PropertyEntity->GetBoundObjects();
+	if (Objects.IsEmpty())
+	{
+		return;
+	}
+
+	FAvaRundownPageRCObject& PageRCObject = InPropertyPanel->FindOrAddPageRCObject(Objects[0]);
+	FieldPath = PropertyEntity->FieldPathInfo.ToString();
+
+	TSharedPtr<IDetailTreeNode> PropertyNode = PageRCObject.FindTreeNode(FieldPath);
+	if (!PropertyNode.IsValid())
+	{
+		return;
+	}
+
+	Initialize(PropertyNode.ToSharedRef());
+}
+
 void FAvaRundownRCDetailTreeNodeItem::RefreshChildren()
 {
 	Children.Reset();
@@ -61,6 +81,7 @@ void FAvaRundownRCDetailTreeNodeItem::RefreshChildren()
 	for (const TSharedRef<IDetailTreeNode>& ChildNode : ChildNodes)
 	{
 		TSharedRef<FAvaRundownRCDetailTreeNodeItem> ChildPropertyItem = MakeShared<FAvaRundownRCDetailTreeNodeItem>();
+		ChildPropertyItem->PropertyEntityWeak = PropertyEntityWeak;
 		ChildPropertyItem->EntityOwnerWeak = EntityOwnerWeak;
 		ChildPropertyItem->bEntityControlled = bEntityControlled;
 		ChildPropertyItem->Initialize(ChildNode);
diff --git a/Engine/Plugins/Experimental/Avalanche/Source/AvalancheMediaEditor/Private/Rundown/DetailsView/RemoteControl/Properties/Item/AvaRundownRCDetailTreeNodeItem.h b/Engine/Plugins/Experimental/Avalanche/Source/AvalancheMediaEditor/Private/Rundown/DetailsView/RemoteControl/Properties/Item/AvaRundownRCDetailTreeNodeItem.h
--- a/Engine/Plugins/Experimental/Avalanche/Source/AvalancheMediaEditor/Private/Rundown/DetailsView/RemoteControl/Properties/Item/AvaRundownRCDetailTreeNodeItem.h
+++ b/Engine/Plugins/Experimental/Avalanche/Source/AvalancheMediaEditor/Private/Rundown/DetailsView/RemoteControl/Properties/Item/AvaRundownRCDetailTreeNodeItem.h
@@ -19,9 +19,17 @@ public:
 	void Initialize(TSharedRef<IDetailTreeNode> InDetailTreeNode);
 
 	//~ Begin FAvaRundownRCFieldItem
-	virtual void RefreshChildren() override;
+	virtual FStringView GetPath() const override;
+	virtual void Refresh(const TSharedRef<SAvaRundownPageRemoteControlProps>& InPropertyPanel) override;
 	//~ End FAvaRundownRCFieldItem
 
 private:
+	void RefreshChildren();
+
+	TWeakPtr<FRemoteControlProperty> PropertyEntityWeak;
+
 	TSharedPtr<IDetailTreeNode> DetailTreeNode;
+
+	/** Cached path of the Field */
+	FString FieldPath;
 };
diff --git a/Engine/Plugins/Experimental/Avalanche/Source/AvalancheMediaEditor/Private/Rundown/DetailsView/RemoteControl/Properties/Item/AvaRundownRCFieldItem.cpp b/Engine/Plugins/Experimental/Avalanche/Source/AvalancheMediaEditor/Private/Rundown/DetailsView/RemoteControl/Properties/Item/AvaRundownRCFieldItem.cpp
--- a/Engine/Plugins/Experimental/Avalanche/Source/AvalancheMediaEditor/Private/Rundown/DetailsView/RemoteControl/Properties/Item/AvaRundownRCFieldItem.cpp
+++ b/Engine/Plugins/Experimental/Avalanche/Source/AvalancheMediaEditor/Private/Rundown/DetailsView/RemoteControl/Properties/Item/AvaRundownRCFieldItem.cpp
@@ -46,3 +46,8 @@ TSharedRef<ITableRow> FAvaRundownRCFieldItem::CreateWidget(const TSharedRef<SAva
 {
 	return SNew(SAvaRundownRCPropertyItemRow, InPropertyPanel, InOwnerTable, SharedThis(this));
 }
+
+FStringView FAvaRundownRCFieldItem::GetPath() const
+{
+	return FStringView();
+}
diff --git a/Engine/Plugins/Experimental/Avalanche/Source/AvalancheMediaEditor/Private/Rundown/DetailsView/RemoteControl/Properties/Item/AvaRundownRCFieldItem.h b/Engine/Plugins/Experimental/Avalanche/Source/AvalancheMediaEditor/Private/Rundown/DetailsView/RemoteControl/Properties/Item/AvaRundownRCFieldItem.h
--- a/Engine/Plugins/Experimental/Avalanche/Source/AvalancheMediaEditor/Private/Rundown/DetailsView/RemoteControl/Properties/Item/AvaRundownRCFieldItem.h
+++ b/Engine/Plugins/Experimental/Avalanche/Source/AvalancheMediaEditor/Private/Rundown/DetailsView/RemoteControl/Properties/Item/AvaRundownRCFieldItem.h
@@ -33,7 +33,7 @@ public:
 		return bEntityControlled;
 	}
 
-	virtual void RefreshChildren() {}
+	virtual void Refresh(const TSharedRef<SAvaRundownPageRemoteControlProps>& InPropertyPanel) {}
 
 	TConstArrayView<TSharedPtr<FAvaRundownRCFieldItem>> GetChildren() const
 	{
@@ -47,12 +47,14 @@ public:
 		return NodeWidgets;
 	}
 
+	virtual FStringView GetPath() const;
+
 protected:
 	TWeakPtr<FRemoteControlEntity> EntityOwnerWeak;
 
-	bool bEntityControlled = false;
-
 	TArray<TSharedPtr<FAvaRundownRCFieldItem>> Children;
 
 	FNodeWidgets NodeWidgets;
+
+	bool bEntityControlled = false;
 };
diff --git a/Engine/Plugins/Experimental/Avalanche/Source/AvalancheMediaEditor/Private/Rundown/DetailsView/RemoteControl/Properties/Item/AvaRundownRCFunctionItem.cpp b/Engine/Plugins/Experimental/Avalanche/Source/AvalancheMediaEditor/Private/Rundown/DetailsView/RemoteControl/Properties/Item/AvaRundownRCFunctionItem.cpp
--- a/Engine/Plugins/Experimental/Avalanche/Source/AvalancheMediaEditor/Private/Rundown/DetailsView/RemoteControl/Properties/Item/AvaRundownRCFunctionItem.cpp
+++ b/Engine/Plugins/Experimental/Avalanche/Source/AvalancheMediaEditor/Private/Rundown/DetailsView/RemoteControl/Properties/Item/AvaRundownRCFunctionItem.cpp
@@ -1,6 +1,7 @@
 // Copyright Epic Games, Inc. All Rights Reserved.
 
 #include "AvaRundownRCFunctionItem.h"
+#include "DetailLayoutBuilder.h"
 #include "RemoteControlField.h"
 #include "ScopedTransaction.h"
 #include "UObject/Script.h"
@@ -26,6 +27,11 @@ void FAvaRundownRCFunctionItem::Initialize()
 	NodeWidgets.NameWidgetLayoutData.VerticalAlignment = VAlign_Center;
 	NodeWidgets.ValueWidgetLayoutData.HorizontalAlignment = HAlign_Left;
 
+	NodeWidgets.NameWidget = SNew(STextBlock)
+		.Margin(FMargin(8.f, 2.f, 0.f, 2.f))
+		.Text(GetLabel())
+		.Font(IDetailLayoutBuilder::GetDetailFont());
+
 	NodeWidgets.ValueWidget = SNew(SBox)
 		.Padding(0.f, 3.f)
 		[
@@ -39,6 +45,15 @@ void FAvaRundownRCFunctionItem::Initialize()
 		];
 }
 
+FText FAvaRundownRCFunctionItem::GetLabel() const
+{
+	if (TSharedPtr<FRemoteControlEntity> Entity = EntityOwnerWeak.Pin())
+	{
+		return FText::FromName(Entity->GetLabel());
+	}
+	return FText::GetEmpty();
+}
+
 FReply FAvaRundownRCFunctionItem::OnFunctionButtonClicked() const
 {
 	TSharedPtr<FRemoteControlEntity> Entity = GetEntity();
diff --git a/Engine/Plugins/Experimental/Avalanche/Source/AvalancheMediaEditor/Private/Rundown/DetailsView/RemoteControl/Properties/Item/AvaRundownRCFunctionItem.h b/Engine/Plugins/Experimental/Avalanche/Source/AvalancheMediaEditor/Private/Rundown/DetailsView/RemoteControl/Properties/Item/AvaRundownRCFunctionItem.h
--- a/Engine/Plugins/Experimental/Avalanche/Source/AvalancheMediaEditor/Private/Rundown/DetailsView/RemoteControl/Properties/Item/AvaRundownRCFunctionItem.h
+++ b/Engine/Plugins/Experimental/Avalanche/Source/AvalancheMediaEditor/Private/Rundown/DetailsView/RemoteControl/Properties/Item/AvaRundownRCFunctionItem.h
@@ -17,5 +17,7 @@ public:
 private:
 	void Initialize();
 
+	FText GetLabel() const;
+
 	FReply OnFunctionButtonClicked() const;
 };
diff --git a/Engine/Plugins/Experimental/Avalanche/Source/AvalancheMediaEditor/Private/Rundown/DetailsView/RemoteControl/Properties/SAvaRundownPageRemoteControlProps.cpp b/Engine/Plugins/Experimental/Avalanche/Source/AvalancheMediaEditor/Private/Rundown/DetailsView/RemoteControl/Properties/SAvaRundownPageRemoteControlProps.cpp
--- a/Engine/Plugins/Experimental/Avalanche/Source/AvalancheMediaEditor/Private/Rundown/DetailsView/RemoteControl/Properties/SAvaRundownPageRemoteControlProps.cpp
+++ b/Engine/Plugins/Experimental/Avalanche/Source/AvalancheMediaEditor/Private/Rundown/DetailsView/RemoteControl/Properties/SAvaRundownPageRemoteControlProps.cpp
@@ -105,6 +105,7 @@ void SAvaRundownPageRemoteControlProps::Construct(const FArguments& InArgs, TSha
 			.OnContextMenuOpening(this, &SAvaRundownPageRemoteControlProps::GetContextMenuContent)
 			.OnGenerateRow(this, &SAvaRundownPageRemoteControlProps::OnGenerateControllerRow)
 			.OnGetChildren(this, &SAvaRundownPageRemoteControlProps::OnGetEntityChildren)
+			.OnExpansionChanged(this, &SAvaRundownPageRemoteControlProps::OnExpansionChanged)
 			.HeaderRow(HeaderRow)
 		]
 	];
@@ -146,6 +147,36 @@ void SAvaRundownPageRemoteControlProps::OnGetEntityChildren(TSharedPtr<FAvaRundo
 	}
 }
 
+void SAvaRundownPageRemoteControlProps::OnExpansionChanged(TSharedPtr<FAvaRundownRCFieldItem> InItem, bool bInIsExpanded)
+{
+	if (InItem.IsValid())
+	{
+		const FStringView Path = InItem->GetPath();
+		if (bInIsExpanded)
+		{
+			ExpandedPaths.Emplace(Path);
+		}
+		else
+		{
+			ExpandedPaths.RemoveByHash(GetTypeHash(Path), Path);
+		}
+	}
+}
+
+void SAvaRundownPageRemoteControlProps::UpdateItemExpansionsRecursive(TConstArrayView<TSharedPtr<FAvaRundownRCFieldItem>> InItems)
+{
+	for (const TSharedPtr<FAvaRundownRCFieldItem>& Item : InItems)
+	{
+		if (Item.IsValid())
+		{
+			const FStringView Path = Item->GetPath();
+			const bool bShouldExpand = ExpandedPaths.ContainsByHash(GetTypeHash(Path), Path);
+			FieldContainer->SetItemExpansion(Item, bShouldExpand);
+			UpdateItemExpansionsRecursive(Item->GetChildren());
+		}
+	}
+}
+
 void SAvaRundownPageRemoteControlProps::RefreshTable(const TSet<FGuid>& InEntityIds)
 {
 	for (const TSharedPtr<FAvaRundownRCFieldItem>& PropertyItem : FieldItems)
@@ -364,10 +395,17 @@ void SAvaRundownPageRemoteControlProps::Tick(const FGeometry& InAllottedGeometry
 	// Ensure all the items have their hierarchies up-to-date prior to Tree Refresh
 	if (bRefreshRequested)
 	{
+		for (FAvaRundownPageRCObject& PageRCObject : PageRCObjects)
+		{
+			PageRCObject.CacheTreeNodes();
+		}
+
+		TSharedRef<SAvaRundownPageRemoteControlProps> This = SharedThis(this);
 		for (const TSharedPtr<FAvaRundownRCFieldItem>& FieldItem : FieldItems)
 		{
-			FieldItem->RefreshChildren();
+			FieldItem->Refresh(This);
 		}
+		UpdateItemExpansionsRecursive(FieldItems);
 		bRefreshRequested = false;
 	}
 }
@@ -602,7 +640,7 @@ TSharedPtr<FAvaRundownPageRCPropsNotifyHook> SAvaRundownPageRemoteControlProps::
 FAvaRundownPageRCObject& SAvaRundownPageRemoteControlProps::FindOrAddPageRCObject(UObject* InObject)
 {
 	bool bIsAlreadyInSet;
-	FAvaRundownPageRCObject& PageRCObject = PageRCObjectMap.FindOrAdd(FAvaRundownPageRCObject(InObject), &bIsAlreadyInSet);
+	FAvaRundownPageRCObject& PageRCObject = PageRCObjects.FindOrAdd(FAvaRundownPageRCObject(InObject), &bIsAlreadyInSet);
 
 	// Initialize if new element in set
 	if (!bIsAlreadyInSet)
diff --git a/Engine/Plugins/Experimental/Avalanche/Source/AvalancheMediaEditor/Private/Rundown/DetailsView/RemoteControl/Properties/SAvaRundownPageRemoteControlProps.h b/Engine/Plugins/Experimental/Avalanche/Source/AvalancheMediaEditor/Private/Rundown/DetailsView/RemoteControl/Properties/SAvaRundownPageRemoteControlProps.h
--- a/Engine/Plugins/Experimental/Avalanche/Source/AvalancheMediaEditor/Private/Rundown/DetailsView/RemoteControl/Properties/SAvaRundownPageRemoteControlProps.h
+++ b/Engine/Plugins/Experimental/Avalanche/Source/AvalancheMediaEditor/Private/Rundown/DetailsView/RemoteControl/Properties/SAvaRundownPageRemoteControlProps.h
@@ -134,6 +134,10 @@ private:
 
 	void OnGetEntityChildren(TSharedPtr<FAvaRundownRCFieldItem> InItem, TArray<TSharedPtr<FAvaRundownRCFieldItem>>& OutChildren);
 
+	void OnExpansionChanged(TSharedPtr<FAvaRundownRCFieldItem> InItem, bool bInIsExpanded);
+
+	void UpdateItemExpansionsRecursive(TConstArrayView<TSharedPtr<FAvaRundownRCFieldItem>> InItems);
+
 	void RefreshTable(const TSet<FGuid>& InEntityIds = TSet<FGuid>());
 
 	TSharedPtr<SWidget> GetContextMenuContent();
@@ -156,7 +160,9 @@ private:
 
 	TSharedPtr<FAvaRundownPageRCPropsNotifyHook> NotifyHook;
 
-	TSet<FAvaRundownPageRCObject> PageRCObjectMap;
+	TSet<FAvaRundownPageRCObject> PageRCObjects;
+
+	TSet<FString> ExpandedPaths;
 
 	friend class FAvaRundownPageRCPropsNotifyHookImpl;
 
