[Motion Design] Rundown Server: Fix combo tempalte creation.

diff --git a/Engine/Plugins/Experimental/Avalanche/Source/AvalancheMedia/Private/Rundown/AvaRundownServer.cpp b/Engine/Plugins/Experimental/Avalanche/Source/AvalancheMedia/Private/Rundown/AvaRundownServer.cpp
--- a/Engine/Plugins/Experimental/Avalanche/Source/AvalancheMedia/Private/Rundown/AvaRundownServer.cpp
+++ b/Engine/Plugins/Experimental/Avalanche/Source/AvalancheMedia/Private/Rundown/AvaRundownServer.cpp
@@ -1289,7 +1289,7 @@ void FAvaRundownServer::HandleCreateComboTemplate(const FAvaRundownCreateComboTe
 	FAvaRundownServerErrorContext ErrorContext;
 	const TArray<int32> TemplateIds = Rundown->ValidateTemplateIdsForComboTemplate(InMessage.CombinedTemplateIds, ErrorContext);
 
-	if (TemplateIds.Num() > 2)
+	if (TemplateIds.Num() > 1)
 	{
 		const FAvaRundownPageIdGeneratorParams IdGeneratorParams(InMessage.IdGeneratorParams.ReferenceId, InMessage.IdGeneratorParams.Increment);
 		const int32 TemplateId = Rundown->AddComboTemplate(TemplateIds, IdGeneratorParams);
