Motion Design: Piloted camera saved
- The piloted camera is now saved and re-piloted when re-opening a Motion Design map.
- Changing the piloted camera now marks the level dirty.

diff --git a/Engine/Plugins/Experimental/Avalanche/Source/AvalancheLevelViewport/Private/ViewportClient/AvaLevelViewportClient.cpp b/Engine/Plugins/Experimental/Avalanche/Source/AvalancheLevelViewport/Private/ViewportClient/AvaLevelViewportClient.cpp
--- a/Engine/Plugins/Experimental/Avalanche/Source/AvalancheLevelViewport/Private/ViewportClient/AvaLevelViewportClient.cpp
+++ b/Engine/Plugins/Experimental/Avalanche/Source/AvalancheLevelViewport/Private/ViewportClient/AvaLevelViewportClient.cpp
@@ -2,6 +2,7 @@
 
 #include "AvaLevelViewportClient.h"
 #include "AvaActorUtils.h"
+#include "AvaViewportDataSubsystem.h"
 #include "AvaViewportPostProcessManager.h"
 #include "AvaViewportUtils.h"
 #include "AvaVisibleArea.h"
@@ -64,6 +65,8 @@ void FAvaLevelViewportClient::Init()
 	ZoomController = MakeShared<FAvaCameraZoomController>(SharedThis(this), 90.f);
 	PostProcessManager = MakeShared<FAvaViewportPostProcessManager>(SharedThis(this));
 	IsolateActorsOperation = MakeShared<FAvaIsolateActorsOperation>(SharedThis(this));
+
+	LoadActiveCameraComponent();
 }
 
 bool FAvaLevelViewportClient::HasCachedViewportSize() const
@@ -331,6 +334,9 @@ void FAvaLevelViewportClient::Draw(const FSceneView* InView, FPrimitiveDrawInter
 {
 	FLevelEditorViewportClient::Draw(InView, InPDI);
 
+	// There is no event triggered for unpiloting a camera, so we need to check.
+	CheckActiveCameraComponent();
+
 	if (InPDI)
 	{
 		if (UWorld* ViewportClientWorld = GetWorld())
@@ -670,6 +676,39 @@ FSceneView* FAvaLevelViewportClient::CalcNonZoomedSceneView(FSceneViewFamily* Vi
 	return CalcSceneView(ViewFamily, StereoViewIndex);
 }
 
+void FAvaLevelViewportClient::LoadActiveCameraComponent()
+{
+	if (UWorld* ViewportWorld = GetWorld())
+	{
+		if (UAvaViewportDataSubsystem* ViewportDataSubsystem = ViewportWorld->GetSubsystem<UAvaViewportDataSubsystem>())
+		{
+			if (FAvaViewportData* ViewportData = ViewportDataSubsystem->GetData())
+			{
+				if (AActor* Actor = ViewportData->PilotedCamera.Get())
+				{
+					SetViewTarget(Actor);
+				}
+			}
+		}
+	}
+}
+
+void FAvaLevelViewportClient::CheckActiveCameraComponent()
+{
+	if (AActor* ActorViewTarget = GetViewTarget())
+	{
+		UpdateActiveCameraComponent(ActorViewTarget);
+	}
+	else if (AActor* CinematicViewTarget = GetCinematicViewTarget())
+	{
+		UpdateActiveCameraComponent(CinematicViewTarget);
+	}
+	else
+	{
+		UpdateActiveCameraComponent(nullptr);
+	}
+}
+
 UCameraComponent* FAvaLevelViewportClient::UpdateActiveCameraComponent(AActor* InViewTarget)
 {
 	if (!IsValid(InViewTarget))
@@ -719,6 +758,23 @@ void FAvaLevelViewportClient::SetActiveCameraComponent(UCameraComponent* InCamer
 		bLockedCameraView = false;
 		ActiveCameraComponentWeak.Reset();
 	}
+
+	if (UWorld* ViewportWorld = GetWorld())
+	{
+		if (UAvaViewportDataSubsystem* ViewportDataSubsystem = ViewportWorld->GetSubsystem<UAvaViewportDataSubsystem>())
+		{
+			if (FAvaViewportData* ViewportData = ViewportDataSubsystem->GetData())
+			{
+				AActor* CurrentActor = ViewportData->PilotedCamera.Get();
+
+				if (Actor != CurrentActor)
+				{
+					ViewportDataSubsystem->ModifyDataSource();
+					ViewportData->PilotedCamera = Actor;
+				}
+			}
+		}
+	}
 }
 
 #undef LOCTEXT_NAMESPACE
diff --git a/Engine/Plugins/Experimental/Avalanche/Source/AvalancheLevelViewport/Private/ViewportClient/AvaLevelViewportClient.h b/Engine/Plugins/Experimental/Avalanche/Source/AvalancheLevelViewport/Private/ViewportClient/AvaLevelViewportClient.h
--- a/Engine/Plugins/Experimental/Avalanche/Source/AvalancheLevelViewport/Private/ViewportClient/AvaLevelViewportClient.h
+++ b/Engine/Plugins/Experimental/Avalanche/Source/AvalancheLevelViewport/Private/ViewportClient/AvaLevelViewportClient.h
@@ -137,6 +137,10 @@ protected:
 
 	void SetCinematicViewTarget(AActor* InCinematicViewTarget);
 
+	void LoadActiveCameraComponent();
+
+	void CheckActiveCameraComponent();
+
 	UCameraComponent* UpdateActiveCameraComponent(AActor* InViewTarget);
 
 	void SetActiveCameraComponent(UCameraComponent* InCameraComponent);
diff --git a/Engine/Plugins/Experimental/Avalanche/Source/AvalancheViewport/Public/AvaViewportDataSubsystem.h b/Engine/Plugins/Experimental/Avalanche/Source/AvalancheViewport/Public/AvaViewportDataSubsystem.h
--- a/Engine/Plugins/Experimental/Avalanche/Source/AvalancheViewport/Public/AvaViewportDataSubsystem.h
+++ b/Engine/Plugins/Experimental/Avalanche/Source/AvalancheViewport/Public/AvaViewportDataSubsystem.h
@@ -8,6 +8,7 @@
 #include "AvaViewportPostProcessManager.h"
 #include "AvaViewportVirtualSizeEnums.h"
 #include "Math/IntPoint.h"
+#include "UObject/SoftObjectPtr.h"
 #include "AvaViewportDataSubsystem.generated.h"
 
 class AAvaViewportDataActor;
@@ -28,6 +29,9 @@ struct FAvaViewportData
 
 	UPROPERTY()
 	EAvaViewportVirtualSizeAspectRatioState VirtualSizeAspectRatioState = EAvaViewportVirtualSizeAspectRatioState::LockedToCamera;
+
+	UPROPERTY()
+	TSoftObjectPtr<AActor> PilotedCamera;
 };
 
 UCLASS()
