MotionDesign : Modifiers
- Adding SizeToTexture modifier to size/scale dynamic mesh based on a texture asset and size rules

diff --git a/Engine/Plugins/Experimental/Avalanche/Source/AvalancheModifiers/Private/Modifiers/AvaSizeToTextureModifier.cpp b/Engine/Plugins/Experimental/Avalanche/Source/AvalancheModifiers/Private/Modifiers/AvaSizeToTextureModifier.cpp
new file mode 100644
--- /dev/null
+++ b/Engine/Plugins/Experimental/Avalanche/Source/AvalancheModifiers/Private/Modifiers/AvaSizeToTextureModifier.cpp
@@ -0,0 +1,321 @@
+// Copyright Epic Games, Inc. All Rights Reserved.
+
+#include "Modifiers/AvaSizeToTextureModifier.h"
+
+#include "AvaShapeActor.h"
+#include "Components/DynamicMeshComponent.h"
+#include "DynamicMesh/MeshTransforms.h"
+#include "DynamicMeshes/AvaShape2DDynMeshBase.h"
+#include "Engine/Texture.h"
+
+#define LOCTEXT_NAMESPACE "AvaTextureSizeModifier"
+
+#if WITH_EDITOR
+const TAvaPropertyChangeDispatcher<UAvaSizeToTextureModifier> UAvaSizeToTextureModifier::PropertyChangeDispatcher =
+{
+	{ GET_MEMBER_NAME_CHECKED(UAvaSizeToTextureModifier, Texture), &UAvaSizeToTextureModifier::OnTextureOptionsChanged },
+	{ GET_MEMBER_NAME_CHECKED(UAvaSizeToTextureModifier, Rule), &UAvaSizeToTextureModifier::OnTextureOptionsChanged },
+	{ GET_MEMBER_NAME_CHECKED(UAvaSizeToTextureModifier, FixedHeight), &UAvaSizeToTextureModifier::OnTextureOptionsChanged },
+	{ GET_MEMBER_NAME_CHECKED(UAvaSizeToTextureModifier, FixedWidth), &UAvaSizeToTextureModifier::OnTextureOptionsChanged },
+};
+
+void UAvaSizeToTextureModifier::PostEditChangeProperty(FPropertyChangedEvent& InPropertyChangedEvent)
+{
+	Super::PostEditChangeProperty(InPropertyChangedEvent);
+
+	PropertyChangeDispatcher.OnPropertyChanged(this, InPropertyChangedEvent);
+}
+#endif
+
+void UAvaSizeToTextureModifier::SetTexture(UTexture* InTexture)
+{
+	if (InTexture == Texture)
+	{
+		return;
+	}
+
+	Texture = InTexture;
+	OnTextureOptionsChanged();
+}
+
+void UAvaSizeToTextureModifier::SetRule(EAvaSizeToTextureRule InRule)
+{
+	if (Rule == InRule)
+	{
+		return;
+	}
+
+	Rule = InRule;
+	OnTextureOptionsChanged();
+}
+
+void UAvaSizeToTextureModifier::SetFixedHeight(float InFixedHeight)
+{
+	InFixedHeight = FMath::Max(0, InFixedHeight);
+	if (FMath::IsNearlyEqual(FixedHeight, InFixedHeight))
+	{
+		return;
+	}
+
+	FixedHeight = InFixedHeight;
+	OnTextureOptionsChanged();
+}
+
+void UAvaSizeToTextureModifier::SetFixedWidth(float InFixedWidth)
+{
+	InFixedWidth = FMath::Max(0, InFixedWidth);
+	if (FMath::IsNearlyEqual(FixedWidth, InFixedWidth))
+	{
+		return;
+	}
+
+	FixedWidth = InFixedWidth;
+	OnTextureOptionsChanged();
+}
+
+void UAvaSizeToTextureModifier::OnModifierCDOSetup(FActorModifierCoreMetadata& InMetadata)
+{
+	Super::OnModifierCDOSetup(InMetadata);
+
+	InMetadata.SetName(TEXT("SizeToTexture"));
+	InMetadata.SetCategory(TEXT("Geometry"));
+#if WITH_EDITOR
+	InMetadata.SetDescription(LOCTEXT("ModifierDescription", "The modified actor will be resized to match a texture size based on the provided rule"));
+#endif
+
+	InMetadata.SetCompatibilityRule([](const AActor* InActor)->bool
+	{
+		bool bSupported = false;
+		if (InActor)
+		{
+			if (const UDynamicMeshComponent* DynMeshComponent = InActor->FindComponentByClass<UDynamicMeshComponent>())
+			{
+				DynMeshComponent->ProcessMesh([&bSupported](const FDynamicMesh3& InProcessMesh)
+				{
+					// Only flat meshes are supported => no depth => (x == 0)
+					bSupported = InProcessMesh.VertexCount() > 0 && FMath::IsNearlyZero((InProcessMesh.GetBounds(true).Width()));
+				});
+			}
+		}
+		return bSupported;
+	});
+}
+
+void UAvaSizeToTextureModifier::OnModifierAdded(EActorModifierCoreEnableReason InReason)
+{
+	Super::OnModifierAdded(InReason);
+
+	AddExtension<FActorModifierRenderStateUpdateExtension>(this);
+}
+
+void UAvaSizeToTextureModifier::OnModifierEnabled(EActorModifierCoreEnableReason InReason)
+{
+	Super::OnModifierEnabled(InReason);
+
+	if (InReason == EActorModifierCoreEnableReason::User)
+	{
+		if (const AAvaShapeActor* const ShapeActor = Cast<AAvaShapeActor>(GetModifiedActor()))
+		{
+			if (UAvaShape2DDynMeshBase* const Shape2D = Cast<UAvaShape2DDynMeshBase>(ShapeActor->GetDynamicMesh()))
+			{
+				Shape2DWeak = Shape2D;
+				PreModifierShape2DSize = Shape2D->GetSize2D();
+			}
+		}
+	}
+}
+
+void UAvaSizeToTextureModifier::OnModifierDisabled(EActorModifierCoreDisableReason InReason)
+{
+	Super::OnModifierDisabled(InReason);
+
+	if (UAvaShape2DDynMeshBase* const Shape2D = Shape2DWeak.Get())
+	{
+		Shape2D->SetSize2D(PreModifierShape2DSize);
+	}
+}
+
+void UAvaSizeToTextureModifier::Apply()
+{
+	if (!IsValid(Texture))
+	{
+		Next();
+		return;
+	}
+
+	const float Width = Texture->GetSurfaceWidth();
+	const float Height = Texture->GetSurfaceHeight();
+
+	if (FMath::IsNearlyZero(Width) || FMath::IsNearlyZero(Height))
+	{
+		Fail(LOCTEXT("InvalidTextureSize", "Invalid Texture Size"));
+		return;
+	}
+
+	if (Rule == EAvaSizeToTextureRule::FixedHeight && FMath::IsNearlyZero(FixedHeight))
+	{
+		Fail(LOCTEXT("InvalidFixedHeight", "Invalid Fixed Height"));
+		return;
+	}
+
+	if (Rule == EAvaSizeToTextureRule::FixedWidth && FMath::IsNearlyZero(FixedWidth))
+	{
+		Fail(LOCTEXT("InvalidFixedWidth", "Invalid Fixed Width"));
+		return;
+	}
+
+	UDynamicMeshComponent* const DynMeshComponent = GetMeshComponent();
+	if (!IsValid(DynMeshComponent))
+	{
+		Fail(LOCTEXT("InvalidDynamicMeshComponent", "Invalid dynamic mesh component on modified actor"));
+		return;
+	}
+	
+	FVector2D ShapeSize;
+	FVector2D ShapeScale;
+	GetShapeSizeAndScale(ShapeSize, ShapeScale);
+
+	UE_LOG(LogTemp, Warning, TEXT("Apply Modifier : Size %s vs %s Scale %s vs %s"), *ShapeSize.ToString(), *CachedSize.ToString(), *ShapeScale.ToString(), *CachedScale.ToString())
+	
+	if (ShapeSize.IsNearlyZero() || ShapeScale.IsNearlyZero())
+	{
+		Fail(LOCTEXT("InvalidShapeSizeOrScale", "Invalid Shape Size or Scale"));
+		return;
+	}
+
+	// Scaled shape size
+	ShapeSize = ShapeSize * ShapeScale;
+
+	const float WidthHeightRatio = Width / Height;
+	const float HeightWidthRatio = Height / Width;
+	
+	// Check if size is already correct and compute new size
+	FVector2D NewSize = FVector2D::ZeroVector;
+	if (Rule == EAvaSizeToTextureRule::FixedHeight)
+	{
+		if (FMath::IsNearlyEqual(ShapeSize.Y, FixedHeight) && FMath::IsNearlyEqual(ShapeSize.X, FixedHeight * WidthHeightRatio))
+		{
+			Next();
+			return;
+		}
+
+		NewSize = FVector2D(FixedHeight * WidthHeightRatio, FixedHeight);
+	}
+	else if (Rule == EAvaSizeToTextureRule::FixedWidth)
+	{
+		if (FMath::IsNearlyEqual(ShapeSize.X, FixedWidth) && FMath::IsNearlyEqual(ShapeSize.Y, FixedWidth * HeightWidthRatio))
+		{
+			Next();
+			return;
+		}
+		
+		NewSize = FVector2D(FixedWidth, FixedWidth * HeightWidthRatio);
+	}
+	else if (Rule == EAvaSizeToTextureRule::AdaptiveHeight)
+	{
+		if (FMath::IsNearlyEqual(ShapeSize.X, ShapeSize.Y * WidthHeightRatio))
+		{
+			Next();
+			return;
+		}
+
+		NewSize = FVector2D(ShapeSize.Y * WidthHeightRatio, ShapeSize.Y);
+	}
+	else if (Rule == EAvaSizeToTextureRule::AdaptiveWidth)
+	{
+		if (FMath::IsNearlyEqual(ShapeSize.Y, ShapeSize.X * HeightWidthRatio))
+		{
+			Next();
+			return;
+		}
+
+		NewSize = FVector2D(ShapeSize.X, ShapeSize.X * HeightWidthRatio);
+	}
+
+	if (NewSize.IsNearlyZero())
+	{
+		Fail(LOCTEXT("InvalidNewSize", "Invalid New Size for Shape"));
+		return;
+	}
+	
+	CachedScale = ShapeScale;
+	CachedSize = NewSize;
+
+	if (UAvaShape2DDynMeshBase* const Shape2D = Shape2DWeak.Get())
+	{
+		const FVector2D NewUnscaledShapeSize = CachedSize / ShapeScale;
+		if (!Shape2D->GetSize2D().Equals(NewUnscaledShapeSize, 0.01))
+		{
+			Shape2D->SetSize2D(NewUnscaledShapeSize);
+		}
+	}
+	else
+	{
+		const FVector NewUnscaledShapeSize = FVector(1, CachedSize.X / ShapeSize.X, CachedSize.Y / ShapeSize.Y);
+		DynMeshComponent->GetDynamicMesh()->EditMesh([this, &NewUnscaledShapeSize](FDynamicMesh3& InEditMesh)
+		{
+			MeshTransforms::Scale(InEditMesh, NewUnscaledShapeSize, FVector::Zero(), /** ReverseOrientation */true);
+		}
+		, EDynamicMeshChangeType::GeneralEdit, EDynamicMeshAttributeChangeFlags::Unknown, false);
+	}
+
+	Next();
+}
+
+void UAvaSizeToTextureModifier::OnModifiedActorTransformed()
+{
+	Super::OnModifiedActorTransformed();
+
+	CheckSizeOrScaleChanged();
+}
+
+void UAvaSizeToTextureModifier::OnRenderStateUpdated(AActor* InActor, UActorComponent* InComponent)
+{
+	if (InComponent == GetMeshComponent())
+	{
+		CheckSizeOrScaleChanged();
+	}
+}
+
+void UAvaSizeToTextureModifier::OnActorVisibilityChanged(AActor* InActor)
+{
+	// Do nothing
+}
+
+void UAvaSizeToTextureModifier::OnTextureOptionsChanged()
+{
+	MarkModifierDirty();
+}
+
+void UAvaSizeToTextureModifier::GetShapeSizeAndScale(FVector2D& OutShapeSize, FVector2D& OutShapeScale) const
+{
+	OutShapeSize = FVector2D(0, 0);
+	OutShapeScale = FVector2D(0, 0);
+
+	if (const UDynamicMeshComponent* const DynMeshComponent = GetMeshComponent())
+	{
+		DynMeshComponent->ProcessMesh([&OutShapeSize](const FDynamicMesh3& InProcessMesh)
+		{
+			const UE::Geometry::FAxisAlignedBox3d AlignedBox = InProcessMesh.GetBounds(true);
+			OutShapeSize = FVector2D(AlignedBox.Height(), AlignedBox.Depth());
+		});
+
+		OutShapeScale = FVector2D(DynMeshComponent->GetComponentScale().Y, DynMeshComponent->GetComponentScale().Z);
+	}
+}
+
+void UAvaSizeToTextureModifier::CheckSizeOrScaleChanged()
+{
+	FVector2D ShapeSize;
+	FVector2D ShapeScale;
+	GetShapeSizeAndScale(ShapeSize, ShapeScale);
+
+	ShapeSize = ShapeSize * ShapeScale;
+
+	if (!CachedSize.Equals(ShapeSize))
+	{
+		MarkModifierDirty();	
+	}
+}
+
+#undef LOCTEXT_NAMESPACE
diff --git a/Engine/Plugins/Experimental/Avalanche/Source/AvalancheModifiers/Public/Modifiers/AvaSizeToTextureModifier.h b/Engine/Plugins/Experimental/Avalanche/Source/AvalancheModifiers/Public/Modifiers/AvaSizeToTextureModifier.h
new file mode 100644
--- /dev/null
+++ b/Engine/Plugins/Experimental/Avalanche/Source/AvalancheModifiers/Public/Modifiers/AvaSizeToTextureModifier.h
@@ -0,0 +1,128 @@
+// Copyright Epic Games, Inc. All Rights Reserved.
+
+#pragma once
+
+#include "AvaGeometryBaseModifier.h"
+#include "AvaPropertyChangeDispatcher.h"
+#include "Extensions/ActorModifierRenderStateUpdateExtension.h"
+#include "AvaSizeToTextureModifier.generated.h"
+
+class UActorComponent;
+class UAvaShape2DDynMeshBase;
+class UTexture;
+
+UENUM()
+enum class EAvaSizeToTextureRule : uint8
+{
+	/** Changes to width are allowed, height will react accordingly */
+	AdaptiveWidth,
+	/** Changes to height are allowed, width will react accordingly */
+	AdaptiveHeight,
+	/** Lock height and adapt width */
+	FixedHeight,
+	/** Lock width and adapt height */
+	FixedWidth,
+};
+
+/**
+ * Adapts the modified actor geometry size/scale to match the texture size/ratio
+ */
+UCLASS(MinimalAPI, BlueprintType)
+class UAvaSizeToTextureModifier : public UAvaGeometryBaseModifier
+	, public IActorModifierRenderStateUpdateHandler
+{
+	GENERATED_BODY()
+
+public:
+	UFUNCTION(BlueprintCallable, Category="Motion Design|Modifiers|TextureSize")
+	void SetTexture(UTexture* InTexture);
+
+	UFUNCTION(BlueprintPure, Category="Motion Design|Modifiers|TextureSize")
+	UTexture* GetTexture() const
+	{
+		return Texture;
+	}
+
+	UFUNCTION(BlueprintCallable, Category="Motion Design|Modifiers|TextureSize")
+	void SetRule(EAvaSizeToTextureRule InRule);
+
+	UFUNCTION(BlueprintPure, Category="Motion Design|Modifiers|TextureSize")
+	EAvaSizeToTextureRule GetRule() const
+	{
+		return Rule;
+	}
+
+	UFUNCTION(BlueprintCallable, Category="Motion Design|Modifiers|TextureSize")
+	void SetFixedHeight(float InFixedHeight);
+
+	UFUNCTION(BlueprintPure, Category="Motion Design|Modifiers|TextureSize")
+	float GetFixedHeight() const
+	{
+		return FixedHeight;
+	}
+
+	UFUNCTION(BlueprintCallable, Category="Motion Design|Modifiers|TextureSize")
+	void SetFixedWidth(float InFixedWidth);
+
+	UFUNCTION(BlueprintPure, Category="Motion Design|Modifiers|TextureSize")
+	float GetFixedWidth() const
+	{
+		return FixedWidth;
+	}
+	
+protected:
+	//~ Begin UObject
+#if WITH_EDITOR
+	virtual void PostEditChangeProperty(FPropertyChangedEvent& InPropertyChangedEvent) override;
+#endif
+	//~ End UObject
+
+	//~ Begin UActorModifierCoreBase
+	virtual void OnModifierCDOSetup(FActorModifierCoreMetadata& InMetadata) override;
+	virtual void OnModifiedActorTransformed() override;
+	virtual void OnModifierAdded(EActorModifierCoreEnableReason InReason) override;
+	virtual void OnModifierEnabled(EActorModifierCoreEnableReason InReason) override;
+	virtual void OnModifierDisabled(EActorModifierCoreDisableReason InReason) override;
+	virtual void Apply() override;
+	//~ End UActorModifierCoreBase
+
+	//~ Begin IAvaRenderStateUpdateExtension
+	virtual void OnRenderStateUpdated(AActor* InActor, UActorComponent* InComponent) override;
+	virtual void OnActorVisibilityChanged(AActor* InActor) override;
+	//~ End IAvaRenderStateUpdateExtension
+
+	void OnTextureOptionsChanged();
+
+	void GetShapeSizeAndScale(FVector2D& OutShapeSize, FVector2D& OutShapeScale) const;
+	void CheckSizeOrScaleChanged();
+
+	/** Texture to resize to */
+	UPROPERTY(EditInstanceOnly, Setter, Getter, Category="TextureSize", meta=(AllowPrivateAccess="true"))
+	TObjectPtr<UTexture> Texture;
+
+	/** Rule for texture resize */
+	UPROPERTY(EditInstanceOnly, Setter, Getter, Category="TextureSize", meta=(AllowPrivateAccess="true"))
+	EAvaSizeToTextureRule Rule = EAvaSizeToTextureRule::FixedHeight;
+
+	/** The fixed height size */
+	UPROPERTY(EditInstanceOnly, Setter, Getter, Category="TextureSize", meta=(ClampMin="0", EditCondition="Rule == EAvaTextureSizeRule::FixedHeight", EditConditionHides, AllowPrivateAccess="true"))
+	float FixedHeight = 512.f;
+
+	/** The fixed width size */
+	UPROPERTY(EditInstanceOnly, Setter, Getter, Category="TextureSize", meta=(ClampMin="0", EditCondition="Rule == EAvaTextureSizeRule::FixedWidth", EditConditionHides, AllowPrivateAccess="true"))
+	float FixedWidth = 512.f;
+
+private:
+	UPROPERTY()
+	FVector2D PreModifierShape2DSize = FVector2D::ZeroVector;
+
+	FVector2D CachedScale = FVector2D::ZeroVector;
+	FVector2D CachedSize = FVector2D::ZeroVector;
+
+	UPROPERTY()
+	TWeakObjectPtr<UAvaShape2DDynMeshBase> Shape2DWeak;
+
+#if WITH_EDITOR
+	static const TAvaPropertyChangeDispatcher<UAvaSizeToTextureModifier> PropertyChangeDispatcher;
+#endif
+};
\ No newline at end of file
