MotionDesign : Modifier
- Global Opacity modifier improvements to support Shapes/Text/SVG materials instead of only Material Designer materials, modifier will show an error if no translucent/masked material is available
- Shape parametric material mode can be set to translucent/opaque or auto
- Updating translucent materials to include opacity parameters

#virtualized

diff --git a/Engine/Plugins/Experimental/Avalanche/Source/AvalancheMask/Private/Handling/AvaMaskParametricMaterialHandle.cpp b/Engine/Plugins/Experimental/Avalanche/Source/AvalancheMask/Private/Handling/AvaMaskParametricMaterialHandle.cpp
--- a/Engine/Plugins/Experimental/Avalanche/Source/AvalancheMask/Private/Handling/AvaMaskParametricMaterialHandle.cpp
+++ b/Engine/Plugins/Experimental/Avalanche/Source/AvalancheMask/Private/Handling/AvaMaskParametricMaterialHandle.cpp
@@ -99,7 +99,7 @@ const EBlendMode FAvaMaskParametricMaterialHandle::GetBlendMode()
 {
 	if (const FAvaShapeParametricMaterial* ParametricMtl = GetParametricMaterial())
 	{
-		return ParametricMtl->GetUseTranslucentMaterial() ? EBlendMode::BLEND_Translucent : EBlendMode::BLEND_Opaque;
+		return ParametricMtl->ShouldUseTranslucentMaterial() ? EBlendMode::BLEND_Translucent : EBlendMode::BLEND_Opaque;
 	}
 	
 	if (const UMaterialInterface* MaterialInstance = GetMaterialInstance())
@@ -120,8 +120,8 @@ void FAvaMaskParametricMaterialHandle::SetBlendMode(const EBlendMode InBlendMode
 	const EBlendMode TargetBlendMode = UE::AvaMask::Internal::GetTargetBlendMode( GetBlendMode(), InBlendMode);
 	if (FAvaShapeParametricMaterial* ParametricMtl = GetParametricMaterial())
 	{
-		ParametricMtl->SetUseAutoTranslucency(false); // Disable so that user changes won't override the mask
-		ParametricMtl->SetUseTranslucentMaterial(TargetBlendMode != EBlendMode::BLEND_Opaque);
+		// Disable auto-switch so that user changes won't override the mask
+		ParametricMtl->SetTranslucency(TargetBlendMode != EBlendMode::BLEND_Opaque ? EAvaShapeParametricMaterialTranslucency::Enabled : EAvaShapeParametricMaterialTranslucency::Disabled);
 		return;
 	}
 	
@@ -196,8 +196,7 @@ bool FAvaMaskParametricMaterialHandle::ApplyOriginalState(
 				, MaterialHandleData->OriginalMaskMaterialParameters.InnerFeatherRadius);
 
 			ParametricMtl->SetStyle(MaterialHandleData->OriginalMaterialStyle);
-			ParametricMtl->SetUseAutoTranslucency(true);
-			ParametricMtl->SetUseTranslucentMaterial(MaterialHandleData->OriginalBlendMode == EBlendMode::BLEND_Translucent);
+			ParametricMtl->SetTranslucency(MaterialHandleData->OriginalBlendMode == EBlendMode::BLEND_Translucent ? EAvaShapeParametricMaterialTranslucency::Enabled : EAvaShapeParametricMaterialTranslucency::Auto);
 
 			InMaterialSetter(ParametricMtl->GetMaterial());
 
diff --git a/Engine/Plugins/Experimental/Avalanche/Source/AvalancheModifiers/Private/AvaModifiersParametricMaterial.cpp b/Engine/Plugins/Experimental/Avalanche/Source/AvalancheModifiers/Private/AvaModifiersParametricMaterial.cpp
--- a/Engine/Plugins/Experimental/Avalanche/Source/AvalancheModifiers/Private/AvaModifiersParametricMaterial.cpp
+++ b/Engine/Plugins/Experimental/Avalanche/Source/AvalancheModifiers/Private/AvaModifiersParametricMaterial.cpp
@@ -39,7 +39,8 @@ void FAvaModifiersParametricMaterial::ApplyParams() const
 		return;
 	}
 
-	InstanceMaterial->SetVectorParameterValue("MaskColor", MaskColor);
+	InstanceMaterial->SetVectorParameterValue(TEXT("MaskColor"), MaskColor);
+	InstanceMaterial->SetScalarParameterValue(TEXT("GlobalOpacity"), MaskColor.A);
 }
 
 void FAvaModifiersParametricMaterial::CreateAndApply(UObject* Outer)
diff --git a/Engine/Plugins/Experimental/Avalanche/Source/AvalancheModifiers/Private/Modifiers/AvaGlobalOpacityModifier.cpp b/Engine/Plugins/Experimental/Avalanche/Source/AvalancheModifiers/Private/Modifiers/AvaGlobalOpacityModifier.cpp
--- a/Engine/Plugins/Experimental/Avalanche/Source/AvalancheModifiers/Private/Modifiers/AvaGlobalOpacityModifier.cpp
+++ b/Engine/Plugins/Experimental/Avalanche/Source/AvalancheModifiers/Private/Modifiers/AvaGlobalOpacityModifier.cpp
@@ -5,6 +5,7 @@
 #include "Components/DMMaterialValue.h"
 #include "Components/MaterialValues/DMMaterialValueFloat1.h"
 #include "Material/DynamicMaterialInstance.h"
+#include "Materials/Material.h"
 #include "Model/DynamicMaterialModel.h"
 
 #define LOCTEXT_NAMESPACE "AvaGlobalOpacityModifier"
@@ -14,7 +15,6 @@ UAvaGlobalOpacityModifier::UAvaGlobalOpacityModifier()
 #if WITH_EDITOR
 	bShowMaterialParameters = false;
 #endif
-	MaterialClass = UDynamicMaterialInstance::StaticClass();
 	MaterialParameters.ScalarParameters.Add(UDynamicMaterialModel::GlobalOpacityParameterName, GlobalOpacity);
 }
 
@@ -58,8 +58,19 @@ void UAvaGlobalOpacityModifier::OnActorMaterialAdded(UMaterialInstanceDynamic* I
 {
 	Super::OnActorMaterialAdded(InAdded);
 
+	if (IsValid(InAdded))
+	{
+		if (const UMaterial* Material = InAdded->GetBaseMaterial())
+		{
+			if (Material->BlendMode == BLEND_Masked || Material->BlendMode == BLEND_Translucent)
+			{
+				++SupportedOpacityMaterialCount;
+			}
+		}
+	}
+
 #if WITH_EDITOR
-	if (UDynamicMaterialInstance* MDI = Cast<UDynamicMaterialInstance>(InAdded))
+	if (const UDynamicMaterialInstance* MDI = Cast<UDynamicMaterialInstance>(InAdded))
 	{
 		if (const UDynamicMaterialModel* Model = MDI->GetMaterialModel())
 		{
@@ -76,9 +87,20 @@ void UAvaGlobalOpacityModifier::OnActorMaterialAdded(UMaterialInstanceDynamic* I
 void UAvaGlobalOpacityModifier::OnActorMaterialRemoved(UMaterialInstanceDynamic* InRemoved)
 {
 	Super::OnActorMaterialRemoved(InRemoved);
+	
+	if (IsValid(InRemoved))
+	{
+		if (const UMaterial* Material = InRemoved->GetBaseMaterial())
+		{
+			if (Material->BlendMode == BLEND_Masked || Material->BlendMode == BLEND_Translucent)
+			{
+				--SupportedOpacityMaterialCount;
+			}
+		}
+	}
 
 #if WITH_EDITOR
-	if (UDynamicMaterialInstance* MDI = Cast<UDynamicMaterialInstance>(InRemoved))
+	if (const UDynamicMaterialInstance* MDI = Cast<UDynamicMaterialInstance>(InRemoved))
 	{
 		if (const UDynamicMaterialModel* Model = MDI->GetMaterialModel())
 		{
@@ -118,4 +140,15 @@ void UAvaGlobalOpacityModifier::OnModifierCDOSetup(FActorModifierCoreMetadata& I
 #endif
 }
 
+void UAvaGlobalOpacityModifier::Apply()
+{
+	if (SupportedOpacityMaterialCount <= 0)
+	{
+		Fail(LOCTEXT("NoOpacityMaterialFound", "No Supported Opacity Material Found"));
+		return;
+	}
+
+	Super::Apply();
+}
+
 #undef LOCTEXT_NAMESPACE
diff --git a/Engine/Plugins/Experimental/Avalanche/Source/AvalancheModifiers/Private/Modifiers/AvaMaterialParameterModifier.cpp b/Engine/Plugins/Experimental/Avalanche/Source/AvalancheModifiers/Private/Modifiers/AvaMaterialParameterModifier.cpp
--- a/Engine/Plugins/Experimental/Avalanche/Source/AvalancheModifiers/Private/Modifiers/AvaMaterialParameterModifier.cpp
+++ b/Engine/Plugins/Experimental/Avalanche/Source/AvalancheModifiers/Private/Modifiers/AvaMaterialParameterModifier.cpp
@@ -188,19 +188,19 @@ void UAvaMaterialParameterModifier::RestorePreState()
 	RestoreMaterialParameters();
 }
 
-void UAvaMaterialParameterModifier::Apply()
+void UAvaMaterialParameterModifier::SavePreState()
 {
-	const AActor* ActorModified = GetModifiedActor();
-	if (!ActorModified)
-	{
-		return;
-	}
+	Super::SavePreState();
 
 	// Scan for materials in actors
-	ScanActorMaterials();
+    ScanActorMaterials();
 
-	SaveMaterialParameters();
+	// Save original parameters
+    SaveMaterialParameters();
+}
 
+void UAvaMaterialParameterModifier::Apply()
+{
 	// Set new parameter value to them
 	for (TPair<TWeakObjectPtr<UMaterialInstanceDynamic>, FAvaMaterialParameterMap>& MaterialParameterPair : SavedMaterialParameters)
 	{
diff --git a/Engine/Plugins/Experimental/Avalanche/Source/AvalancheModifiers/Public/Modifiers/AvaGlobalOpacityModifier.h b/Engine/Plugins/Experimental/Avalanche/Source/AvalancheModifiers/Public/Modifiers/AvaGlobalOpacityModifier.h
--- a/Engine/Plugins/Experimental/Avalanche/Source/AvalancheModifiers/Public/Modifiers/AvaGlobalOpacityModifier.h
+++ b/Engine/Plugins/Experimental/Avalanche/Source/AvalancheModifiers/Public/Modifiers/AvaGlobalOpacityModifier.h
@@ -34,6 +34,7 @@ protected:
 	//~ End UObject
 
 	//~ Begin UActorModifierCoreBase
+	virtual void Apply() override;
 	virtual void OnModifierCDOSetup(FActorModifierCoreMetadata& InMetadata) override;
 	//~ End UActorModifierCoreBase
 
@@ -46,4 +47,7 @@ protected:
 	/** Global opacity to set on all Material Designer Instances */
 	UPROPERTY(EditInstanceOnly, Setter="SetGlobalOpacity", Getter="GetGlobalOpacity", Category="Material Parameter", meta=(ClampMin="0", ClampMax="1", AllowPrivateAccess="true"))
 	float GlobalOpacity = 1.f;
+
+	UPROPERTY()
+	int32 SupportedOpacityMaterialCount = 0;
 };
diff --git a/Engine/Plugins/Experimental/Avalanche/Source/AvalancheModifiers/Public/Modifiers/AvaMaterialParameterModifier.h b/Engine/Plugins/Experimental/Avalanche/Source/AvalancheModifiers/Public/Modifiers/AvaMaterialParameterModifier.h
--- a/Engine/Plugins/Experimental/Avalanche/Source/AvalancheModifiers/Public/Modifiers/AvaMaterialParameterModifier.h
+++ b/Engine/Plugins/Experimental/Avalanche/Source/AvalancheModifiers/Public/Modifiers/AvaMaterialParameterModifier.h
@@ -73,6 +73,7 @@ protected:
 	virtual void OnModifierDisabled(EActorModifierCoreDisableReason InReason) override;
 	virtual void OnModifiedActorTransformed() override;
 	virtual void RestorePreState() override;
+	virtual void SavePreState() override;
 	virtual void Apply() override;
 	//~ End UActorModifierCoreBase
 
diff --git a/Engine/Plugins/Experimental/Avalanche/Source/AvalancheShapes/Private/AvaShapeParametricMaterial.cpp b/Engine/Plugins/Experimental/Avalanche/Source/AvalancheShapes/Private/AvaShapeParametricMaterial.cpp
--- a/Engine/Plugins/Experimental/Avalanche/Source/AvalancheShapes/Private/AvaShapeParametricMaterial.cpp
+++ b/Engine/Plugins/Experimental/Avalanche/Source/AvalancheShapes/Private/AvaShapeParametricMaterial.cpp
@@ -8,27 +8,10 @@
 
 FAvaShapeParametricMaterial::FOnMaterialChanged FAvaShapeParametricMaterial::OnMaterialChangedDelegate;
 
-FAvaShapeParametricMaterial::FAvaShapeParametricMaterial()
-{
-	bUseTranslucentMaterial = false;
-	bUseUnlitMaterial       = false;
-	bUseTwoSidedMaterial    = true;
-
-	Style    = EAvaShapeParametricMaterialStyle::Solid;
-
-	ColorA   = FLinearColor::White;
-	ColorA.A = 1.f;
-
-	ColorB   = FLinearColor::Black;
-	ColorB.A = 1.f;
-
-	GradientOffset = 0.5f;
-}
-
 FAvaShapeParametricMaterial::FAvaShapeParametricMaterial(const FAvaShapeParametricMaterial& Other)
 {
 	bUseUnlitMaterial        = Other.bUseUnlitMaterial;
-	bUseTranslucentMaterial  = Other.bUseTranslucentMaterial;
+	Translucency             = Other.Translucency;
 	bUseTwoSidedMaterial     = Other.bUseTwoSidedMaterial;
 	ColorA                   = Other.ColorA;
 	ColorB                   = Other.ColorB;
@@ -43,7 +26,7 @@ FAvaShapeParametricMaterial::FAvaShapeParametricMaterial(const FAvaShapeParametr
 FAvaShapeParametricMaterial& FAvaShapeParametricMaterial::operator=(const FAvaShapeParametricMaterial& Other)
 {
 	bUseUnlitMaterial        = Other.bUseUnlitMaterial;
-	bUseTranslucentMaterial  = Other.bUseTranslucentMaterial;
+	Translucency             = Other.Translucency;
 	bUseTwoSidedMaterial     = Other.bUseTwoSidedMaterial;
 	ColorA                   = Other.ColorA;
 	ColorB                   = Other.ColorB;
@@ -162,6 +145,16 @@ void FAvaShapeParametricMaterial::SetMaterialParameterValues(UMaterialInstanceDy
 	}
 }
 
+bool FAvaShapeParametricMaterial::ShouldUseTranslucentMaterial() const
+{
+	if (Translucency == EAvaShapeParametricMaterialTranslucency::Auto)
+	{
+		return ColorA.A < 1.f || ColorB.A < 1.f;
+	}
+
+	return Translucency == EAvaShapeParametricMaterialTranslucency::Enabled;
+}
+
 void FAvaShapeParametricMaterial::OnMaterialParameterUpdated()
 {
 	const int32 PreviousActiveIndex = ActiveInstanceIndex;
@@ -185,7 +178,7 @@ int32 FAvaShapeParametricMaterial::GetActiveInstanceIndex() const
 {
 	FAvaShapeParametricMaterial* MutableThis = const_cast<FAvaShapeParametricMaterial*>(this);
 
-	MutableThis->ActiveInstanceIndex = (bUseTranslucentMaterial ? Translucent : Opaque)
+	MutableThis->ActiveInstanceIndex = (ShouldUseTranslucentMaterial() ? Translucent : Opaque)
 		+ (bUseUnlitMaterial ? Unlit : Lit)
 		+ (bUseTwoSidedMaterial ? TwoSided : OneSided);
 
@@ -332,12 +325,12 @@ UMaterialInstanceDynamic* FAvaShapeParametricMaterial::GetOrCreateMaterial(UObje
 	return Material;
 }
 
-FAvaShapeParametricMaterial::FOnMaterialParameterChanged& FAvaShapeParametricMaterial::OnMaterialParameterChanged()
+FAvaShapeParametricMaterial::FOnMaterialParameterChanged::RegistrationType& FAvaShapeParametricMaterial::OnMaterialParameterChanged()
 {
 	return OnMaterialParameterChangedDelegate;
 }
 
-FAvaShapeParametricMaterial::FOnMaterialChanged& FAvaShapeParametricMaterial::OnMaterialChanged()
+FAvaShapeParametricMaterial::FOnMaterialChanged::RegistrationType& FAvaShapeParametricMaterial::OnMaterialChanged()
 {
 	return OnMaterialChangedDelegate;
 }
@@ -360,25 +353,14 @@ void FAvaShapeParametricMaterial::SetMaterial(UMaterialInstanceDynamic* InMateri
 	OnMaterialParameterUpdated();
 }
 
-void FAvaShapeParametricMaterial::SetUseTranslucentMaterial(bool bInUse)
+void FAvaShapeParametricMaterial::SetTranslucency(EAvaShapeParametricMaterialTranslucency InTranslucency)
 {
-	if (bInUse == bUseTranslucentMaterial)
+	if (InTranslucency == Translucency)
 	{
 		return;
 	}
 
-	bUseTranslucentMaterial = bInUse;
-	OnMaterialParameterUpdated();
-}
-
-void FAvaShapeParametricMaterial::SetUseAutoTranslucency(bool bInUse)
-{
-	if (bInUse == bUseAutoTranslucency)
-	{
-		return;
-	}
-
-	bUseAutoTranslucency = bInUse;
+	Translucency = InTranslucency;
 	OnMaterialParameterUpdated();
 }
 
@@ -401,12 +383,6 @@ void FAvaShapeParametricMaterial::SetPrimaryColor(const FLinearColor& InColor)
 	}
 
 	ColorA = InColor;
-
-	if (bUseAutoTranslucency)
-	{
-		bUseTranslucentMaterial = ColorA.A < 1.f || ColorB.A < 1.f;	
-	}
-	
 	OnMaterialParameterUpdated();
 }
 
@@ -418,12 +394,6 @@ void FAvaShapeParametricMaterial::SetSecondaryColor(const FLinearColor& InColor)
 	}
 
 	ColorB = InColor;
-
-	if (bUseAutoTranslucency)
-	{
-		bUseTranslucentMaterial = ColorA.A < 1.f || ColorB.A < 1.f;
-	}
-	
 	OnMaterialParameterUpdated();
 }
 
diff --git a/Engine/Plugins/Experimental/Avalanche/Source/AvalancheShapes/Private/DynamicMeshes/AvaShapeDynMeshBase.cpp b/Engine/Plugins/Experimental/Avalanche/Source/AvalancheShapes/Private/DynamicMeshes/AvaShapeDynMeshBase.cpp
--- a/Engine/Plugins/Experimental/Avalanche/Source/AvalancheShapes/Private/DynamicMeshes/AvaShapeDynMeshBase.cpp
+++ b/Engine/Plugins/Experimental/Avalanche/Source/AvalancheShapes/Private/DynamicMeshes/AvaShapeDynMeshBase.cpp
@@ -1218,11 +1218,7 @@ void UAvaShapeDynamicMeshBase::OnParametricMaterialChanged(int32 MaterialIndex)
 	}
 
 	FAvaShapeParametricMaterial* ParametricMaterial = GetParametricMaterial(MaterialIndex);
-	if (ParametricMaterial->GetUseAutoTranslucency())
-	{
-		ParametricMaterial->SetUseTranslucentMaterial(ParametricMaterial->GetPrimaryColor().A < 1.f || ParametricMaterial->GetSecondaryColor().A < 1.f);	
-	}
-	
+
 	UMaterialInterface* ParametricMaterialInstance = ParametricMaterial->GetOrCreateMaterial(GetShapeActor());
 
 	UMaterialInterface* CurrentMaterial = GetMaterial(MaterialIndex);
diff --git a/Engine/Plugins/Experimental/Avalanche/Source/AvalancheShapes/Public/AvaShapeParametricMaterial.h b/Engine/Plugins/Experimental/Avalanche/Source/AvalancheShapes/Public/AvaShapeParametricMaterial.h
--- a/Engine/Plugins/Experimental/Avalanche/Source/AvalancheShapes/Public/AvaShapeParametricMaterial.h
+++ b/Engine/Plugins/Experimental/Avalanche/Source/AvalancheShapes/Public/AvaShapeParametricMaterial.h
@@ -15,6 +15,17 @@ class UMaterialInterface;
 class UObject;
 struct FAvaShapeParametricMaterial;
 
+UENUM(BlueprintType)
+enum class EAvaShapeParametricMaterialTranslucency : uint8
+{
+	/** Switches when parameters have opacity < 1 */
+	Auto,
+	/** Uses opaque material regardless of opacity parameters */
+	Disabled,
+	/** Uses translucent material regardless of opacity parameters */
+	Enabled
+};
+
 USTRUCT(BlueprintType)
 struct FAvaShapeParametricMaterial
 {
@@ -36,17 +47,16 @@ struct FAvaShapeParametricMaterial
 
 	DECLARE_MULTICAST_DELEGATE_OneParam(FOnMaterialChanged, FAvaShapeParametricMaterial&)
 	DECLARE_MULTICAST_DELEGATE(FOnMaterialParameterChanged)
+	
+	/** Called when active material index changed */
+	static FOnMaterialChanged::RegistrationType& OnMaterialChanged();
 
-	AVALANCHESHAPES_API FAvaShapeParametricMaterial();
-
+	AVALANCHESHAPES_API FAvaShapeParametricMaterial() = default;
 	AVALANCHESHAPES_API FAvaShapeParametricMaterial(const FAvaShapeParametricMaterial& Other);
 	AVALANCHESHAPES_API FAvaShapeParametricMaterial& operator=(const FAvaShapeParametricMaterial& Other);
 
 	/** Called when active material parameters are updated */
-	AVALANCHESHAPES_API FOnMaterialParameterChanged& OnMaterialParameterChanged();
-
-	/** Called when active material index changed */
-	static FOnMaterialChanged& OnMaterialChanged();
+	AVALANCHESHAPES_API FOnMaterialParameterChanged::RegistrationType& OnMaterialParameterChanged();
 
 	AVALANCHESHAPES_API bool CopyFromMaterialParameters(UMaterialInstance* InMaterial);
 
@@ -65,19 +75,14 @@ struct FAvaShapeParametricMaterial
 	/** Set instanced material currently active */
 	void SetMaterial(UMaterialInstanceDynamic* InMaterial);
 
-	bool GetUseTranslucentMaterial() const
+	AVALANCHESHAPES_API void SetTranslucency(EAvaShapeParametricMaterialTranslucency InTranslucency);
+	
+	EAvaShapeParametricMaterialTranslucency GetTranslucency() const
 	{
-		return bUseTranslucentMaterial;
+		return Translucency;
 	}
-
-	AVALANCHESHAPES_API void SetUseTranslucentMaterial(bool bInUse);
-
-	bool GetUseAutoTranslucency() const
-	{
-		return bUseAutoTranslucency;
-	}
-
-	AVALANCHESHAPES_API void SetUseAutoTranslucency(bool bInUse);
+	
+	AVALANCHESHAPES_API bool ShouldUseTranslucentMaterial() const;
 
 	EAvaShapeParametricMaterialStyle GetStyle() const
 	{
@@ -125,60 +130,57 @@ struct FAvaShapeParametricMaterial
 	AVALANCHESHAPES_API void SetMaterialParameterValues(UMaterialInstanceDynamic* InMaterialInstance, bool bInNotifyUpdate = false) const;
 
 protected:
-	/** Parent material
-	 * 1. Opaque Lit
-	 * 2. Translucent Lit
-	 * 3. Opaque Unlit
-	 * 4. Translucent Unlit
-	 * ...
-	 */
-	UPROPERTY(Transient)
-	TArray<TObjectPtr<UMaterialInterface>> DefaultMaterials;
-
-	/** Material instances corresponding to the parent material instanced
-	 * 1. Opaque Lit
-	 * 2. Translucent Lit
-	 * 3. Opaque Unlit
-	 * 4. Translucent Unlit
-	 * ...
-	 */
-	UPROPERTY(Transient)
-	TArray<TObjectPtr<UMaterialInstanceDynamic>> InstanceMaterials;
-
-	// whether we use opaque or translucent material, set automatically when opacity is below 1
-	UPROPERTY()
-	bool bUseTranslucentMaterial;
-
-	UPROPERTY(EditInstanceOnly, BlueprintReadWrite, AdvancedDisplay, Category="Material")
-	bool bUseAutoTranslucency = true;
-
 	// Default style for the material
 	UPROPERTY(EditInstanceOnly, BlueprintReadWrite, Category="Material", meta=(DisplayName="Simple Style", AllowPrivateAccess = "true"))
-	EAvaShapeParametricMaterialStyle Style;
+	EAvaShapeParametricMaterialStyle Style = EAvaShapeParametricMaterialStyle::Solid;
 
 	// Primary colour for the material
 	UPROPERTY(EditInstanceOnly, BlueprintReadWrite, Category="Material", meta=(DisplayName="Simple Color", AllowPrivateAccess = "true"))
-	FLinearColor ColorA;
+	FLinearColor ColorA = FLinearColor::White;
 
 	// Secondary colour for the material
 	UPROPERTY(EditInstanceOnly, BlueprintReadWrite, Category="Material", meta=(DisplayName="Gradient Color", EditCondition="Style == EAvaShapeParametricMaterialStyle::LinearGradient", EditConditionHides, AllowPrivateAccess = "true"))
-	FLinearColor ColorB;
+	FLinearColor ColorB = FLinearColor::Black;
 
 	// Offset for gradient style material
 	UPROPERTY(EditInstanceOnly, BlueprintReadWrite, Category="Material", meta=(ClampMin="0", ClampMax="1", DisplayName="Gradient Offset", EditCondition = "Style == EAvaShapeParametricMaterialStyle::LinearGradient", EditConditionHides, AllowPrivateAccess = "true"))
-	float GradientOffset;
+	float GradientOffset = 0.5f;
 
 	// whether the material is unlit or default lit
 	UPROPERTY(EditInstanceOnly, BlueprintReadWrite, Category="Material", meta=(DisplayName="Use Unlit Material", AllowPrivateAccess = "true"))
-	bool bUseUnlitMaterial;
+	bool bUseUnlitMaterial = false;
 
 	// whether the material is one sided or two sided
 	UPROPERTY(EditInstanceOnly, BlueprintReadWrite, Category = "Material", meta = (DisplayName = "Use Two Sided Material", AllowPrivateAccess = "true"))
-	bool bUseTwoSidedMaterial;
+	bool bUseTwoSidedMaterial = true;
+
+	/** How to handle translucency for the underlying material */
+	UPROPERTY(EditInstanceOnly, BlueprintReadWrite, Category="Material")
+	EAvaShapeParametricMaterialTranslucency Translucency = EAvaShapeParametricMaterialTranslucency::Auto;
 
 private:
 	static FOnMaterialChanged OnMaterialChangedDelegate;
 
+	/** Parent material
+	 * 1. Opaque Lit
+	 * 2. Translucent Lit
+	 * 3. Opaque Unlit
+	 * 4. Translucent Unlit
+	 * ...
+	 */
+	UPROPERTY(Transient)
+	TArray<TObjectPtr<UMaterialInterface>> DefaultMaterials;
+
+	/** Material instances corresponding to the parent material instanced
+	 * 1. Opaque Lit
+	 * 2. Translucent Lit
+	 * 3. Opaque Unlit
+	 * 4. Translucent Unlit
+	 * ...
+	 */
+	UPROPERTY(Transient)
+	TArray<TObjectPtr<UMaterialInstanceDynamic>> InstanceMaterials;
+
 	FOnMaterialParameterChanged OnMaterialParameterChangedDelegate;
 
 	int32 ActiveInstanceIndex = INDEX_NONE;
diff --git a/Engine/Plugins/Experimental/Avalanche/Source/AvalancheText/Public/AvaTextDefs.h b/Engine/Plugins/Experimental/Avalanche/Source/AvalancheText/Public/AvaTextDefs.h
--- a/Engine/Plugins/Experimental/Avalanche/Source/AvalancheText/Public/AvaTextDefs.h
+++ b/Engine/Plugins/Experimental/Avalanche/Source/AvalancheText/Public/AvaTextDefs.h
@@ -28,7 +28,7 @@ UENUM()
 enum class EAvaTextTranslucency : uint8 
 {
 	Invalid        UMETA(Hidden),
-	None,
+	None           UMETA(DisplayName="Opaque"),
 	Translucent,
 	GradientMask,
 };
diff --git a/Engine/Plugins/Experimental/Avalanche/Source/AvalancheText/Public/MaterialHub/AvaTextMaterialHub.h b/Engine/Plugins/Experimental/Avalanche/Source/AvalancheText/Public/MaterialHub/AvaTextMaterialHub.h
--- a/Engine/Plugins/Experimental/Avalanche/Source/AvalancheText/Public/MaterialHub/AvaTextMaterialHub.h
+++ b/Engine/Plugins/Experimental/Avalanche/Source/AvalancheText/Public/MaterialHub/AvaTextMaterialHub.h
@@ -155,7 +155,7 @@ public:
 
 		/** make sure the following names match the ones of their respective material parameter */
 		inline static FName SolidColor_MatParam         = TEXT("Color");
-		inline static FName Opacity_MatParam            = TEXT("Opacity");
+		inline static FName Opacity_MatParam            = TEXT("GlobalOpacity");
 
 		inline static FName MainTexture_MatParam        = TEXT("MainTexture");
 		inline static FName TexturedUTiling_MatParam    = TEXT("U_Tiling");
