MotionDesign : Shapes
- Adding multi-select shape section edits in details panel (must be same generator class)
- Fixing vector customization multi-select edit only affecting first selected object instead of all of them

diff --git a/Engine/Plugins/Experimental/Avalanche/Source/AvalancheEditor/Private/DetailView/Customizations/AvaMeshesDetailCustomization.cpp b/Engine/Plugins/Experimental/Avalanche/Source/AvalancheEditor/Private/DetailView/Customizations/AvaMeshesDetailCustomization.cpp
--- a/Engine/Plugins/Experimental/Avalanche/Source/AvalancheEditor/Private/DetailView/Customizations/AvaMeshesDetailCustomization.cpp
+++ b/Engine/Plugins/Experimental/Avalanche/Source/AvalancheEditor/Private/DetailView/Customizations/AvaMeshesDetailCustomization.cpp
@@ -43,22 +43,37 @@ void FAvaMeshesDetailCustomization::CustomizeDetails(IDetailLayoutBuilder& InDet
 	IDetailCategoryBuilder& MaterialCategoryBuilder = InDetailBuilder.EditCategory(FName("Material"));
 	MaterialCategoryBuilder.SetSortOrder(ShapeCategoryBuilder.GetSortOrder() + 1);
 
-	if (MeshGeneratorsWeak.Num() == 1 && MeshGeneratorsWeak[0].IsValid())
+	// Make sure we have the mesh generator class in multi selection
+	bool bSameMeshGenerator = true;
+	UClass* MeshGeneratorClass = nullptr;
+	for (TArray<TWeakObjectPtr<UAvaShapeDynamicMeshBase>>::TIterator It(MeshGeneratorsWeak); It; ++It)
 	{
-		UAvaShapeDynamicMeshBase* DynMesh = MeshGeneratorsWeak[0].Get();
-
-		TSharedPtr<IPropertyHandleMap> MapHandle = MeshDatasHandle->AsMap();
-
-		uint32 Count = 0;
-		MapHandle->GetNumElements(Count);
+		if (UAvaShapeDynamicMeshBase* MeshGenerator = It->Get())
+		{
+			if (MeshGeneratorClass != nullptr && MeshGeneratorClass != MeshGenerator->GetClass())
+			{
+				bSameMeshGenerator = false;
+				break;
+			}
+			
+			MeshGeneratorClass = MeshGenerator->GetClass();
+		}
+		else
+		{
+			It.RemoveCurrent();
+		}
+	}
 
+	if (!MeshGeneratorsWeak.IsEmpty() && bSameMeshGenerator)
+	{
+		UAvaShapeDynamicMeshBase* DynMesh = MeshGeneratorsWeak[0].Get();
 		const TArray<FName> MeshNames = DynMesh->GetMeshDataNames();
 
-		for (uint32 Index = 0; Index < Count; ++Index)
+		for (int32 Index = 0; Index < MeshNames.Num(); ++Index)
 		{
 			FString MeshName = TEXT("");
 
-			if (Index > 0 || !DynMesh->GetUsePrimaryMaterialEverywhere())
+			if (Index > 0)
 			{
 				MeshName = MeshNames[Index].ToString() + TEXT(" ");
 			}
@@ -153,7 +168,7 @@ void FAvaMeshesDetailCustomization::CustomizeDetails(IDetailLayoutBuilder& InDet
 				return MaterialUVHandle->IsEditable() ? EVisibility::Visible : EVisibility::Hidden;
 			}));
 
-			if (Index < Count - 1)
+			if (Index < MeshNames.Num())
 			{
 				// Separator row
 				FDetailWidgetRow& SeparatorRow = MaterialCategoryBuilder.AddCustomRow(FText::GetEmpty());
diff --git a/Engine/Plugins/Experimental/Avalanche/Source/AvalancheEditor/Private/DetailView/Customizations/AvaVectorPropertyTypeCustomization.cpp b/Engine/Plugins/Experimental/Avalanche/Source/AvalancheEditor/Private/DetailView/Customizations/AvaVectorPropertyTypeCustomization.cpp
--- a/Engine/Plugins/Experimental/Avalanche/Source/AvalancheEditor/Private/DetailView/Customizations/AvaVectorPropertyTypeCustomization.cpp
+++ b/Engine/Plugins/Experimental/Avalanche/Source/AvalancheEditor/Private/DetailView/Customizations/AvaVectorPropertyTypeCustomization.cpp
@@ -102,9 +102,13 @@ void FAvaVectorPropertyTypeCustomization::CustomizeHeader(TSharedRef<IPropertyHa
 	if (bIsVector3d)
 	{
 		bPixelSizeProperty = VectorPropertyHandle->GetProperty()->GetFName() == GET_MEMBER_NAME_CHECKED(UAvaShape3DDynMeshBase, PixelSize3D);
-		XPropertyHandle = StructPropertyHandle->GetChildHandle(GET_MEMBER_NAME_CHECKED(FVector, X));
-		YPropertyHandle = StructPropertyHandle->GetChildHandle(GET_MEMBER_NAME_CHECKED(FVector, Y));
-		ZPropertyHandle = StructPropertyHandle->GetChildHandle(GET_MEMBER_NAME_CHECKED(FVector, Z));
+
+		VectorComponentPropertyHandles =
+		{
+			StructPropertyHandle->GetChildHandle(GET_MEMBER_NAME_CHECKED(FVector, X)),
+			StructPropertyHandle->GetChildHandle(GET_MEMBER_NAME_CHECKED(FVector, Y)),
+			StructPropertyHandle->GetChildHandle(GET_MEMBER_NAME_CHECKED(FVector, Z)),
+		};
 
 		if (StructPropertyHandle->HasMetaData("ClampMin"))
 		{
@@ -168,8 +172,12 @@ void FAvaVectorPropertyTypeCustomization::CustomizeHeader(TSharedRef<IPropertyHa
 	else
 	{
 		bPixelSizeProperty = VectorPropertyHandle->GetProperty()->GetFName() == GET_MEMBER_NAME_CHECKED(UAvaShape2DDynMeshBase, PixelSize2D);
-		XPropertyHandle = StructPropertyHandle->GetChildHandle(GET_MEMBER_NAME_CHECKED(FVector2D, X));
-		YPropertyHandle = StructPropertyHandle->GetChildHandle(GET_MEMBER_NAME_CHECKED(FVector2D, Y));
+
+		VectorComponentPropertyHandles =
+		{
+			StructPropertyHandle->GetChildHandle(GET_MEMBER_NAME_CHECKED(FVector2D, X)),
+			StructPropertyHandle->GetChildHandle(GET_MEMBER_NAME_CHECKED(FVector2D, Y)),
+		};
 
 		if (StructPropertyHandle->HasMetaData("ClampMin"))
 		{
@@ -358,29 +366,17 @@ TOptional<double> FAvaVectorPropertyTypeCustomization::GetVectorComponent(const
 	{
 		return TOptional<double>();
 	}
+
+	if (!VectorComponentPropertyHandles.IsValidIndex(InComponent) || !VectorComponentPropertyHandles[InComponent].IsValid())
+	{
+		return TOptional<double>();
+	}
+
 	double OutValue = 0.f;
-	switch(InComponent)
+
+	if (VectorComponentPropertyHandles[InComponent]->GetValue(OutValue) != FPropertyAccess::Success)
 	{
-		case 0:
-			if (XPropertyHandle->GetValue(OutValue) != FPropertyAccess::Success)
-			{
-				return TOptional<double>();
-			}
-		break;
-		case 1:
-			if (YPropertyHandle->GetValue(OutValue) != FPropertyAccess::Success)
-			{
-				return TOptional<double>();
-			}
-		break;
-		case 2:
-			if (ZPropertyHandle->GetValue(OutValue) != FPropertyAccess::Success)
-			{
-				return TOptional<double>();
-			}
-		break;
-		default:
-			return TOptional<double>();
+		return TOptional<double>();
 	}
 
 	// handle specific case
@@ -638,14 +634,20 @@ void FAvaVectorPropertyTypeCustomization::SetComponentValue(const double InNewVa
 		default:
 		break;
 	}
+
 	const TArray<bool> PreserveRatios
 	{
 		(RatioMode & ERatioMode::X) != ERatioMode::None && bPreserveRatio,
 		(RatioMode & ERatioMode::Y) != ERatioMode::None && bPreserveRatio,
 		(RatioMode & ERatioMode::Z) != ERatioMode::None && bPreserveRatio
 	};
+
 	// set property per object since we need to handle ratios for each objects
 	const uint8 MaxComponentCount = bIsVector3d ? 3 : 2;
+
+	TArray<TArray<FString>> ComponentValues;
+	ComponentValues.SetNum(MaxComponentCount);
+
 	for (int32 ObjIdx = 0; ObjIdx < SelectedObjectNum; ObjIdx++)
 	{
 		if (!bIsVector3d && !Begin2DValues.IsValidIndex(ObjIdx) && !Begin2DValues[ObjIdx].IsSet())
@@ -656,8 +658,10 @@ void FAvaVectorPropertyTypeCustomization::SetComponentValue(const double InNewVa
 		{
 			continue;
 		}
+
 		// compute clamped ratio for value change
 		const double ClampedRatio = GetClampedRatioValueChange(ObjIdx, InNewValue, InComponent, PreserveRatios);
+
 		// loop for each component (X,Y,Z)
 		for (uint8 ComponentIdx = 0; ComponentIdx < MaxComponentCount; ComponentIdx++)
 		{
@@ -666,32 +670,42 @@ void FAvaVectorPropertyTypeCustomization::SetComponentValue(const double InNewVa
 			{
 				continue;
 			}
+
 			// compute new component value
 			const double NewComponentValue = GetClampedComponentValue(ObjIdx, InNewValue, ClampedRatio, ComponentIdx, InComponent);
+
 			switch(ComponentIdx)
 			{
 				case 0:
 					if (PreserveRatios[0] || ComponentIdx == InComponent)
 					{
-						XPropertyHandle->SetPerObjectValue(ObjIdx, FString::SanitizeFloat(NewComponentValue), InFlags);
+						ComponentValues[ComponentIdx].Emplace(FString::SanitizeFloat(NewComponentValue));
 					}
 					break;
 				case 1:
 					if (PreserveRatios[1] || ComponentIdx == InComponent)
 					{
-						YPropertyHandle->SetPerObjectValue(ObjIdx, FString::SanitizeFloat(NewComponentValue), InFlags);
+						ComponentValues[ComponentIdx].Emplace(FString::SanitizeFloat(NewComponentValue));
 					}
 					break;
 				case 2:
 					if (PreserveRatios[2] || ComponentIdx == InComponent)
 					{
-						ZPropertyHandle->SetPerObjectValue(ObjIdx, FString::SanitizeFloat(NewComponentValue), InFlags);
+						ComponentValues[ComponentIdx].Emplace(FString::SanitizeFloat(NewComponentValue));
 					}
 					break;
 				default:;
 			}
 		}
 	}
+
+	for (int32 Index = 0; Index < VectorComponentPropertyHandles.Num(); Index++)
+	{
+		if (!ComponentValues[Index].IsEmpty() && VectorComponentPropertyHandles[Index].IsValid())
+		{
+			VectorComponentPropertyHandles[Index]->SetPerObjectValues(ComponentValues[Index], InFlags);
+		}
+	}
 }
 
 double FAvaVectorPropertyTypeCustomization::GetClampedRatioValueChange(const int32 InObjectIdx, const double InNewValue, const uint8 InComponent, const TArray<bool>& InPreserveRatios) const
diff --git a/Engine/Plugins/Experimental/Avalanche/Source/AvalancheEditor/Private/DetailView/Customizations/AvaVectorPropertyTypeCustomization.h b/Engine/Plugins/Experimental/Avalanche/Source/AvalancheEditor/Private/DetailView/Customizations/AvaVectorPropertyTypeCustomization.h
--- a/Engine/Plugins/Experimental/Avalanche/Source/AvalancheEditor/Private/DetailView/Customizations/AvaVectorPropertyTypeCustomization.h
+++ b/Engine/Plugins/Experimental/Avalanche/Source/AvalancheEditor/Private/DetailView/Customizations/AvaVectorPropertyTypeCustomization.h
@@ -113,9 +113,7 @@ private:
 	TWeakPtr<IAvaViewportClient> ViewportClient;
 
 	TSharedPtr<IPropertyHandle> VectorPropertyHandle;
-	TSharedPtr<IPropertyHandle> XPropertyHandle;
-	TSharedPtr<IPropertyHandle> YPropertyHandle;
-	TSharedPtr<IPropertyHandle> ZPropertyHandle;
+	TArray<TSharedPtr<IPropertyHandle>> VectorComponentPropertyHandles;
 
 	/** optional begin values to compute ratios change */
 	TArray<TOptional<FVector>> Begin3DValues;
diff --git a/Engine/Plugins/Experimental/Avalanche/Source/AvalancheShapes/Private/DynamicMeshes/AvaShapeDynMeshBase.cpp b/Engine/Plugins/Experimental/Avalanche/Source/AvalancheShapes/Private/DynamicMeshes/AvaShapeDynMeshBase.cpp
--- a/Engine/Plugins/Experimental/Avalanche/Source/AvalancheShapes/Private/DynamicMeshes/AvaShapeDynMeshBase.cpp
+++ b/Engine/Plugins/Experimental/Avalanche/Source/AvalancheShapes/Private/DynamicMeshes/AvaShapeDynMeshBase.cpp
@@ -1637,7 +1637,14 @@ void UAvaShapeDynamicMeshBase::PostEditChangeChainProperty(FPropertyChangedChain
 	if (PropertyChangedEvent.PropertyChain.GetActiveMemberNode()->GetValue()->GetFName() == MeshDatasName)
 	{
 		const int32 MeshDatasIndex = PropertyChangedEvent.GetArrayIndex(MeshDatasName.ToString());
-		const int32 MeshIndex = GetMeshesIndexes().Array()[MeshDatasIndex];
+
+		const TArray<int32> MeshesIndexes = GetMeshesIndexes().Array();
+		if (!MeshesIndexes.IsValidIndex(MeshDatasIndex))
+		{
+			return;
+		}
+
+		const int32 MeshIndex = MeshesIndexes[MeshDatasIndex];
 
 		static const FName MaterialTypeName = GET_MEMBER_NAME_CHECKED(FAvaShapeMeshData, MaterialType);
 		static const FName MaterialName = GET_MEMBER_NAME_CHECKED(FAvaShapeMeshData, Material);
