[MotionDesign] Stagger tool:
- Process random stagger on current element instead of previous
- Move "Align to Playhead" button and change to use Sequencer function
- Tweak range and sensitivity of frame spin boxes
- Disallow operation with layer bars and keyframes in locked sections

diff --git a/Engine/Plugins/Experimental/Avalanche/Source/AvalancheSequencer/Private/StaggerTool/AvaStaggerTool.cpp b/Engine/Plugins/Experimental/Avalanche/Source/AvalancheSequencer/Private/StaggerTool/AvaStaggerTool.cpp
--- a/Engine/Plugins/Experimental/Avalanche/Source/AvalancheSequencer/Private/StaggerTool/AvaStaggerTool.cpp
+++ b/Engine/Plugins/Experimental/Avalanche/Source/AvalancheSequencer/Private/StaggerTool/AvaStaggerTool.cpp
@@ -1,13 +1,11 @@
 // Copyright Epic Games, Inc. All Rights Reserved.
 
 #include "AvaStaggerTool.h"
-#include "AvaSequencer.h"
 #include "Channels/MovieSceneChannel.h"
 #include "Commands/AvaSequencerCommands.h"
 #include "Framework/Application/SlateApplication.h"
 #include "Framework/Commands/UICommandList.h"
 #include "IKeyArea.h"
-#include "ISequencer.h"
 #include "Kismet/KismetMathLibrary.h"
 #include "Misc/FrameTime.h"
 #include "MVVM/Selection/Selection.h"
@@ -15,6 +13,7 @@
 #include "MVVM/ViewModels/LayerBarModel.h"
 #include "MVVM/ViewModels/SequencerEditorViewModel.h"
 #include "ScopedTransaction.h"
+#include "Tools/SequencerSelectionAlignmentUtils.h"
 #include "Widgets/SAvaStaggerTool.h"
 
 using namespace UE::Sequencer;
@@ -251,8 +250,7 @@ FFrameNumber FAvaStaggerTool::GetInterval(const FFrameNumber InRangeSize, const
 	default:
 	case EAvaSequencerStaggerDistribution::Increment:
 		{
-			const FFrameNumber Interval = ConvertTime(Settings->ToolOptions.Interval);
-			return Settings->ToolOptions.bUseCurve && InElementSize.Value != 0 ? InElementSize : Interval;
+			return Settings->ToolOptions.bUseCurve && InElementSize.Value != 0 ? InElementSize : CachedInterval;
 		}
 
 	case EAvaSequencerStaggerDistribution::Range:
@@ -310,10 +308,12 @@ void FAvaStaggerTool::CacheOriginalElements()
 	{
 		OriginalKeyFrameElements = GatherSelectionKeyFrameElements();
 	}
+
+	CachedLayerBarCount = OriginalLayerBarElements.Num();
+	CachedKeyframeCount = OriginalKeyFrameElements.Num();
 }
 
-FFrameNumber FAvaStaggerTool::FindFirstStaggerPoint(const TArrayView<FStaggerLayerBarElement>& InStaggerElements
-	, const TRange<FFrameNumber>& InRange) const
+FFrameNumber FAvaStaggerTool::FindFirstLayerBarStaggerPoint() const
 {
 	FFrameNumber OutFrameNumber;
 
@@ -327,12 +327,12 @@ FFrameNumber FAvaStaggerTool::FindFirstStaggerPoint(const TArrayView<FStaggerLay
 		{
 		case EAvaSequencerStaggerRange::Playback:
 			OutFrameNumber = Owner.GetSequencer()->GetPlaybackRange().GetLowerBoundValue();
-			OutFrameNumber += ConvertTime(Settings->ToolOptions.Shift);
+			OutFrameNumber += CachedShiftFrames;
 			return MoveTemp(OutFrameNumber);
 
 		case EAvaSequencerStaggerRange::Selection:
 			OutFrameNumber = Owner.GetSequencer()->GetSelectionRange().GetLowerBoundValue();
-			OutFrameNumber += ConvertTime(Settings->ToolOptions.Shift);
+			OutFrameNumber += CachedShiftFrames;
 			return MoveTemp(OutFrameNumber);
 
 		case EAvaSequencerStaggerRange::Custom:
@@ -341,76 +341,23 @@ FFrameNumber FAvaStaggerTool::FindFirstStaggerPoint(const TArrayView<FStaggerLay
 		break;
 
 	case EAvaSequencerStaggerDistribution::Random:
-		const int32 OriginalElementCount = OriginalLayerBarElements.Num();
-		const TRange<FFrameNumber> Range = GetOperationRange(Settings->ToolOptions.Range);
-		const FFrameNumber RangeSize = Range.Size<FFrameNumber>();
 		const FFrameNumber ElementSize = OriginalLayerBarElements[0].LayerBarModel->ComputeRange().Size<FFrameNumber>();
-		return Range.GetLowerBoundValue() + GetInterval(RangeSize, OriginalElementCount, ElementSize);
+		const FFrameNumber Interval = GetInterval(CachedRangeSize, CachedLayerBarCount, ElementSize);
+		return CachedRange.GetLowerBoundValue() + Interval + CachedShiftFrames;
 	}
 
-	switch (Settings->ToolOptions.StartPosition)
-	{
-	case EAvaSequencerStaggerStartPosition::FirstSelected:
-		if (!InStaggerElements.IsEmpty())
-		{
-			OutFrameNumber = InStaggerElements[0].OriginalFrame;
-		}
-		break;
-
-	case EAvaSequencerStaggerStartPosition::FirstInTimeline:
-		FindFirstFrameInTimeline(InStaggerElements, OutFrameNumber);
-		break;
-
-	case EAvaSequencerStaggerStartPosition::Playhead:
-		OutFrameNumber = Owner.GetSequencer()->GetGlobalTime().Time.FrameNumber;
-		break;
-
-	case EAvaSequencerStaggerStartPosition::PlaybackRange:
-		OutFrameNumber = Owner.GetSequencer()->GetPlaybackRange().GetLowerBoundValue();
-		break;
-
-	case EAvaSequencerStaggerStartPosition::SelectionRange:
-		OutFrameNumber = Owner.GetSequencer()->GetSelectionRange().GetLowerBoundValue();
-		break;
-	}
-
-	OutFrameNumber += ConvertTime(Settings->ToolOptions.Shift);
-
-	return MoveTemp(OutFrameNumber);
+	return GetStartPosition(OriginalLayerBarElements);
 }
 
-FFrameNumber FAvaStaggerTool::FindFirstStaggerPoint(const TArrayView<FStaggerKeyFrameElement>& InStaggerElements
-	, const TRange<FFrameNumber>& InRange) const
+FFrameNumber FAvaStaggerTool::FindFirstKeyframeStaggerPoint() const
 {
-	FFrameNumber OutFrameNumber;
-
-	switch (Settings->ToolOptions.StartPosition)
+	if (Settings->ToolOptions.Distribution == EAvaSequencerStaggerDistribution::Random)
 	{
-	case EAvaSequencerStaggerStartPosition::FirstSelected:
-		if (!InStaggerElements.IsEmpty())
-		{
-			OutFrameNumber = InStaggerElements[0].OriginalFrame;
-		}
-		break;
-
-	case EAvaSequencerStaggerStartPosition::FirstInTimeline:
-		FindFirstFrameInTimeline(InStaggerElements, OutFrameNumber);
-		break;
-
-	case EAvaSequencerStaggerStartPosition::Playhead:
-		OutFrameNumber = Owner.GetSequencer()->GetGlobalTime().Time.FrameNumber;
-		break;
-
-	case EAvaSequencerStaggerStartPosition::PlaybackRange:
-		OutFrameNumber = Owner.GetSequencer()->GetPlaybackRange().GetLowerBoundValue();
-		break;
-
-	case EAvaSequencerStaggerStartPosition::SelectionRange:
-		OutFrameNumber = Owner.GetSequencer()->GetSelectionRange().GetLowerBoundValue();
-		break;
+		const FFrameNumber Interval = GetInterval(CachedRangeSize, CachedKeyframeCount);
+		return CachedRange.GetLowerBoundValue() + Interval + CachedShiftFrames;
 	}
 
-	return OutFrameNumber;
+	return GetStartPosition(OriginalLayerBarElements);
 }
 
 TArray<TViewModelPtr<FLayerBarModel>> FAvaStaggerTool::GatherAllSelectedLayerBars() const
@@ -442,19 +389,13 @@ TArray<TViewModelPtr<FLayerBarModel>> FAvaStaggerTool::GatherAllSelectedLayerBar
 			{
 				if (const TViewModelPtr<FLayerBarModel> LayerBarModel = TrackAreaModel.ImplicitCast())
 				{
-					if (!LayerBarModels.Contains(LayerBarModel))
-					{
-						LayerBarModels.Add(LayerBarModel);
-					}
+					LayerBarModels.AddUnique(LayerBarModel);
 				}
 			}
 
 			for (const TViewModelPtr<FLayerBarModel>& LayerBarModel : Track->GetTrackAreaModelListAs<FLayerBarModel>())
 			{
-				if (!LayerBarModels.Contains(LayerBarModel))
-				{
-					LayerBarModels.Add(LayerBarModel);
-				}
+				LayerBarModels.AddUnique(LayerBarModel);
 			}
 		}
 	}
@@ -466,7 +407,7 @@ TArray<TViewModelPtr<FLayerBarModel>> FAvaStaggerTool::GatherAllSelectedLayerBar
 		{
 			if (const TViewModelPtr<FLayerBarModel> LayerBarModel = ViewModel.ImplicitCast())
 			{
-				LayerBarModels.Add(LayerBarModel);
+				LayerBarModels.AddUnique(LayerBarModel);
 			}
 		}
 	}
@@ -510,14 +451,14 @@ TArray<FAvaStaggerTool::FStaggerLayerBarElement> FAvaStaggerTool::GatherSelectio
 
 TArray<FAvaStaggerTool::FStaggerKeyFrameElement> FAvaStaggerTool::GatherSelectionKeyFrameElements() const
 {
-	TArray<FStaggerKeyFrameElement> OutElements;
-
 	const TSharedPtr<FSequencerSelection> SequencerSelection = GetSequencerSelection();
 	if (!SequencerSelection.IsValid())
 	{
-		return OutElements;
+		return {};
 	}
 
+	TArray<FStaggerKeyFrameElement> OutElements;
+
 	// Gather all the selected keyframe handles
 	OutElements.Reserve(SequencerSelection->KeySelection.Num());
 
@@ -527,10 +468,14 @@ TArray<FAvaStaggerTool::FStaggerKeyFrameElement> FAvaStaggerTool::GatherSelectio
 		{
 			if (const TViewModelPtr<FChannelModel> KeyModel = SequencerSelection->KeySelection.GetModelForKey(KeyHandle))
 			{
-				FFrameNumber KeyFrameNumber;
-				KeyModel->GetChannel()->GetKeyTime(KeyHandle, KeyFrameNumber);
+				const UMovieSceneSection* const KeySection = KeyModel->GetSection();
+				if (IsValid(KeySection) && !KeySection->IsReadOnly())
+				{
+					FFrameNumber KeyFrameNumber;
+					KeyModel->GetChannel()->GetKeyTime(KeyHandle, KeyFrameNumber);
 
-				OutElements.Add(FStaggerKeyFrameElement(KeyModel, KeyHandle, KeyFrameNumber));
+					OutElements.Add(FStaggerKeyFrameElement(KeyModel, KeyHandle, KeyFrameNumber));
+				}
 			}
 		}
 	}
@@ -592,10 +537,7 @@ FFrameNumber FAvaStaggerTool::FindNextStaggerLocation(const FStaggerLayerBarElem
 {
 	const FRichCurve* const RichCurve = Settings->ToolOptions.Curve.GetRichCurve();
 	const FFrameNumber ElementSize = InElement.LayerBarModel->ComputeRange().Size<FFrameNumber>();
-	const int32 OriginalElementCount = OriginalLayerBarElements.Num();
-	const TRange<FFrameNumber> Range = GetOperationRange(Settings->ToolOptions.Range);
-	const FFrameNumber RangeSize = Range.Size<FFrameNumber>();
-	const FFrameNumber Interval = GetInterval(RangeSize, OriginalElementCount, ElementSize);
+	const FFrameNumber Interval = GetInterval(CachedRangeSize, CachedLayerBarCount, ElementSize);
 
 	if (Settings->ToolOptions.bUseCurve)
 	{
@@ -614,16 +556,23 @@ FFrameNumber FAvaStaggerTool::FindNextStaggerLocation(const FStaggerLayerBarElem
 			return InFirstFrame + LocalCurveOffset + Interval;
 		}
 
-		const float CurveStepTime = 1.f / OriginalElementCount;
+		const float CurveStepTime = 1.f / CachedLayerBarCount;
 		const float CurveTime = CurveStepTime * InElementIndex;
 
-		const FFrameNumber LocalCurveOffset = CalculateLocalCurveOffset(CurveTime, Range);
+		const FFrameNumber LocalCurveOffset = CalculateLocalCurveOffset(CurveTime, CachedRange);
 		return InFirstFrame + LocalCurveOffset + Interval;
 	}
 
 	if (Settings->ToolOptions.Distribution == EAvaSequencerStaggerDistribution::Random)
 	{
-		return Range.GetLowerBoundValue() + Interval;
+		if (OriginalLayerBarElements.IsValidIndex(InElementIndex + 1))
+		{
+			const FStaggerLayerBarElement& NextElement = OriginalLayerBarElements[InElementIndex + 1];
+			const FFrameNumber NextElementSize = NextElement.LayerBarModel->ComputeRange().Size<FFrameNumber>();
+			const FFrameNumber RandomInterval = GetInterval(CachedRangeSize, CachedLayerBarCount, NextElementSize);
+			return CachedRange.GetLowerBoundValue() + RandomInterval + CachedShiftFrames;
+		}
+		return CachedRange.GetLowerBoundValue() + CachedShiftFrames;
 	}
 
 	const FFrameNumber OperationOffset = GetLayerBarOperationOffset(InElement.LayerBarModel);
@@ -636,10 +585,7 @@ FFrameNumber FAvaStaggerTool::FindNextStaggerLocation(const FStaggerKeyFrameElem
 	, const FFrameNumber InCurrentFrame) const
 {
 	const FRichCurve* const RichCurve = Settings->ToolOptions.Curve.GetRichCurve();
-	const int32 OriginalElementCount = OriginalLayerBarElements.Num();
-	const TRange<FFrameNumber> Range = GetOperationRange(Settings->ToolOptions.Range);
-	const FFrameNumber RangeSize = Range.Size<FFrameNumber>();
-	const FFrameNumber Interval = GetInterval(RangeSize, OriginalElementCount, FFrameNumber());
+	const FFrameNumber Interval = GetInterval(CachedRangeSize, CachedKeyframeCount);
 
 	if (Settings->ToolOptions.bUseCurve)
 	{
@@ -648,15 +594,15 @@ FFrameNumber FAvaStaggerTool::FindNextStaggerLocation(const FStaggerKeyFrameElem
 			return InFirstFrame;
 		}
 
-		const float CurveStepTime = 1.f / OriginalElementCount;
-		const FFrameNumber NextStaggerFrameNumber = CalculateLocalCurveOffset(CurveStepTime * InElementIndex, Range);
+		const float CurveStepTime = 1.f / CachedKeyframeCount;
+		const FFrameNumber NextStaggerFrameNumber = CalculateLocalCurveOffset(CurveStepTime * InElementIndex, CachedRange);
 
 		return InFirstFrame + NextStaggerFrameNumber + Interval;
 	}
 
 	if (Settings->ToolOptions.Distribution == EAvaSequencerStaggerDistribution::Random)
 	{
-		return Range.GetLowerBoundValue() + Interval;
+		return CachedRange.GetLowerBoundValue() + CachedShiftFrames;
 	}
 
 	return InCurrentFrame + Interval;
@@ -664,6 +610,11 @@ FFrameNumber FAvaStaggerTool::FindNextStaggerLocation(const FStaggerKeyFrameElem
 
 void FAvaStaggerTool::Stagger()
 {
+	CachedRange = GetOperationRange(Settings->ToolOptions.Range);
+	CachedRangeSize = CachedRange.Size<FFrameNumber>();
+	CachedInterval = ConvertTime(Settings->ToolOptions.Interval);
+	CachedShiftFrames = ConvertTime(Settings->ToolOptions.Shift);
+
 	if (Settings->ToolOptions.Distribution == EAvaSequencerStaggerDistribution::Random)
 	{
 		RandomStream = FRandomStream(Settings->ToolOptions.RandomSeed);
@@ -678,7 +629,7 @@ void FAvaStaggerTool::Stagger()
 			Element.LastOffset = 0;
 		}
 
-		StaggerLayers(OriginalLayerBarElements);
+		StaggerLayers();
 	}
 	else if (IsKeyFrameSelection())
 	{
@@ -688,7 +639,7 @@ void FAvaStaggerTool::Stagger()
 			SetKeyFrameElementTime(Element, Element.OriginalFrame);
 		}
 
-		StaggerKeyFrames(OriginalKeyFrameElements);
+		StaggerKeyFrames();
 	}
 }
 
@@ -697,24 +648,22 @@ bool FAvaStaggerTool::IsAutoAppying() const
 	return Settings->bAutoApply;
 }
 
-void FAvaStaggerTool::StaggerLayers(const TArrayView<FStaggerLayerBarElement>& InOriginalElements)
+void FAvaStaggerTool::StaggerLayers()
 {
-	const int32 OriginalElementCount = InOriginalElements.Num();
-	if (OriginalElementCount <= 1)
+	if (CachedLayerBarCount <= 1)
 	{
 		return;
 	}
 
 	FScopedTransaction Transaction(LOCTEXT("StaggerLayers", "Stagger Layers"));
 
-	const TRange<FFrameNumber> Range = GetOperationRange(Settings->ToolOptions.Range);
-	const FFrameNumber FirstStaggerPoint = FindFirstStaggerPoint(InOriginalElements, Range);
+	const FFrameNumber FirstStaggerPoint = FindFirstLayerBarStaggerPoint();
 
 	FFrameNumber NextStaggerPoint = FirstStaggerPoint;
 	int32 CurrentElementNum = 0;
 	int32 CurrentGroupNum = 0;
 
-	for (FStaggerLayerBarElement& Element : InOriginalElements)
+	for (FStaggerLayerBarElement& Element : OriginalLayerBarElements)
 	{
 		Element.LastOffset = NextStaggerPoint - Element.OriginalFrame;
 		Element.LayerBarModel->Offset(Element.LastOffset);
@@ -733,10 +682,9 @@ void FAvaStaggerTool::StaggerLayers(const TArrayView<FStaggerLayerBarElement>& I
 	}
 }
 
-void FAvaStaggerTool::StaggerKeyFrames(const TArrayView<FStaggerKeyFrameElement>& InOriginalElements)
+void FAvaStaggerTool::StaggerKeyFrames()
 {
-	const int32 OriginalElementCount = InOriginalElements.Num();
-	if (OriginalElementCount <= 1)
+	if (CachedKeyframeCount <= 1)
 	{
 		return;
 	}
@@ -744,14 +692,13 @@ void FAvaStaggerTool::StaggerKeyFrames(const TArrayView<FStaggerKeyFrameElement>
 	// Stagger all elements
 	FScopedTransaction Transaction(LOCTEXT("StaggerKeyFrames", "Stagger Key Frames"));
 
-	const TRange<FFrameNumber> Range = GetOperationRange(Settings->ToolOptions.Range);
-	const FFrameNumber FirstStaggerPoint = FindFirstStaggerPoint(InOriginalElements, Range);
+	const FFrameNumber FirstStaggerPoint = FindFirstKeyframeStaggerPoint();
 
 	FFrameNumber NextStaggerPoint = FirstStaggerPoint;
 	int32 CurrentElementNum = 0;
 	int32 CurrentGroupNum = 0;
 
-	for (FStaggerKeyFrameElement& Element : InOriginalElements)
+	for (FStaggerKeyFrameElement& Element : OriginalKeyFrameElements)
 	{
 		SetKeyFrameElementTime(Element, NextStaggerPoint);
 
@@ -768,42 +715,14 @@ void FAvaStaggerTool::StaggerKeyFrames(const TArrayView<FStaggerKeyFrameElement>
 	}
 }
 
-void FAvaStaggerTool::AlignToPlayhead()
+bool FAvaStaggerTool::CanAlignToPlayhead() const
 {
-	const TSharedRef<ISequencer> Sequencer = Owner.GetSequencer();
-	const FFrameNumber PlayheadFrameNumber = Sequencer->GetGlobalTime().Time.FrameNumber;
-
-	if (IsLayerBarSelection())
-	{
-		const TArray<TViewModelPtr<FLayerBarModel>> LayerBarModels = GatherAllSelectedLayerBars();
-
-		for (const TViewModelPtr<FLayerBarModel>& LayerBarModel : LayerBarModels)
-		{
-			const TRange<FFrameNumber> Range = LayerBarModel->ComputeRange();
-			const FFrameNumber Offset = PlayheadFrameNumber - Range.GetLowerBoundValue();
-			LayerBarModel->Offset(Offset);
-		}
-	}
-	else if (IsKeyFrameSelection())
-	{
-		const TSharedPtr<FSequencerSelection> SequencerSelection = GetSequencerSelection();
-		if (!SequencerSelection.IsValid())
-		{
-			return;
-		}
+	return FSequencerSelectionAlignmentUtils::CanAlignSelection(*Owner.GetSequencer());
+}
 
-		for (const FKeyHandle KeyHandle : SequencerSelection->KeySelection)
-		{
-			if (KeyHandle != FKeyHandle::Invalid())
-			{
-				if (const TViewModelPtr<FChannelModel> KeyModel = SequencerSelection->KeySelection.GetModelForKey(KeyHandle))
-				{
-					const FStaggerKeyFrameElement Element(KeyModel, KeyHandle);
-					SetKeyFrameElementTime(Element, PlayheadFrameNumber);
-				}
-			}
-		}
-	}
+void FAvaStaggerTool::AlignToPlayhead()
+{
+	FSequencerSelectionAlignmentUtils::AlignSelectionToPlayhead(*Owner.GetSequencer());
 }
 
 void FAvaStaggerTool::SetKeyFrameElementTime(const FStaggerKeyFrameElement& InElement, const FFrameNumber InKeyTime)
@@ -814,7 +733,7 @@ void FAvaStaggerTool::SetKeyFrameElementTime(const FStaggerKeyFrameElement& InEl
 	}
 
 	UMovieSceneSection* const Section = InElement.KeyChannelModel->GetSection();
-	if (!IsValid(Section))
+	if (!IsValid(Section) || !Section->IsReadOnly())
 	{
 		return;
 	}
diff --git a/Engine/Plugins/Experimental/Avalanche/Source/AvalancheSequencer/Private/StaggerTool/AvaStaggerTool.h b/Engine/Plugins/Experimental/Avalanche/Source/AvalancheSequencer/Private/StaggerTool/AvaStaggerTool.h
--- a/Engine/Plugins/Experimental/Avalanche/Source/AvalancheSequencer/Private/StaggerTool/AvaStaggerTool.h
+++ b/Engine/Plugins/Experimental/Avalanche/Source/AvalancheSequencer/Private/StaggerTool/AvaStaggerTool.h
@@ -2,10 +2,12 @@
 
 #pragma once
 
+#include "AvaSequencer.h"
 #include "AvaStaggerToolSettings.h"
 #include "Commands/AvaSequencerAction.h"
 #include "Containers/ArrayView.h"
 #include "Containers/ContainersFwd.h"
+#include "ISequencer.h"
 #include "MVVM/ViewModelPtr.h"
 #include "Math/Range.h"
 #include "Misc/FrameNumber.h"
@@ -42,6 +44,7 @@ public:
 
 	void OnSequencerSelectionChanged();
 
+	bool CanAlignToPlayhead() const;
 	void AlignToPlayhead();
 
 	void CloseToolWindow();
@@ -82,17 +85,55 @@ private:
 	void OnApply();
 
 	FFrameNumber ConvertTime(FFrameTime InSourceTime) const;
+
+	template<typename InElementType>
+	FFrameNumber GetStartPosition(const TArray<InElementType>& InElements) const
+	{
+		FFrameNumber OutFrameNumber;
+
+		switch (Settings->ToolOptions.StartPosition)
+		{
+		case EAvaSequencerStaggerStartPosition::FirstSelected:
+			if (!InElements.IsEmpty())
+			{
+				OutFrameNumber = InElements[0].OriginalFrame;
+			}
+			break;
+
+		case EAvaSequencerStaggerStartPosition::FirstInTimeline:
+			FindFirstFrameInTimeline(InElements, OutFrameNumber);
+			break;
+
+		case EAvaSequencerStaggerStartPosition::Playhead:
+			OutFrameNumber = Owner.GetSequencer()->GetGlobalTime().Time.FrameNumber;
+			break;
+
+		case EAvaSequencerStaggerStartPosition::PlaybackRange:
+			OutFrameNumber = Owner.GetSequencer()->GetPlaybackRange().GetLowerBoundValue();
+			break;
+
+		case EAvaSequencerStaggerStartPosition::SelectionRange:
+			OutFrameNumber = Owner.GetSequencer()->GetSelectionRange().GetLowerBoundValue();
+			break;
+		}
+
+		OutFrameNumber += ConvertTime(Settings->ToolOptions.Shift);
+
+		return OutFrameNumber;
+	}
+
 	TRange<FFrameNumber> GetOperationRange(const EAvaSequencerStaggerRange InRange) const;
-	FFrameNumber GetInterval(const FFrameNumber InRangeSize, const int32 InElementCount, const FFrameNumber InElementSize) const;
+
+	FFrameNumber GetInterval(const FFrameNumber InRangeSize, const int32 InElementCount, const FFrameNumber InElementSize = 0) const;
+
 	FFrameNumber GetLayerBarOperationOffset(const UE::Sequencer::TViewModelPtr<UE::Sequencer::FLayerBarModel>& InLayerBarModel) const;
+
 	static FFrameNumber CalculateLayerBarFrameSpan(const TArray<FStaggerLayerBarElement>& InStaggerElements, const int32 InStopIndex = INDEX_NONE);
 
 	void CacheOriginalElements();
 
-	FFrameNumber FindFirstStaggerPoint(const TArrayView<FStaggerLayerBarElement>& InStaggerElements
-		, const TRange<FFrameNumber>& InRange) const;
-	FFrameNumber FindFirstStaggerPoint(const TArrayView<FStaggerKeyFrameElement>& InStaggerElements
-		, const TRange<FFrameNumber>& InRange) const;
+	FFrameNumber FindFirstLayerBarStaggerPoint() const;
+	FFrameNumber FindFirstKeyframeStaggerPoint() const;
 
 	TArray<UE::Sequencer::TViewModelPtr<UE::Sequencer::FLayerBarModel>> GatherAllSelectedLayerBars() const;
 
@@ -110,13 +151,13 @@ private:
 		, const FFrameNumber InFirstFrame
 		, const FFrameNumber InCurrentFrame) const;
 
-	void StaggerLayers(const TArrayView<FStaggerLayerBarElement>& InOriginalElements);
-    void StaggerKeyFrames(const TArrayView<FStaggerKeyFrameElement>& InOriginalElements);
+	void StaggerLayers();
+    void StaggerKeyFrames();
 
 	static void SetKeyFrameElementTime(const FStaggerKeyFrameElement& InElement, const FFrameNumber InKeyTime);
 
 	template<typename InElementType>
-	static void FindFirstFrameInTimeline(const TArrayView<InElementType>& InElements, FFrameNumber& OutFrame)
+	static void FindFirstFrameInTimeline(const TArray<InElementType>& InElements, FFrameNumber& OutFrame)
 	{
 		bool bFound = false;
 		for (const InElementType& Element : InElements)
@@ -139,5 +180,13 @@ private:
 	TArray<FStaggerLayerBarElement> OriginalLayerBarElements;
 	TArray<FStaggerKeyFrameElement> OriginalKeyFrameElements;
 
+	int32 CachedLayerBarCount = 0;
+	int32 CachedKeyframeCount = 0;
+
+	TRange<FFrameNumber> CachedRange;
+	FFrameNumber CachedRangeSize = 0;
+	FFrameNumber CachedInterval = 0;
+	FFrameNumber CachedShiftFrames = 0;
+
 	FRandomStream RandomStream;
 };
diff --git a/Engine/Plugins/Experimental/Avalanche/Source/AvalancheSequencer/Private/StaggerTool/AvaStaggerToolSettings.h b/Engine/Plugins/Experimental/Avalanche/Source/AvalancheSequencer/Private/StaggerTool/AvaStaggerToolSettings.h
--- a/Engine/Plugins/Experimental/Avalanche/Source/AvalancheSequencer/Private/StaggerTool/AvaStaggerToolSettings.h
+++ b/Engine/Plugins/Experimental/Avalanche/Source/AvalancheSequencer/Private/StaggerTool/AvaStaggerToolSettings.h
@@ -86,7 +86,7 @@ struct FAvaSequencerStaggerOptions
 	UPROPERTY(EditAnywhere, Category = "Stagger Tool Options", meta = (Tooltip = "Time range to perform the operation on"))
 	EAvaSequencerStaggerRange Range = EAvaSequencerStaggerRange::Default;
 
-	UPROPERTY(EditAnywhere, Category = "Stagger Tool Options", meta = (UIMin = "-500", UIMax = "500", Tooltip = "Custom time range to perform the operation with"))
+	UPROPERTY(EditAnywhere, Category = "Stagger Tool Options", meta = (UIMin = "-60", UIMax = "60", Tooltip = "Custom time range to perform the operation with"))
 	int32 CustomRange = 30;
 
 	UPROPERTY(EditAnywhere, Category = "Stagger Tool Options", meta = (Tooltip = "Position to start the first layer or keyframe"))
@@ -95,10 +95,10 @@ struct FAvaSequencerStaggerOptions
 	UPROPERTY(EditAnywhere, Category = "Stagger Tool Options", meta = (Tooltip = "Point on the current layer bar at which the next layer is positioned"))
 	float OperationPoint = 0.f;
 
-	UPROPERTY(EditAnywhere, Category = "Stagger Tool Options", meta = (UIMin = "-500", UIMax = "500", Tooltip = "Stepping interval between layers or keyframes"))
+	UPROPERTY(EditAnywhere, Category = "Stagger Tool Options", meta = (UIMin = "-60", UIMax = "60", Tooltip = "Stepping interval between layers or keyframes"))
 	int32 Interval = 0;
 
-	UPROPERTY(EditAnywhere, Category = "Stagger Tool Options", meta = (UIMin = "-500", UIMax = "500", Tooltip = "Offset of operation"))
+	UPROPERTY(EditAnywhere, Category = "Stagger Tool Options", meta = (UIMin = "-60", UIMax = "60", Tooltip = "Offset of operation"))
 	int32 Shift = 0;
 
 	UPROPERTY(EditAnywhere, Category = "Stagger Tool Options", meta = (ClampMin = "1", ClampMax = "255", UIMin = "1", UIMax = "10", Tooltip = "Number of layers or keyframes to group together before moving to the next stagger point"))
diff --git a/Engine/Plugins/Experimental/Avalanche/Source/AvalancheSequencer/Private/StaggerTool/Customization/AvaStaggerToolSettingsDetailsCustomization.cpp b/Engine/Plugins/Experimental/Avalanche/Source/AvalancheSequencer/Private/StaggerTool/Customization/AvaStaggerToolSettingsDetailsCustomization.cpp
--- a/Engine/Plugins/Experimental/Avalanche/Source/AvalancheSequencer/Private/StaggerTool/Customization/AvaStaggerToolSettingsDetailsCustomization.cpp
+++ b/Engine/Plugins/Experimental/Avalanche/Source/AvalancheSequencer/Private/StaggerTool/Customization/AvaStaggerToolSettingsDetailsCustomization.cpp
@@ -132,18 +132,22 @@ EVisibility FAvaStaggerToolSettingsDetailsCustomization::GetCurveOffsetPropertyV
 EVisibility FAvaStaggerToolSettingsDetailsCustomization::GetRandomSeedPropertyVisibility() const
 {
 	const EAvaSequencerStaggerDistribution DistributionValue = GetDistributionPropertyValue();
-	return DistributionValue == EAvaSequencerStaggerDistribution::Random ? EVisibility::Visible : EVisibility::Collapsed;
+	return (DistributionValue == EAvaSequencerStaggerDistribution::Random)
+		? EVisibility::Visible : EVisibility::Collapsed;
 }
 
 EVisibility FAvaStaggerToolSettingsDetailsCustomization::GetIntervalPropertyVisibility() const
 {
-	return EVisibility::Visible;
+	const EAvaSequencerStaggerDistribution DistributionValue = GetDistributionPropertyValue();
+	return (DistributionValue != EAvaSequencerStaggerDistribution::Range && DistributionValue != EAvaSequencerStaggerDistribution::Random)
+		? EVisibility::Visible : EVisibility::Collapsed;
 }
 
 EVisibility FAvaStaggerToolSettingsDetailsCustomization::GetShiftPropertyVisibility() const
 {
 	const EAvaSequencerStaggerDistribution DistributionValue = GetDistributionPropertyValue();
-	return (DistributionValue != EAvaSequencerStaggerDistribution::Range) ? EVisibility::Visible : EVisibility::Collapsed;
+	return (DistributionValue != EAvaSequencerStaggerDistribution::Range)
+		? EVisibility::Visible : EVisibility::Collapsed;
 }
 
 EVisibility FAvaStaggerToolSettingsDetailsCustomization::GetRangePropertyVisibility() const
@@ -185,7 +189,7 @@ EVisibility FAvaStaggerToolSettingsDetailsCustomization::GetOperationPointProper
 }
 
 template<typename InValueType>
-TSharedRef<SWidget> CreateSpinBox(const TWeakPtr<FAvaStaggerTool>& InWeakTool, IDetailPropertyRow& PropertyRow, const TSharedRef<IPropertyHandle> InProperty)
+TSharedRef<SWidget> CreateSpinBox(const TWeakPtr<FAvaStaggerTool>& InWeakTool, const TSharedRef<IPropertyHandle> InProperty)
 {
 	const TWeakPtr<IPropertyHandle> WeakProperty = InProperty;
 
@@ -212,6 +216,8 @@ TSharedRef<SWidget> CreateSpinBox(const TWeakPtr<FAvaStaggerTool>& InWeakTool, I
 		.MaxValue(ClampMaxValue)
 		.MinSliderValue(UIMinValue)
 		.MaxSliderValue(UIMaxValue)
+		.ShiftMultiplier(2.0f) // Default is 10.0f
+		.CtrlMultiplier(0.05f) // Default is 0.1f
 		.Value_Lambda([WeakProperty]()
 			{
 				InValueType OutValue = 0;
@@ -239,17 +245,6 @@ TSharedRef<SWidget> CreateSpinBox(const TWeakPtr<FAvaStaggerTool>& InWeakTool, I
 
 void FAvaStaggerToolSettingsDetailsCustomization::AddCustomIntSpinBoxRow(IDetailPropertyRow& PropertyRow, const TSharedRef<IPropertyHandle> InProperty)
 {
-	const TWeakPtr<IPropertyHandle> WeakProperty = InProperty;
-
-	auto GetMetaDataValue = [&InProperty](const FName InKey)
-	{
-		return InProperty->HasMetaData(InKey) ? InProperty->GetIntMetaData(InKey) : TOptional<int32>();	
-	};
-	const TOptional<int32> ClampMinValue = GetMetaDataValue(TEXT("ClampMin"));
-	const TOptional<int32> ClampMaxValue = GetMetaDataValue(TEXT("ClampMax"));
-	const TOptional<int32> UIMinValue = GetMetaDataValue(TEXT("UIMin"));
-	const TOptional<int32> UIMaxValue = GetMetaDataValue(TEXT("UIMax"));
-
 	PropertyRow.CustomWidget()
 		.NameContent()
 		[
@@ -257,23 +252,12 @@ void FAvaStaggerToolSettingsDetailsCustomization::AddCustomIntSpinBoxRow(IDetail
 		]
 		.ValueContent()
 		[
-			CreateSpinBox<int32>(WeakTool, PropertyRow, InProperty)
+			CreateSpinBox<int32>(WeakTool, InProperty)
 		];
 }
 
 void FAvaStaggerToolSettingsDetailsCustomization::AddCustomFloatSpinBoxRow(IDetailPropertyRow& PropertyRow, const TSharedRef<IPropertyHandle> InProperty)
 {
-	const TWeakPtr<IPropertyHandle> WeakProperty = InProperty;
-
-	auto GetMetaDataValue = [&InProperty](const FName InKey)
-	{
-		return InProperty->HasMetaData(InKey) ? InProperty->GetFloatMetaData(InKey) : TOptional<float>();	
-	};
-	const TOptional<float> ClampMinValue = GetMetaDataValue(TEXT("ClampMin"));
-	const TOptional<float> ClampMaxValue = GetMetaDataValue(TEXT("ClampMax"));
-	const TOptional<float> UIMinValue = GetMetaDataValue(TEXT("UIMin"));
-	const TOptional<float> UIMaxValue = GetMetaDataValue(TEXT("UIMax"));
-
 	PropertyRow.CustomWidget()
 		.NameContent()
 		[
@@ -281,7 +265,7 @@ void FAvaStaggerToolSettingsDetailsCustomization::AddCustomFloatSpinBoxRow(IDeta
 		]
 		.ValueContent()
 		[
-			CreateSpinBox<float>(WeakTool, PropertyRow, InProperty)
+			CreateSpinBox<float>(WeakTool, InProperty)
 		];
 }
 
@@ -313,7 +297,7 @@ void FAvaStaggerToolSettingsDetailsCustomization::AddRandomSeedRow(IDetailProper
 			.VAlign(VAlign_Center)
 			.Padding(FMargin(0.f, 0.f, 5.f, 0.f))
 			[
-				CreateSpinBox<int32>(WeakTool, PropertyRow, InProperty)
+				CreateSpinBox<int32>(WeakTool, InProperty)
 			]
 			+ SHorizontalBox::Slot()
 			.AutoWidth()
diff --git a/Engine/Plugins/Experimental/Avalanche/Source/AvalancheSequencer/Private/StaggerTool/Widgets/SAvaStaggerTool.cpp b/Engine/Plugins/Experimental/Avalanche/Source/AvalancheSequencer/Private/StaggerTool/Widgets/SAvaStaggerTool.cpp
--- a/Engine/Plugins/Experimental/Avalanche/Source/AvalancheSequencer/Private/StaggerTool/Widgets/SAvaStaggerTool.cpp
+++ b/Engine/Plugins/Experimental/Avalanche/Source/AvalancheSequencer/Private/StaggerTool/Widgets/SAvaStaggerTool.cpp
@@ -7,6 +7,7 @@
 #include "Misc/Optional.h"
 #include "Modules/ModuleManager.h"
 #include "PropertyEditorModule.h"
+#include "SequencerCommands.h"
 #include "StaggerTool/AvaStaggerTool.h"
 #include "StaggerTool/Customization/AvaStaggerToolSettingsDetailsCustomization.h"
 #include "Styling/StyleColors.h"
@@ -33,19 +34,6 @@ void SAvaStaggerTool::Construct(const FArguments& InArgs, const TWeakPtr<FAvaSta
 		SNew(SVerticalBox)
 		+ SVerticalBox::Slot()
 		.AutoHeight()
-		.Padding(5.f, 1.f, 5.f, 2.f)
-		[
-			ConstructAlignTools()
-		]
-		+ SVerticalBox::Slot()
-		.AutoHeight()
-		[
-			SNew(SSeparator)
-			.ColorAndOpacity(FStyleColors::Dropdown)
-			.SeparatorImage(FAppStyle::Get().GetBrush(TEXT("WhiteBrush")))
-		]
-		+ SVerticalBox::Slot()
-		.AutoHeight()
 		.Padding(5.f, 5.f, 5.f, 2.f)
 		[
 			ConstructApplyRow()
@@ -87,7 +75,7 @@ TSharedRef<SWidget> SAvaStaggerTool::ConstructDetails()
 TSharedRef<SWidget> SAvaStaggerTool::ConstructApplyRow()
 {
 	constexpr float ButtonSize = 22.f;
-	constexpr float ImageSize = 18.f;
+	constexpr float ImageSize = 14.f;
 	constexpr float Spacing = 5.f;
 
 	return 
@@ -106,6 +94,22 @@ TSharedRef<SWidget> SAvaStaggerTool::ConstructApplyRow()
 		.AutoWidth()
 		.VAlign(VAlign_Bottom)
 		.Padding(Spacing, 0.f, 0.f, 0.f)
+		[
+			SNew(SButton)
+			.ContentPadding(2.f)
+			.ToolTipText(FSequencerCommands::Get().AlignSelectionToPlayhead->GetDescription())
+			.IsEnabled(this, &SAvaStaggerTool::OnCanAlignToPlayhead)
+			.OnClicked(this, &SAvaStaggerTool::OnAlignToPlayhead)
+			[
+				SNew(SImage)
+				.DesiredSizeOverride(FVector2D(ImageSize))
+				.Image(FAppStyle::GetBrush(TEXT("HorizontalAlignment_Left")))
+			]
+		]
+		+ SHorizontalBox::Slot()
+		.AutoWidth()
+		.VAlign(VAlign_Bottom)
+		.Padding(Spacing, 0.f, 0.f, 0.f)
 		[
 			SNew(SButton)
 			.Text(LOCTEXT("ApplyButtonText", "Apply"))
@@ -179,40 +183,6 @@ TSharedRef<SWidget> SAvaStaggerTool::ConstructApplyRow()
 		];
 }
 
-TSharedRef<SWidget> SAvaStaggerTool::ConstructAlignTools()
-{
-	return SNew(SBox)
-		.Padding(5.f)
-		[
-			SNew(SHorizontalBox)
-			+ SHorizontalBox::Slot()
-			.AutoWidth()
-			.HAlign(HAlign_Left)
-			.VAlign(VAlign_Center)
-			.Padding(FMargin(0.f, 0.f, 10.f, 0.f))
-			[
-				SNew(STextBlock)
-				.TextStyle(FAppStyle::Get(), TEXT("SmallText"))
-				.Text(LOCTEXT("AlignToText", "Align to: "))
-			]
-			+ SHorizontalBox::Slot()
-			.AutoWidth()
-			.HAlign(HAlign_Left)
-			.VAlign(VAlign_Center)
-			[
-				SNew(SButton)
-				.TextStyle(FAppStyle::Get(), TEXT("SmallText"))
-				.Text(LOCTEXT("AlignToPlayheadText", "Playhead"))
-				.IsEnabled_Lambda([this]()
-					{
-						const TSharedPtr<FAvaStaggerTool> StaggerTool = WeakTool.Pin();
-						return StaggerTool.IsValid() ? (StaggerTool->IsLayerBarSelection() || StaggerTool->IsKeyFrameSelection()) : false;
-					})
-				.OnClicked(this, &SAvaStaggerTool::OnAlignToPlayhead)
-			]
-		];
-}
-
 FText SAvaStaggerTool::GetSelectionText() const
 {
 	const TSharedPtr<FAvaStaggerTool> StaggerTool = WeakTool.Pin();
@@ -281,13 +251,21 @@ void SAvaStaggerTool::OnToggleAutoUpdateClick(const ECheckBoxState InNewState)
 	}
 }
 
+bool SAvaStaggerTool::OnCanAlignToPlayhead() const
+{
+	if (const TSharedPtr<FAvaStaggerTool> Tool = WeakTool.Pin())
+	{
+		return Tool->CanAlignToPlayhead();
+	}
+	return false;
+}
+
 FReply SAvaStaggerTool::OnAlignToPlayhead()
 {
 	if (const TSharedPtr<FAvaStaggerTool> Tool = WeakTool.Pin())
 	{
 		Tool->AlignToPlayhead();
 	}
-
 	return FReply::Handled();
 }
 
diff --git a/Engine/Plugins/Experimental/Avalanche/Source/AvalancheSequencer/Private/StaggerTool/Widgets/SAvaStaggerTool.h b/Engine/Plugins/Experimental/Avalanche/Source/AvalancheSequencer/Private/StaggerTool/Widgets/SAvaStaggerTool.h
--- a/Engine/Plugins/Experimental/Avalanche/Source/AvalancheSequencer/Private/StaggerTool/Widgets/SAvaStaggerTool.h
+++ b/Engine/Plugins/Experimental/Avalanche/Source/AvalancheSequencer/Private/StaggerTool/Widgets/SAvaStaggerTool.h
@@ -30,7 +30,6 @@ public:
 protected:
 	TSharedRef<SWidget> ConstructDetails();
 	TSharedRef<SWidget> ConstructApplyRow();
-	TSharedRef<SWidget> ConstructAlignTools();
 
 	FText GetSelectionText() const;
 	FText GetSelectionToolTipText() const;
@@ -40,6 +39,7 @@ protected:
 	FReply OnResetToDefaultsClick();
 	void OnToggleAutoUpdateClick(const ECheckBoxState InNewState);
 
+	bool OnCanAlignToPlayhead() const;
 	FReply OnAlignToPlayhead();
 
 	//~ Begin FNotifyHook
