MotionDesign : Modifiers
- Fix SizeToTexture modifier EditCondition was not renamed properly

diff --git a/Engine/Plugins/Experimental/Avalanche/Source/AvalancheModifiers/Public/Modifiers/AvaSizeToTextureModifier.h b/Engine/Plugins/Experimental/Avalanche/Source/AvalancheModifiers/Public/Modifiers/AvaSizeToTextureModifier.h
--- a/Engine/Plugins/Experimental/Avalanche/Source/AvalancheModifiers/Public/Modifiers/AvaSizeToTextureModifier.h
+++ b/Engine/Plugins/Experimental/Avalanche/Source/AvalancheModifiers/Public/Modifiers/AvaSizeToTextureModifier.h
@@ -105,11 +105,11 @@ protected:
 	EAvaSizeToTextureRule Rule = EAvaSizeToTextureRule::FixedHeight;
 
 	/** The fixed height size */
-	UPROPERTY(EditInstanceOnly, Setter, Getter, Category="TextureSize", meta=(ClampMin="0", EditCondition="Rule == EAvaTextureSizeRule::FixedHeight", EditConditionHides, AllowPrivateAccess="true"))
+	UPROPERTY(EditInstanceOnly, Setter, Getter, Category="TextureSize", meta=(ClampMin="0", EditCondition="Rule == EAvaSizeToTextureRule::FixedHeight", EditConditionHides, AllowPrivateAccess="true"))
 	float FixedHeight = 512.f;
 
 	/** The fixed width size */
-	UPROPERTY(EditInstanceOnly, Setter, Getter, Category="TextureSize", meta=(ClampMin="0", EditCondition="Rule == EAvaTextureSizeRule::FixedWidth", EditConditionHides, AllowPrivateAccess="true"))
+	UPROPERTY(EditInstanceOnly, Setter, Getter, Category="TextureSize", meta=(ClampMin="0", EditCondition="Rule == EAvaSizeToTextureRule::FixedWidth", EditConditionHides, AllowPrivateAccess="true"))
 	float FixedWidth = 512.f;
 
 private:
