Fixed missing return.

diff --git a/Engine/Plugins/Experimental/Avalanche/Source/AvalancheSceneRigEditor/Private/AvaSceneRigEditorModule.cpp b/Engine/Plugins/Experimental/Avalanche/Source/AvalancheSceneRigEditor/Private/AvaSceneRigEditorModule.cpp
--- a/Engine/Plugins/Experimental/Avalanche/Source/AvalancheSceneRigEditor/Private/AvaSceneRigEditorModule.cpp
+++ b/Engine/Plugins/Experimental/Avalanche/Source/AvalancheSceneRigEditor/Private/AvaSceneRigEditorModule.cpp
@@ -188,7 +188,7 @@ FSoftObjectPath FAvaSceneRigEditorModule::GetActiveSceneRig(UWorld* const InWorl
 	const UAvaSceneSettings* const SceneSettings = GetSceneSettings(InWorld);
 	if (!IsValid(SceneSettings))
 	{
-		FSoftObjectPath();
+		return FSoftObjectPath();
 	}
 
 	return SceneSettings->GetSceneRig();
