Motion Design: fix issues with tabs registered through motion design editor were not properly retained

diff --git a/Engine/Plugins/Experimental/Avalanche/Source/AvalancheEditor/Private/Outliner/AvaOutlinerExtension.cpp b/Engine/Plugins/Experimental/Avalanche/Source/AvalancheEditor/Private/Outliner/AvaOutlinerExtension.cpp
--- a/Engine/Plugins/Experimental/Avalanche/Source/AvalancheEditor/Private/Outliner/AvaOutlinerExtension.cpp
+++ b/Engine/Plugins/Experimental/Avalanche/Source/AvalancheEditor/Private/Outliner/AvaOutlinerExtension.cpp
@@ -127,7 +127,7 @@ void FAvaOutlinerExtension::ExtendLevelEditorLayout(FLayoutExtender& InExtender)
 	{
 		InExtender.ExtendLayout(LevelEditorTabIds::LevelEditorSceneOutliner
 			, ELayoutExtensionPosition::Before
-			, FTabManager::FTab(FAvaOutlinerTabSpawner::GetTabID(Index), Index == 0 ? ETabState::OpenedTab : ETabState::ClosedTab));
+			, FTabManager::FTab(FAvaOutlinerTabSpawner::GetTabID(Index), ETabState::ClosedTab));
 	}
 }
 
diff --git a/Engine/Plugins/Experimental/Avalanche/Source/AvalancheEditor/Private/Outliner/AvaOutlinerTabSpawner.cpp b/Engine/Plugins/Experimental/Avalanche/Source/AvalancheEditor/Private/Outliner/AvaOutlinerTabSpawner.cpp
--- a/Engine/Plugins/Experimental/Avalanche/Source/AvalancheEditor/Private/Outliner/AvaOutlinerTabSpawner.cpp
+++ b/Engine/Plugins/Experimental/Avalanche/Source/AvalancheEditor/Private/Outliner/AvaOutlinerTabSpawner.cpp
@@ -13,7 +13,7 @@
 
 FName FAvaOutlinerTabSpawner::GetTabID(int32 InOutlinerId)
 {
-	return *FString::Printf(TEXT("AvalancheOutliner%d"), InOutlinerId + 1);
+	return *FString::Printf(TEXT("AvaOutliner%d"), InOutlinerId + 1);
 }
 
 FAvaOutlinerTabSpawner::FAvaOutlinerTabSpawner(int32 InOutlinerId, const TSharedRef<IAvaEditor>& InEditor)
diff --git a/Engine/Plugins/Experimental/Avalanche/Source/AvalancheEditorCore/Private/AvaEditor.cpp b/Engine/Plugins/Experimental/Avalanche/Source/AvalancheEditorCore/Private/AvaEditor.cpp
--- a/Engine/Plugins/Experimental/Avalanche/Source/AvalancheEditorCore/Private/AvaEditor.cpp
+++ b/Engine/Plugins/Experimental/Avalanche/Source/AvalancheEditorCore/Private/AvaEditor.cpp
@@ -158,6 +158,7 @@ FAvaEditor::FAvaEditor(FAvaEditorBuilder& Initializer)
 
 FAvaEditor::~FAvaEditor()
 {
+	FAvaEditor::UnregisterTabSpawners();
 	UnbindDelegates();
 	EditorExtensions.Reset();
 	FAvaEditorExtensionTypeRegistry::Get().RemoveStaleEntries();
@@ -212,6 +213,12 @@ void FAvaEditor::InvokeTabs()
 
 	for (const FTabId& TabId : TabsToInvoke)
 	{
+		// Do not re-use the existing tabs here.
+		// Instead, invoke a whole new tab
+		if (TSharedPtr<SDockTab> Tab = TabManager->FindExistingLiveTab(TabId))
+		{
+			Tab->RequestCloseTab();
+		}
 		TabManager->TryInvokeTab(TabId);
 	}
 }
@@ -774,15 +781,11 @@ void FAvaEditor::DeactivateExtensions()
 	// and restoring back to its original value (which could well still be false)
 	TGuardValue<bool> EditorActive(bIsActive, false);
 
-	CloseTabs();
-
 	// Force deactivation
 	ForEachExtension([](const TSharedRef<IAvaEditorExtension>& InExtension)
 		{
 			InExtension->Deactivate();
 		});
-
-	UnregisterTabSpawners();
 }
 
 bool FAvaEditor::CopyToString(FString& OutCopyData) const
