Motion Design: outliner to not show empty level item and to clear the widget item list on the underlying outliner data destruction

diff --git a/Engine/Plugins/Experimental/Avalanche/Source/AvalancheOutliner/Private/AvaOutlinerView.cpp b/Engine/Plugins/Experimental/Avalanche/Source/AvalancheOutliner/Private/AvaOutlinerView.cpp
--- a/Engine/Plugins/Experimental/Avalanche/Source/AvalancheOutliner/Private/AvaOutlinerView.cpp
+++ b/Engine/Plugins/Experimental/Avalanche/Source/AvalancheOutliner/Private/AvaOutlinerView.cpp
@@ -46,6 +46,7 @@ FAvaOutlinerView::FAvaOutlinerView(FPrivateToken)
 FAvaOutlinerView::~FAvaOutlinerView()
 {
 	SaveState();
+	ClearRootItemSource();
 
 	if (UObjectInitialized())
 	{
@@ -581,6 +582,21 @@ void FAvaOutlinerView::UpdateRootVisibleItems()
 	OutlinerStats->MarkCountTypeDirty(EAvaOutlinerStatCountType::VisibleItemCount);
 }
 
+void FAvaOutlinerView::ClearRootItemSource()
+{
+	RootVisibleItems.Reset();
+	ReadOnlyItems.Reset();
+	OutlinerStats->MarkCountTypeDirty(EAvaOutlinerStatCountType::VisibleItemCount);
+
+	if (OutlinerWidget.IsValid())
+	{
+		if (TSharedPtr<SAvaOutlinerTreeView> TreeView = OutlinerWidget->GetTreeView())
+		{
+			TreeView->RebuildList();	
+		}
+	}
+}
+
 void FAvaOutlinerView::UpdateItemExpansions()
 {
 	TArray<FAvaOutlinerItemPtr> Items = RootVisibleItems;
diff --git a/Engine/Plugins/Experimental/Avalanche/Source/AvalancheOutliner/Private/AvaOutlinerView.h b/Engine/Plugins/Experimental/Avalanche/Source/AvalancheOutliner/Private/AvaOutlinerView.h
--- a/Engine/Plugins/Experimental/Avalanche/Source/AvalancheOutliner/Private/AvaOutlinerView.h
+++ b/Engine/Plugins/Experimental/Avalanche/Source/AvalancheOutliner/Private/AvaOutlinerView.h
@@ -143,6 +143,8 @@ public:
 	/** Refreshes the Items that will be at the Top Level of the Tree (e.g. Actors with no parent item) */
 	void UpdateRootVisibleItems();
 
+	void ClearRootItemSource();
+
 	/** Updates the Item Expansions in the Tree Widget based on whether the Items have the Expanded flag or not */
 	void UpdateItemExpansions();
 
