Motion Design: Scene Settings tab to extend level editor layout

diff --git a/Engine/Plugins/Experimental/Avalanche/Source/AvalancheEditor/Private/Scene/AvaSceneExtension.cpp b/Engine/Plugins/Experimental/Avalanche/Source/AvalancheEditor/Private/Scene/AvaSceneExtension.cpp
--- a/Engine/Plugins/Experimental/Avalanche/Source/AvalancheEditor/Private/Scene/AvaSceneExtension.cpp
+++ b/Engine/Plugins/Experimental/Avalanche/Source/AvalancheEditor/Private/Scene/AvaSceneExtension.cpp
@@ -3,6 +3,8 @@
 #include "AvaSceneExtension.h"
 #include "AvaSceneSettingsTabSpawner.h"
 #include "Engine/World.h"
+#include "Framework/Docking/LayoutExtender.h"
+#include "LevelEditor.h"
 #include "Scene/AvaSceneDefaults.h"
 #include "ScopedTransaction.h"
 #include "ToolMenu.h"
@@ -29,6 +31,13 @@ void FAvaSceneExtension::RegisterTabSpawners(const TSharedRef<IAvaEditor>& InEdi
 	InEditor->AddTabSpawner<FAvaSceneSettingsTabSpawner>(InEditor);
 }
 
+void FAvaSceneExtension::ExtendLevelEditorLayout(FLayoutExtender& InExtender) const
+{
+	InExtender.ExtendLayout(LevelEditorTabIds::LevelEditorSelectionDetails
+		, ELayoutExtensionPosition::After
+		, FTabManager::FTab(FAvaSceneSettingsTabSpawner::GetTabID(), ETabState::ClosedTab));
+}
+
 void FAvaSceneExtension::SpawnDefaultScene()
 {
 	TSharedPtr<IAvaEditor> Editor = GetEditor();
diff --git a/Engine/Plugins/Experimental/Avalanche/Source/AvalancheEditor/Private/Scene/AvaSceneExtension.h b/Engine/Plugins/Experimental/Avalanche/Source/AvalancheEditor/Private/Scene/AvaSceneExtension.h
--- a/Engine/Plugins/Experimental/Avalanche/Source/AvalancheEditor/Private/Scene/AvaSceneExtension.h
+++ b/Engine/Plugins/Experimental/Avalanche/Source/AvalancheEditor/Private/Scene/AvaSceneExtension.h
@@ -12,6 +12,7 @@ public:
 	//~ Begin IAvaEditorExtension
 	virtual void ExtendToolbarMenu(UToolMenu& InMenu) override;
 	virtual void RegisterTabSpawners(const TSharedRef<IAvaEditor>& InEditor) const override;
+	virtual void ExtendLevelEditorLayout(FLayoutExtender& InExtender) const override;
 	//~ End IAvaEditorExtension
 
 private:
