MotionDesign : Shapes
- Allow multi edit for different type of shape generator in details panel material section if supported

diff --git a/Engine/Plugins/Experimental/Avalanche/Source/AvalancheEditor/Private/DetailView/Customizations/AvaMeshesDetailCustomization.cpp b/Engine/Plugins/Experimental/Avalanche/Source/AvalancheEditor/Private/DetailView/Customizations/AvaMeshesDetailCustomization.cpp
--- a/Engine/Plugins/Experimental/Avalanche/Source/AvalancheEditor/Private/DetailView/Customizations/AvaMeshesDetailCustomization.cpp
+++ b/Engine/Plugins/Experimental/Avalanche/Source/AvalancheEditor/Private/DetailView/Customizations/AvaMeshesDetailCustomization.cpp
@@ -43,20 +43,36 @@ void FAvaMeshesDetailCustomization::CustomizeDetails(IDetailLayoutBuilder& InDet
 	IDetailCategoryBuilder& MaterialCategoryBuilder = InDetailBuilder.EditCategory(FName("Material"));
 	MaterialCategoryBuilder.SetSortOrder(ShapeCategoryBuilder.GetSortOrder() + 1);
 
-	// Make sure we have the mesh generator class in multi selection
-	bool bSameMeshGenerator = true;
-	UClass* MeshGeneratorClass = nullptr;
+	// Make sure we have common mesh sections to display material properties
+	TMap<int32, FName> MeshSectionNames;
+	bool bFirstElement = true;
 	for (TArray<TWeakObjectPtr<UAvaShapeDynamicMeshBase>>::TIterator It(MeshGeneratorsWeak); It; ++It)
 	{
 		if (UAvaShapeDynamicMeshBase* MeshGenerator = It->Get())
 		{
-			if (MeshGeneratorClass != nullptr && MeshGeneratorClass != MeshGenerator->GetClass())
+			if (bFirstElement)
 			{
-				bSameMeshGenerator = false;
-				break;
+				const TArray<FName> MeshNames = MeshGenerator->GetMeshDataNames();
+				for (int32 Index : MeshGenerator->GetMeshesIndexes())
+				{
+					if (MeshNames.IsValidIndex(Index))
+					{
+						MeshSectionNames.Add(Index, MeshNames[Index]);
+					}
+				}
+				bFirstElement = false;
+			}
+			else
+			{
+				const TSet<int32> MeshIndex = MeshGenerator->GetMeshesIndexes();
+				for (TMap<int32, FName>::TIterator MeshIt(MeshSectionNames); MeshIt; ++MeshIt)
+				{
+					if (!MeshIndex.Contains(MeshIt->Key))
+					{
+						MeshIt.RemoveCurrent();
+					}
+				}
 			}
-			
-			MeshGeneratorClass = MeshGenerator->GetClass();
 		}
 		else
 		{
@@ -64,22 +80,28 @@ void FAvaMeshesDetailCustomization::CustomizeDetails(IDetailLayoutBuilder& InDet
 		}
 	}
 
-	if (!MeshGeneratorsWeak.IsEmpty() && bSameMeshGenerator)
+	if (!MeshGeneratorsWeak.IsEmpty() && !MeshSectionNames.IsEmpty())
 	{
-		UAvaShapeDynamicMeshBase* DynMesh = MeshGeneratorsWeak[0].Get();
-		const TArray<FName> MeshNames = DynMesh->GetMeshDataNames();
+		TSharedPtr<IPropertyHandleMap> MapHandle = MeshDatasHandle->AsMap();
+		check(MapHandle.IsValid())
 
-		for (int32 Index = 0; Index < MeshNames.Num(); ++Index)
+		uint32 NumSections = 0;
+		MapHandle->GetNumElements(NumSections);
+		for (int32 Index = 0; Index < static_cast<int32>(NumSections); ++Index)
 		{
-			FString MeshName = TEXT("");
+			if (!MeshSectionNames.Contains(Index))
+			{
+				continue;
+			}
+
+			TSharedPtr<IPropertyHandle> MeshPropertyHandle = MapHandle->GetElement(Index);
 
-			if (Index > 0)
+			if (!MeshPropertyHandle.IsValid())
 			{
-				MeshName = MeshNames[Index].ToString() + TEXT(" ");
+				continue;
 			}
 
-			TSharedPtr<IPropertyHandle> MeshPropertyHandle = MeshDatasHandle->GetChildHandle(Index);
-			check(MeshPropertyHandle.IsValid());
+			const FString MeshName = Index != 0 ? MeshSectionNames[Index].ToString() + TEXT(" ") : TEXT("");
 
 			// Material Type
 			TSharedPtr<IPropertyHandle> MaterialTypeHandle = MeshPropertyHandle->GetChildHandle(GET_MEMBER_NAME_CHECKED(FAvaShapeMeshData, MaterialType));
@@ -152,7 +174,7 @@ void FAvaMeshesDetailCustomization::CustomizeDetails(IDetailLayoutBuilder& InDet
 			}));
 
 			// Only Add it the first time at this specific point
-			if (Index == 0 && MeshNames.Num() > 1)
+			if (Index == 0 && MeshSectionNames.Num() > 1)
 			{
 				MaterialCategoryBuilder.AddProperty(UsePrimaryMaterialEverywhereHandle);
 			}
@@ -168,7 +190,7 @@ void FAvaMeshesDetailCustomization::CustomizeDetails(IDetailLayoutBuilder& InDet
 				return MaterialUVHandle->IsEditable() ? EVisibility::Visible : EVisibility::Hidden;
 			}));
 
-			if (Index < MeshNames.Num())
+			if (Index < MeshSectionNames.Num())
 			{
 				// Separator row
 				FDetailWidgetRow& SeparatorRow = MaterialCategoryBuilder.AddCustomRow(FText::GetEmpty());
