[Motion Design] Ease Curve Tool:
- Change to hold ISequencer pointer instead of FAvaSequencer
- Move BindCommands() to FAvaEaseCurveTool from SAvaEaseCurveTool
- Change SAvaEaseCurveTool to hold weak FAvaEaseCurveTool pointer instead of shared
- Disable menu entries for Quick Ease when no Quick Ease preset is setup

diff --git a/Engine/Plugins/Experimental/Avalanche/Source/AvalancheSequencer/Private/AvaSequencer.cpp b/Engine/Plugins/Experimental/Avalanche/Source/AvalancheSequencer/Private/AvaSequencer.cpp
--- a/Engine/Plugins/Experimental/Avalanche/Source/AvalancheSequencer/Private/AvaSequencer.cpp
+++ b/Engine/Plugins/Experimental/Avalanche/Source/AvalancheSequencer/Private/AvaSequencer.cpp
@@ -320,19 +320,22 @@ void FAvaSequencer::EnsureSequencer()
 	}
 
 	// Create ease curve tool and map commands
-	EaseCurveTool = MakeShared<FAvaEaseCurveTool>(SharedThis(this));
+	EaseCurveTool = MakeShared<FAvaEaseCurveTool>(Sequencer.ToSharedRef());
 
 	const TSharedRef<FAvaEaseCurveTool> EaseCurveToolRef = EaseCurveTool.ToSharedRef();
 	const FAvaEaseCurveToolCommands& EaseCurveToolCommands = FAvaEaseCurveToolCommands::Get();
 
 	CommandList->MapAction(EaseCurveToolCommands.QuickEaseIn
-		, FExecuteAction::CreateSP(EaseCurveToolRef, &FAvaEaseCurveTool::ApplyQuickEaseToSequencerKeySelections, EAvaEaseCurveToolOperation::In));
+		, FExecuteAction::CreateSP(EaseCurveToolRef, &FAvaEaseCurveTool::ApplyQuickEaseToSequencerKeySelections, EAvaEaseCurveToolOperation::In)
+		, FCanExecuteAction::CreateSP(EaseCurveToolRef, &FAvaEaseCurveTool::CanApplyQuickEaseToSequencerKeySelections));
 
 	CommandList->MapAction(EaseCurveToolCommands.QuickEase
-		, FExecuteAction::CreateSP(EaseCurveToolRef, &FAvaEaseCurveTool::ApplyQuickEaseToSequencerKeySelections, EAvaEaseCurveToolOperation::InOut));
+		, FExecuteAction::CreateSP(EaseCurveToolRef, &FAvaEaseCurveTool::ApplyQuickEaseToSequencerKeySelections, EAvaEaseCurveToolOperation::InOut)
+		, FCanExecuteAction::CreateSP(EaseCurveToolRef, &FAvaEaseCurveTool::CanApplyQuickEaseToSequencerKeySelections));
 
 	CommandList->MapAction(EaseCurveToolCommands.QuickEaseOut
-		, FExecuteAction::CreateSP(EaseCurveToolRef, &FAvaEaseCurveTool::ApplyQuickEaseToSequencerKeySelections, EAvaEaseCurveToolOperation::Out));
+		, FExecuteAction::CreateSP(EaseCurveToolRef, &FAvaEaseCurveTool::ApplyQuickEaseToSequencerKeySelections, EAvaEaseCurveToolOperation::Out)
+		, FCanExecuteAction::CreateSP(EaseCurveToolRef, &FAvaEaseCurveTool::CanApplyQuickEaseToSequencerKeySelections));
 }
 
 TSharedRef<ISequencer> FAvaSequencer::CreateSequencer()
diff --git a/Engine/Plugins/Experimental/Avalanche/Source/AvalancheSequencer/Private/EaseCurveTool/AvaEaseCurveKeySelection.cpp b/Engine/Plugins/Experimental/Avalanche/Source/AvalancheSequencer/Private/EaseCurveTool/AvaEaseCurveKeySelection.cpp
--- a/Engine/Plugins/Experimental/Avalanche/Source/AvalancheSequencer/Private/EaseCurveTool/AvaEaseCurveKeySelection.cpp
+++ b/Engine/Plugins/Experimental/Avalanche/Source/AvalancheSequencer/Private/EaseCurveTool/AvaEaseCurveKeySelection.cpp
@@ -1,7 +1,6 @@
 ﻿// Copyright Epic Games, Inc. All Rights Reserved.
 
 #include "AvaEaseCurveKeySelection.h"
-#include "AvaSequencer.h"
 #include "Channels/MovieSceneChannel.h"
 #include "Channels/MovieSceneDoubleChannel.h"
 #include "Channels/MovieSceneFloatChannel.h"
@@ -39,7 +38,7 @@ namespace UE::Ava::EaseCurveTool::Private
 	}
 }
 
-FAvaEaseCurveKeySelection::FAvaEaseCurveKeySelection(const TSharedPtr<FAvaSequencer>& InSequencer)
+FAvaEaseCurveKeySelection::FAvaEaseCurveKeySelection(const TSharedPtr<ISequencer>& InSequencer)
 {
 	using namespace UE::Ava::EaseCurveTool::Private;
 
@@ -50,7 +49,7 @@ FAvaEaseCurveKeySelection::FAvaEaseCurveKeySelection(const TSharedPtr<FAvaSequen
 
 	TSharedPtr<FSequencerSelection> SequencerSelection;
 
-	if (const TSharedPtr<FSequencerEditorViewModel> SequencerViewModel = InSequencer->GetSequencer()->GetViewModel())
+	if (const TSharedPtr<FSequencerEditorViewModel> SequencerViewModel = InSequencer->GetViewModel())
 	{
 		SequencerSelection = SequencerViewModel->GetSelection();
 	}
diff --git a/Engine/Plugins/Experimental/Avalanche/Source/AvalancheSequencer/Private/EaseCurveTool/AvaEaseCurveKeySelection.h b/Engine/Plugins/Experimental/Avalanche/Source/AvalancheSequencer/Private/EaseCurveTool/AvaEaseCurveKeySelection.h
--- a/Engine/Plugins/Experimental/Avalanche/Source/AvalancheSequencer/Private/EaseCurveTool/AvaEaseCurveKeySelection.h
+++ b/Engine/Plugins/Experimental/Avalanche/Source/AvalancheSequencer/Private/EaseCurveTool/AvaEaseCurveKeySelection.h
@@ -6,7 +6,6 @@
 #include "Channels/MovieSceneChannelHandle.h"
 #include "MVVM/Selection/Selection.h"
 
-class FAvaSequencer;
 class UMovieSceneSection;
 struct FAvaEaseCurveTangents;
 enum class EAvaEaseCurveToolOperation : uint8;
@@ -29,7 +28,7 @@ public:
 	};
 
 	FAvaEaseCurveKeySelection() {}
-	FAvaEaseCurveKeySelection(const TSharedPtr<FAvaSequencer>& InSequencer);
+	FAvaEaseCurveKeySelection(const TSharedPtr<ISequencer>& InSequencer);
 
 	void ForEachEaseableKey(const bool bInIncludeEqualValueKeys
 		, const TFunctionRef<bool(const FKeyHandle& /*InKeyHandle*/, const FKeyHandle& /*InNextKeyHandle*/, const FChannelData&)>& InCallable);
diff --git a/Engine/Plugins/Experimental/Avalanche/Source/AvalancheSequencer/Private/EaseCurveTool/AvaEaseCurveTool.cpp b/Engine/Plugins/Experimental/Avalanche/Source/AvalancheSequencer/Private/EaseCurveTool/AvaEaseCurveTool.cpp
--- a/Engine/Plugins/Experimental/Avalanche/Source/AvalancheSequencer/Private/EaseCurveTool/AvaEaseCurveTool.cpp
+++ b/Engine/Plugins/Experimental/Avalanche/Source/AvalancheSequencer/Private/EaseCurveTool/AvaEaseCurveTool.cpp
@@ -3,7 +3,8 @@
 #include "AvaEaseCurveTool.h"
 #include "AssetRegistry/AssetRegistryModule.h"
 #include "AssetToolsModule.h"
-#include "AvaSequencer.h"
+#include "AvaEaseCurveSubsystem.h"
+#include "AvaEaseCurveToolCommands.h"
 #include "Dialogs/DlgPickAssetPath.h"
 #include "EaseCurveTool/AvaEaseCurveKeySelection.h"
 #include "EaseCurveTool/AvaEaseCurveToolSettings.h"
@@ -11,11 +12,13 @@
 #include "Editor.h"
 #include "EngineAnalytics.h"
 #include "Factories/CurveFactory.h"
+#include "Framework/Commands/GenericCommands.h"
 #include "Framework/Notifications/NotificationManager.h"
 #include "HAL/PlatformApplicationMisc.h"
 #include "ISequencer.h"
 #include "ISettingsModule.h"
 #include "Math/UnrealMathUtility.h"
+#include "Misc/MessageDialog.h"
 #include "ScopedTransaction.h"
 #include "Settings/AvaSequencerSettings.h"
 #include "Widgets/Notifications/SNotificationList.h"
@@ -24,21 +27,23 @@ using namespace UE::Sequencer;
 
 #define LOCTEXT_NAMESPACE "AvaEaseCurveTool"
 
-FAvaEaseCurveTool::FAvaEaseCurveTool(const TSharedRef<FAvaSequencer>& InSequencer)
-	: AvaSequencerWeak(InSequencer)
+FAvaEaseCurveTool::FAvaEaseCurveTool(const TSharedRef<ISequencer>& InSequencer)
+	: SequencerWeak(InSequencer)
 {
 	EaseCurve = NewObject<UAvaEaseCurve>(GetTransientPackage(), NAME_None, RF_Transient | RF_Transactional);
 
+	ToolSettings = GetMutableDefault<UAvaEaseCurveToolSettings>();
+
 	UpdateEaseCurveFromSequencerKeySelections();
 
-	InSequencer->GetSequencer()->GetSelectionChangedObjectGuids().AddRaw(this, &FAvaEaseCurveTool::OnSequencerSelectionChanged);
+	InSequencer->GetSelectionChangedObjectGuids().AddRaw(this, &FAvaEaseCurveTool::OnSequencerSelectionChanged);
 }
 
 FAvaEaseCurveTool::~FAvaEaseCurveTool()
 {
-	if (const TSharedPtr<FAvaSequencer> AvaSequencer = AvaSequencerWeak.Pin())
+	if (const TSharedPtr<ISequencer> Sequencer = SequencerWeak.Pin())
 	{
-		AvaSequencer->GetSequencer()->GetSelectionChangedObjectGuids().RemoveAll(this);
+		Sequencer->GetSelectionChangedObjectGuids().RemoveAll(this);
 	}
 }
 
@@ -63,6 +68,8 @@ TSharedRef<SWidget> FAvaEaseCurveTool::GenerateWidget()
 
 	if (!ToolWidget.IsValid())
 	{
+		BindCommands();
+
 		ToolWidget = SNew(SAvaEaseCurveTool, SharedThis(this))
 			.InitialTangents(GetEaseCurveTangents())
 			.Visibility(this, &FAvaEaseCurveTool::GetVisibility)
@@ -134,7 +141,7 @@ void FAvaEaseCurveTool::SetEaseCurveTangents(const FAvaEaseCurveTangents& InTang
 		SetSequencerKeySelectionTangents(InTangents, InOperation);
 	}
 
-	KeyCache = FAvaEaseCurveKeySelection(AvaSequencerWeak.Pin());
+	KeyCache = FAvaEaseCurveKeySelection(SequencerWeak.Pin());
 }
 
 void FAvaEaseCurveTool::ResetEaseCurveTangents(const EAvaEaseCurveToolOperation InOperation)
@@ -222,12 +229,16 @@ void FAvaEaseCurveTool::FlattenOrStraightenTangents(const EAvaEaseCurveToolOpera
 	}
 }
 
-void FAvaEaseCurveTool::ApplyQuickEaseToSequencerKeySelections(const EAvaEaseCurveToolOperation InOperation)
+bool FAvaEaseCurveTool::CanApplyQuickEaseToSequencerKeySelections() const
 {
-	const UAvaEaseCurveToolSettings* const Settings = GetDefault<UAvaEaseCurveToolSettings>();
+	FAvaEaseCurveTangents Tangents;
+	return FAvaEaseCurveTangents::FromString(ToolSettings->GetQuickEaseTangents(), Tangents);
+}
 
+void FAvaEaseCurveTool::ApplyQuickEaseToSequencerKeySelections(const EAvaEaseCurveToolOperation InOperation)
+{
 	FAvaEaseCurveTangents Tangents;
-	if (!FAvaEaseCurveTangents::FromString(Settings->GetQuickEaseTangents(), Tangents))
+	if (!FAvaEaseCurveTangents::FromString(ToolSettings->GetQuickEaseTangents(), Tangents))
 	{
 		UE_LOG(LogTemp, Warning, TEXT("Ease curve tool failed to apply quick ease tangents: "
 			"Could not parse configured quick ease tangent string."));
@@ -263,29 +274,21 @@ void FAvaEaseCurveTool::ApplyQuickEaseToSequencerKeySelections(const EAvaEaseCur
 
 void FAvaEaseCurveTool::SetSequencerKeySelectionTangents(const FAvaEaseCurveTangents& InTangents, const EAvaEaseCurveToolOperation InOperation)
 {
-	KeyCache = FAvaEaseCurveKeySelection(AvaSequencerWeak.Pin());
+	KeyCache = FAvaEaseCurveKeySelection(SequencerWeak.Pin());
 
 	if (KeyCache.GetTotalSelectedKeys() == 0)
 	{
 		return;
 	}
 
-	const UAvaEaseCurveToolSettings* const EaseCurveToolSettings = GetDefault<UAvaEaseCurveToolSettings>();
-	check(EaseCurveToolSettings);
-	const bool bAutoFlipTangents = EaseCurveToolSettings->GetAutoFlipTangents();
-
-	KeyCache.SetTangents(InTangents, InOperation, GetDisplayRate(), GetTickResolution(), bAutoFlipTangents);
+	KeyCache.SetTangents(InTangents, InOperation, GetDisplayRate(), GetTickResolution(), ToolSettings->GetAutoFlipTangents());
 }
 
 void FAvaEaseCurveTool::UpdateEaseCurveFromSequencerKeySelections()
 {
-	KeyCache = FAvaEaseCurveKeySelection(AvaSequencerWeak.Pin());
-
-	const UAvaEaseCurveToolSettings* const EaseCurveToolSettings = GetDefault<UAvaEaseCurveToolSettings>();
-	check(EaseCurveToolSettings);
-	const bool bAutoFlipTangents = EaseCurveToolSettings->GetAutoFlipTangents();
+	KeyCache = FAvaEaseCurveKeySelection(SequencerWeak.Pin());
 
-	const FAvaEaseCurveTangents AverageTangents = KeyCache.AverageTangents(GetDisplayRate(), GetTickResolution(), bAutoFlipTangents);
+	const FAvaEaseCurveTangents AverageTangents = KeyCache.AverageTangents(GetDisplayRate(), GetTickResolution(), ToolSettings->GetAutoFlipTangents());
 
 	SetEaseCurveTangents(AverageTangents, EAvaEaseCurveToolOperation::InOut, /*bInBroadcastUpdate=*/true, false);
 
@@ -487,17 +490,15 @@ void FAvaEaseCurveTool::OpenToolSettings() const
 {
 	if (ISettingsModule* const SettingsModule = FModuleManager::GetModulePtr<ISettingsModule>(TEXT("Settings")))
 	{
-		const UAvaEaseCurveToolSettings* const Settings = GetDefault<UAvaEaseCurveToolSettings>();
-		check(IsValid(Settings));
-		SettingsModule->ShowViewer(Settings->GetContainerName(), Settings->GetCategoryName(), Settings->GetSectionName());
+		SettingsModule->ShowViewer(ToolSettings->GetContainerName(), ToolSettings->GetCategoryName(), ToolSettings->GetSectionName());
 	}
 }
 
 FFrameRate FAvaEaseCurveTool::GetTickResolution() const
 {
-	if (const TSharedPtr<FAvaSequencer> AvaSequencer = AvaSequencerWeak.Pin())
+	if (const TSharedPtr<ISequencer> Sequencer = SequencerWeak.Pin())
 	{
-		return AvaSequencer->GetSequencer()->GetFocusedTickResolution();
+		return Sequencer->GetFocusedTickResolution();
 	}
 
 	return FFrameRate();
@@ -505,9 +506,9 @@ FFrameRate FAvaEaseCurveTool::GetTickResolution() const
 
 FFrameRate FAvaEaseCurveTool::GetDisplayRate() const
 {
-	if (const TSharedPtr<FAvaSequencer> AvaSequencer = AvaSequencerWeak.Pin())
+	if (const TSharedPtr<ISequencer> Sequencer = SequencerWeak.Pin())
 	{
-		return AvaSequencer->GetSequencer()->GetFocusedDisplayRate();
+		return Sequencer->GetFocusedDisplayRate();
 	}
 
 	// Fallback to using config display rate if tool is being used outside sequencer
@@ -540,4 +541,165 @@ bool FAvaEaseCurveTool::AreAllEaseCurves()
 	return KeyCache.AreAllEaseCurves();
 }
 
+void FAvaEaseCurveTool::ResetToDefaultPresets()
+{
+	const FText MessageBoxTitle = LOCTEXT("ResetToDefaultPresets", "Reset To Default Presets");
+	const EAppReturnType::Type Response = FMessageDialog::Open(EAppMsgType::YesNoCancel
+		, LOCTEXT("ConfirmResetToDefaultPresets", "Are you sure you want to reset to default presets?\n\n"
+			"*CAUTION* All directories and files inside '[Project]/Config/EaseCurves' will be lost!")
+		, MessageBoxTitle);
+	if (Response == EAppReturnType::Yes)
+	{
+		UAvaEaseCurveSubsystem::Get().ResetToDefaultPresets(false);
+	}
+}
+
+void FAvaEaseCurveTool::ApplyTangents()
+{
+	SetEaseCurveTangents(GetEaseCurveTangents(), GetToolOperation(), /*bInBroadcastUpdate=*/true, /*bInSetSequencerTangents=*/true);
+}
+
+void FAvaEaseCurveTool::ZoomToFit() const
+{
+	if (ToolWidget.IsValid())
+	{
+		ToolWidget->ZoomToFit();
+	}
+}
+
+TSharedPtr<FUICommandList> FAvaEaseCurveTool::GetCommandList() const
+{
+	return CommandList;
+}
+
+void FAvaEaseCurveTool::BindCommands()
+{
+	const FAvaEaseCurveToolCommands& EaseCurveToolCommands = FAvaEaseCurveToolCommands::Get();
+
+	CommandList = MakeShared<FUICommandList>();
+
+	CommandList->MapAction(FGenericCommands::Get().Undo
+		, FExecuteAction::CreateSP(this, &FAvaEaseCurveTool::UndoAction));
+	
+	CommandList->MapAction(FGenericCommands::Get().Redo
+		, FExecuteAction::CreateSP(this, &FAvaEaseCurveTool::RedoAction));
+
+	CommandList->MapAction(EaseCurveToolCommands.OpenToolSettings
+		, FExecuteAction::CreateSP(this, &FAvaEaseCurveTool::OpenToolSettings));
+
+	CommandList->MapAction(EaseCurveToolCommands.ResetToDefaultPresets
+		, FExecuteAction::CreateSP(this, &FAvaEaseCurveTool::ResetToDefaultPresets));
+
+	CommandList->MapAction(EaseCurveToolCommands.Refresh
+		, FExecuteAction::CreateSP(this, &FAvaEaseCurveTool::UpdateEaseCurveFromSequencerKeySelections));
+	
+	CommandList->MapAction(EaseCurveToolCommands.Apply
+		, FExecuteAction::CreateSP(this, &FAvaEaseCurveTool::ApplyTangents));
+
+	CommandList->MapAction(EaseCurveToolCommands.ZoomToFit
+		, FExecuteAction::CreateSP(this, &FAvaEaseCurveTool::ZoomToFit));
+
+	CommandList->MapAction(EaseCurveToolCommands.ToggleGridSnap
+		, FExecuteAction::CreateUObject(ToolSettings, &UAvaEaseCurveToolSettings::ToggleGridSnap)
+		, FCanExecuteAction()
+		, FIsActionChecked::CreateUObject(ToolSettings, &UAvaEaseCurveToolSettings::GetGridSnap));
+
+	CommandList->MapAction(EaseCurveToolCommands.ToggleAutoFlipTangents
+		, FExecuteAction::CreateUObject(ToolSettings, &UAvaEaseCurveToolSettings::ToggleAutoFlipTangents)
+		, FCanExecuteAction()
+		, FIsActionChecked::CreateUObject(ToolSettings, &UAvaEaseCurveToolSettings::GetAutoFlipTangents));
+
+	CommandList->MapAction(EaseCurveToolCommands.ToggleAutoZoomToFit
+		, FExecuteAction::CreateUObject(ToolSettings, &UAvaEaseCurveToolSettings::ToggleAutoZoomToFit)
+		, FCanExecuteAction()
+		, FIsActionChecked::CreateUObject(ToolSettings, &UAvaEaseCurveToolSettings::GetAutoZoomToFit));
+
+	CommandList->MapAction(EaseCurveToolCommands.SetOperationToEaseOut
+		, FExecuteAction::CreateSP(this, &FAvaEaseCurveTool::SetToolOperation, EAvaEaseCurveToolOperation::Out)
+		, FCanExecuteAction()
+		, FIsActionChecked::CreateSP(this, &FAvaEaseCurveTool::IsToolOperation, EAvaEaseCurveToolOperation::Out));
+
+	CommandList->MapAction(EaseCurveToolCommands.SetOperationToEaseInOut
+		, FExecuteAction::CreateSP(this, &FAvaEaseCurveTool::SetToolOperation, EAvaEaseCurveToolOperation::InOut)
+		, FCanExecuteAction()
+		, FIsActionChecked::CreateSP(this, &FAvaEaseCurveTool::IsToolOperation, EAvaEaseCurveToolOperation::InOut));
+
+	CommandList->MapAction(EaseCurveToolCommands.SetOperationToEaseIn
+		, FExecuteAction::CreateSP(this, &FAvaEaseCurveTool::SetToolOperation, EAvaEaseCurveToolOperation::In)
+		, FCanExecuteAction()
+		, FIsActionChecked::CreateSP(this, &FAvaEaseCurveTool::IsToolOperation, EAvaEaseCurveToolOperation::In));
+
+	CommandList->MapAction(EaseCurveToolCommands.ResetTangents
+		, FExecuteAction::CreateSP(this, &FAvaEaseCurveTool::ResetEaseCurveTangents, EAvaEaseCurveToolOperation::InOut));
+
+	CommandList->MapAction(EaseCurveToolCommands.ResetStartTangent
+		, FExecuteAction::CreateSP(this, &FAvaEaseCurveTool::ResetEaseCurveTangents, EAvaEaseCurveToolOperation::Out));
+
+	CommandList->MapAction(EaseCurveToolCommands.ResetEndTangent
+		, FExecuteAction::CreateSP(this, &FAvaEaseCurveTool::ResetEaseCurveTangents, EAvaEaseCurveToolOperation::In));
+
+	CommandList->MapAction(EaseCurveToolCommands.FlattenTangents
+		, FExecuteAction::CreateSP(this, &FAvaEaseCurveTool::FlattenOrStraightenTangents, EAvaEaseCurveToolOperation::InOut, true));
+
+	CommandList->MapAction(EaseCurveToolCommands.FlattenStartTangent
+		, FExecuteAction::CreateSP(this, &FAvaEaseCurveTool::FlattenOrStraightenTangents, EAvaEaseCurveToolOperation::Out, true));
+
+	CommandList->MapAction(EaseCurveToolCommands.FlattenEndTangent
+		, FExecuteAction::CreateSP(this, &FAvaEaseCurveTool::FlattenOrStraightenTangents, EAvaEaseCurveToolOperation::In, true));
+
+	CommandList->MapAction(EaseCurveToolCommands.StraightenTangents
+		, FExecuteAction::CreateSP(this, &FAvaEaseCurveTool::FlattenOrStraightenTangents, EAvaEaseCurveToolOperation::InOut, false));
+
+	CommandList->MapAction(EaseCurveToolCommands.StraightenStartTangent
+		, FExecuteAction::CreateSP(this, &FAvaEaseCurveTool::FlattenOrStraightenTangents, EAvaEaseCurveToolOperation::Out, false));
+
+	CommandList->MapAction(EaseCurveToolCommands.StraightenEndTangent
+		, FExecuteAction::CreateSP(this, &FAvaEaseCurveTool::FlattenOrStraightenTangents, EAvaEaseCurveToolOperation::In, false));
+
+	CommandList->MapAction(EaseCurveToolCommands.CopyTangents
+		, FExecuteAction::CreateSP(this, &FAvaEaseCurveTool::CopyTangentsToClipboard)
+		, FCanExecuteAction::CreateSP(this, &FAvaEaseCurveTool::CanCopyTangentsToClipboard));
+
+	CommandList->MapAction(EaseCurveToolCommands.PasteTangents
+		, FExecuteAction::CreateSP(this, &FAvaEaseCurveTool::PasteTangentsFromClipboard)
+		, FCanExecuteAction::CreateSP(this, &FAvaEaseCurveTool::CanPasteTangentsFromClipboard));
+
+	CommandList->MapAction(EaseCurveToolCommands.CreateExternalCurveAsset
+		, FExecuteAction::CreateSPLambda(this, [this]()
+			{
+				CreateCurveAsset();
+			})
+		, FCanExecuteAction());
+
+	CommandList->MapAction(EaseCurveToolCommands.SetKeyInterpConstant,
+		FExecuteAction::CreateSP(this, &FAvaEaseCurveTool::SetKeyInterpMode, RCIM_Constant, RCTM_Auto),
+		FCanExecuteAction(),
+		FIsActionChecked::CreateSP(this, &FAvaEaseCurveTool::IsKeyInterpMode, RCIM_Constant, RCTM_Auto));
+
+	CommandList->MapAction(EaseCurveToolCommands.SetKeyInterpLinear,
+		FExecuteAction::CreateSP(this, &FAvaEaseCurveTool::SetKeyInterpMode, RCIM_Linear, RCTM_Auto),
+		FCanExecuteAction(),
+		FIsActionChecked::CreateSP(this, &FAvaEaseCurveTool::IsKeyInterpMode, RCIM_Linear, RCTM_Auto));
+
+	CommandList->MapAction(EaseCurveToolCommands.SetKeyInterpCubicAuto,
+		FExecuteAction::CreateSP(this, &FAvaEaseCurveTool::SetKeyInterpMode, RCIM_Cubic, RCTM_Auto),
+		FCanExecuteAction(),
+		FIsActionChecked::CreateSP(this, &FAvaEaseCurveTool::IsKeyInterpMode, RCIM_Cubic, RCTM_Auto));
+
+	CommandList->MapAction(EaseCurveToolCommands.SetKeyInterpCubicSmartAuto,
+		FExecuteAction::CreateSP(this, &FAvaEaseCurveTool::SetKeyInterpMode, RCIM_Cubic, RCTM_SmartAuto),
+		FCanExecuteAction(),
+		FIsActionChecked::CreateSP(this, &FAvaEaseCurveTool::IsKeyInterpMode, RCIM_Cubic, RCTM_SmartAuto));
+
+	CommandList->MapAction(EaseCurveToolCommands.SetKeyInterpCubicUser,
+		FExecuteAction::CreateSP(this, &FAvaEaseCurveTool::SetKeyInterpMode, RCIM_Cubic, RCTM_User),
+		FCanExecuteAction(),
+		FIsActionChecked::CreateSP(this, &FAvaEaseCurveTool::IsKeyInterpMode, RCIM_Cubic, RCTM_User));
+
+	CommandList->MapAction(EaseCurveToolCommands.SetKeyInterpCubicBreak,
+		FExecuteAction::CreateSP(this, &FAvaEaseCurveTool::SetKeyInterpMode, RCIM_Cubic, RCTM_Break),
+		FCanExecuteAction(),
+		FIsActionChecked::CreateSP(this, &FAvaEaseCurveTool::IsKeyInterpMode, RCIM_Cubic, RCTM_Break));
+}
+
 #undef LOCTEXT_NAMESPACE
diff --git a/Engine/Plugins/Experimental/Avalanche/Source/AvalancheSequencer/Private/EaseCurveTool/AvaEaseCurveTool.h b/Engine/Plugins/Experimental/Avalanche/Source/AvalancheSequencer/Private/EaseCurveTool/AvaEaseCurveTool.h
--- a/Engine/Plugins/Experimental/Avalanche/Source/AvalancheSequencer/Private/EaseCurveTool/AvaEaseCurveTool.h
+++ b/Engine/Plugins/Experimental/Avalanche/Source/AvalancheSequencer/Private/EaseCurveTool/AvaEaseCurveTool.h
@@ -14,15 +14,15 @@
 #include "Widgets/DeclarativeSyntaxSupport.h"
 #include "Widgets/SCompoundWidget.h"
 
-class FAvaSequencer;
 class FCurveEditor;
-class IAvaSequencer;
+class FUICommandList;
 class ISequencer;
 class SAvaEaseCurveTool;
 class SCurveEditor;
 class SCurveEditorPanel;
 class SEditableTextBox;
 class UAvaEaseCurve;
+class UAvaEaseCurveToolSettings;
 class UCurveFloat;
 class UMovieSceneSection;
 struct FAvaEaseCurveTangents;
@@ -63,7 +63,7 @@ public:
 	/** Returns true if the clipboard paste data contains tangent information. */
 	static bool TangentsFromClipboardPaste(FAvaEaseCurveTangents& OutTangents);
 
-	FAvaEaseCurveTool(const TSharedRef<FAvaSequencer>& InSequencer);
+	FAvaEaseCurveTool(const TSharedRef<ISequencer>& InSequencer);
 
 	virtual ~FAvaEaseCurveTool() override;
 
@@ -96,6 +96,7 @@ public:
 
 	void SetSequencerKeySelectionTangents(const FAvaEaseCurveTangents& InTangents, const EAvaEaseCurveToolOperation InOperation = EAvaEaseCurveToolOperation::InOut);
 
+	bool CanApplyQuickEaseToSequencerKeySelections() const;
 	void ApplyQuickEaseToSequencerKeySelections(const EAvaEaseCurveToolOperation InOperation = EAvaEaseCurveToolOperation::InOut);
 
 	/** Updates the ease curve graph view based on the active sequencer key selection. */
@@ -138,17 +139,29 @@ public:
 	virtual void PostRedo(bool bInSuccess) override;
 	//~ End FEditorUndoClient
 
+	TSharedPtr<FUICommandList> GetCommandList() const;
+
 protected:
+	void ResetToDefaultPresets();
+	void ApplyTangents();
+	void ZoomToFit() const;
+
+	void BindCommands();
+
 	/** 
 	 * Sets the internal ease curve tangents and optionally broadcasts a change event for the curve object.
 	 * Changing the internal ease curve tangents will be directly reflected in the ease curve editor graph.
 	 */
 	void SetEaseCurveTangents_Internal(const FAvaEaseCurveTangents& InTangents, const EAvaEaseCurveToolOperation InOperation, const bool bInBroadcastUpdate) const;
 
-	TWeakPtr<FAvaSequencer> AvaSequencerWeak;
+	TWeakPtr<ISequencer> SequencerWeak;
+
+	TSharedPtr<FUICommandList> CommandList;
 
 	TObjectPtr<UAvaEaseCurve> EaseCurve;
 
+	TObjectPtr<UAvaEaseCurveToolSettings> ToolSettings;
+
 	EAvaEaseCurveToolMode ToolMode = EAvaEaseCurveToolMode::DualKeyEdit;
 	EAvaEaseCurveToolOperation OperationMode = EAvaEaseCurveToolOperation::InOut;
 
diff --git a/Engine/Plugins/Experimental/Avalanche/Source/AvalancheSequencer/Private/EaseCurveTool/AvaEaseCurveToolCommands.cpp b/Engine/Plugins/Experimental/Avalanche/Source/AvalancheSequencer/Private/EaseCurveTool/AvaEaseCurveToolCommands.cpp
--- a/Engine/Plugins/Experimental/Avalanche/Source/AvalancheSequencer/Private/EaseCurveTool/AvaEaseCurveToolCommands.cpp
+++ b/Engine/Plugins/Experimental/Avalanche/Source/AvalancheSequencer/Private/EaseCurveTool/AvaEaseCurveToolCommands.cpp
@@ -223,7 +223,7 @@ void FAvaEaseCurveToolCommands::RegisterCommands()
 		, QuickEase
 		, TEXT("QuickEase")
 		, LOCTEXT("QuickEase_Label", "Quick Ease")
-		, LOCTEXT("QuickEase_ToolTip", "Apply the quick ease preset")
+		, LOCTEXT("QuickEase_ToolTip", "Apply the quick ease preset. This will be disabled if there is no valid quick ease preset setup.")
 		, FSlateIcon(FAppStyle::Get().GetStyleSetName(), TEXT("Icons.BulletPoint"))
 		, EUserInterfaceActionType::Button
 		, FInputChord(EKeys::F8));
@@ -232,7 +232,7 @@ void FAvaEaseCurveToolCommands::RegisterCommands()
 		, QuickEaseIn
 		, TEXT("QuickEaseIn")
 		, LOCTEXT("QuickEaseIn_Label", "Quick Ease In")
-		, LOCTEXT("QuickEaseIn_ToolTip", "Apply the quick ease in preset")
+		, LOCTEXT("QuickEaseIn_ToolTip", "Apply the quick ease in preset. This will be disabled if there is no valid quick ease preset setup.")
 		, FSlateIcon(FAppStyle::Get().GetStyleSetName(), TEXT("Icons.ChevronRight"))
 		, EUserInterfaceActionType::Button
 		, FInputChord(EKeys::F8, EModifierKey::Shift));
@@ -241,7 +241,7 @@ void FAvaEaseCurveToolCommands::RegisterCommands()
 		, QuickEaseOut
 		, TEXT("QuickEaseOut")
 		, LOCTEXT("QuickEaseOut_Label", "Quick Ease Out")
-		, LOCTEXT("QuickEaseOut_ToolTip", "Apply the quick ease out preset")
+		, LOCTEXT("QuickEaseOut_ToolTip", "Apply the quick ease out preset. This will be disabled if there is no valid quick ease preset setup.")
 		, FSlateIcon(FAppStyle::Get().GetStyleSetName(), TEXT("Icons.ChevronLeft"))
 		, EUserInterfaceActionType::Button
 		, FInputChord(EKeys::F8, EModifierKey::Command | EModifierKey::Shift));
diff --git a/Engine/Plugins/Experimental/Avalanche/Source/AvalancheSequencer/Private/EaseCurveTool/Widgets/SAvaEaseCurveTangents.cpp b/Engine/Plugins/Experimental/Avalanche/Source/AvalancheSequencer/Private/EaseCurveTool/Widgets/SAvaEaseCurveTangents.cpp
--- a/Engine/Plugins/Experimental/Avalanche/Source/AvalancheSequencer/Private/EaseCurveTool/Widgets/SAvaEaseCurveTangents.cpp
+++ b/Engine/Plugins/Experimental/Avalanche/Source/AvalancheSequencer/Private/EaseCurveTool/Widgets/SAvaEaseCurveTangents.cpp
@@ -8,7 +8,7 @@
 #include "Widgets/SBoxPanel.h"
 #include "Widgets/Text/STextBlock.h"
 
-#define LOCTEXT_NAMESPACE "SAvaEaseCurveTool"
+#define LOCTEXT_NAMESPACE "SAvaEaseCurveTangents"
 
 void SAvaEaseCurveTangents::Construct(const FArguments& InArgs)
 {
diff --git a/Engine/Plugins/Experimental/Avalanche/Source/AvalancheSequencer/Private/EaseCurveTool/Widgets/SAvaEaseCurveTool.cpp b/Engine/Plugins/Experimental/Avalanche/Source/AvalancheSequencer/Private/EaseCurveTool/Widgets/SAvaEaseCurveTool.cpp
--- a/Engine/Plugins/Experimental/Avalanche/Source/AvalancheSequencer/Private/EaseCurveTool/Widgets/SAvaEaseCurveTool.cpp
+++ b/Engine/Plugins/Experimental/Avalanche/Source/AvalancheSequencer/Private/EaseCurveTool/Widgets/SAvaEaseCurveTool.cpp
@@ -5,7 +5,6 @@
 #include "Curves/RichCurve.h"
 #include "EaseCurveTool/AvaEaseCurvePreset.h"
 #include "EaseCurveTool/AvaEaseCurveTool.h"
-#include "EaseCurveTool/AvaEaseCurveToolCommands.h"
 #include "EaseCurveTool/AvaEaseCurveToolSettings.h"
 #include "EaseCurveTool/AvaEaseCurveSubsystem.h"
 #include "EaseCurveTool/Widgets/SAvaEaseCurveEditor.h"
@@ -14,21 +13,10 @@
 #include "Editor.h"
 #include "EngineAnalytics.h"
 #include "SAvaEaseCurveTangents.h"
-#include "Framework/Commands/GenericCommands.h"
-#include "Framework/Commands/UIAction.h"
 #include "Framework/Commands/UICommandList.h"
 #include "Framework/MultiBox/MultiBoxBuilder.h"
 #include "Math/UnrealMathUtility.h"
-#include "Misc/MessageDialog.h"
-#include "Styling/AppStyle.h"
-#include "Styling/StyleColors.h"
-#include "ToolMenu.h"
-#include "ToolMenus.h"
 #include "Widgets/Images/SImage.h"
-#include "Widgets/Input/SButton.h"
-#include "Widgets/Input/SComboBox.h"
-#include "Widgets/Layout/SScrollBox.h"
-#include "Widgets/SNullWidget.h"
 
 #define LOCTEXT_NAMESPACE "SAvaEaseCurveTool"
 
@@ -37,9 +25,7 @@ void SAvaEaseCurveTool::Construct(const FArguments& InArgs, const TSharedRef<FAv
 	ToolMode = InArgs._ToolMode;
 	ToolOperation = InArgs._ToolOperation;
 
-	EaseCurveTool = InEaseCurveTool;
-
-	BindCommands();
+	EaseCurveToolWeak = InEaseCurveTool;
 
 	ChildSlot
 		[
@@ -52,11 +38,7 @@ void SAvaEaseCurveTool::Construct(const FArguments& InArgs, const TSharedRef<FAv
 				SAssignNew(CurvePresetWidget, SAvaEaseCurvePreset)
 				.OnPresetChanged(this, &SAvaEaseCurveTool::OnPresetChanged)
 				.OnQuickPresetChanged(this, &SAvaEaseCurveTool::OnQuickPresetChanged)
-				.OnGetNewPresetTangents_Lambda([this](FAvaEaseCurveTangents& OutTangents) -> bool
-					{
-						OutTangents = EaseCurveTool->GetEaseCurveTangents();
-						return true;
-					})
+				.OnGetNewPresetTangents(this, &SAvaEaseCurveTool::OnGetNewPresetTangents)
 			]
 			+ SVerticalBox::Slot()
 			.AutoHeight()
@@ -76,14 +58,8 @@ void SAvaEaseCurveTool::Construct(const FArguments& InArgs, const TSharedRef<FAv
 				.OnStartWeightChanged(this, &SAvaEaseCurveTool::OnStartTangentWeightSpinBoxChanged)
 				.OnEndTangentChanged(this, &SAvaEaseCurveTool::OnEndTangentSpinBoxChanged)
 				.OnEndWeightChanged(this, &SAvaEaseCurveTool::OnEndTangentWeightSpinBoxChanged)
-				.OnBeginSliderMovement_Lambda([this]()
-					{
-						EaseCurveTool->BeginTransaction(LOCTEXT("SliderDragStartLabel", "Ease Curve Slider Drag"));
-					})
-				.OnEndSliderMovement_Lambda([this](const float InNewValue)
-					{
-						EaseCurveTool->EndTransaction();
-					})
+				.OnBeginSliderMovement(this, &SAvaEaseCurveTool::OnBeginSliderMovement)
+				.OnEndSliderMovement(this, &SAvaEaseCurveTool::OnEndSliderMovement)
 			]
 		];
 
@@ -97,38 +73,28 @@ void SAvaEaseCurveTool::Construct(const FArguments& InArgs, const TSharedRef<FAv
 
 TSharedRef<SWidget> SAvaEaseCurveTool::ConstructCurveEditorPanel()
 {
+	const TSharedPtr<FAvaEaseCurveTool> EaseCurveTool = EaseCurveToolWeak.Pin();
+	check(EaseCurveTool.IsValid());
+
 	CurrentGraphSize = GetDefault<UAvaEaseCurveToolSettings>()->GetGraphSize();
 	
-	ContextMenu = MakeShared<FAvaEaseCurveToolContextMenu>(CommandList
+	ContextMenu = MakeShared<FAvaEaseCurveToolContextMenu>(EaseCurveTool->GetCommandList()
 		, FAvaEaseCurveToolOnGraphSizeChanged::CreateLambda([this](const int32 InNewSize)
 			{
 				CurrentGraphSize = InNewSize;
 			}));
 
-	const TSharedRef<FAvaEaseCurveTool> EaseCurveToolRef = EaseCurveTool.ToSharedRef();
-	
 	return SNew(SBorder)
 		[
 			SNew(SOverlay)
 			+ SOverlay::Slot()
 			[
 				SAssignNew(CurveEaseEditorWidget, SAvaEaseCurveEditor, EaseCurveTool->GetToolCurve())
-				.DisplayRate(EaseCurveToolRef, &FAvaEaseCurveTool::GetDisplayRate)
-				.Operation(EaseCurveToolRef, &FAvaEaseCurveTool::GetToolOperation)
-				.DesiredSize_Lambda([this]() -> FVector2D
-					{
-						return FVector2D(CurrentGraphSize);
-					})
-				// These two are lambda because the functions aren't const and the functions aren't const
-				// because there is no const ForEachEaseableKey implementation
-				.ShowEqualValueKeyError_Lambda([this]()
-					{
-						return !EaseCurveTool->HasCachedKeysToEase();
-					})
-				.IsEaseCurveSelection_Lambda([this]()
-					{
-						return EaseCurveTool->AreAllEaseCurves();
-					})
+				.DisplayRate(this, &SAvaEaseCurveTool::GetDisplayRate)
+				.Operation(this, &SAvaEaseCurveTool::GetToolOperation)
+				.DesiredSize(this, &SAvaEaseCurveTool::GetEditorSize)
+				.ShowEqualValueKeyError(this, &SAvaEaseCurveTool::ShowEqualValueKeyError)
+				.IsEaseCurveSelection(this, &SAvaEaseCurveTool::IsEaseCurveSelection)
 				.OnTangentsChanged(this, &SAvaEaseCurveTool::HandleEditorTangentsChanged)
 				.GridSnap_UObject(GetDefault<UAvaEaseCurveToolSettings>(), &UAvaEaseCurveToolSettings::GetGridSnap)
 				.GridSize_UObject(GetDefault<UAvaEaseCurveToolSettings>(), &UAvaEaseCurveToolSettings::GetGridSize)
@@ -180,6 +146,31 @@ TSharedRef<SWidget> SAvaEaseCurveTool::ConstructCurveEditorPanel()
 		];
 }
 
+FFrameRate SAvaEaseCurveTool::GetDisplayRate() const
+{
+	return EaseCurveToolWeak.IsValid() ? EaseCurveToolWeak.Pin()->GetDisplayRate() : FFrameRate();
+}
+
+EAvaEaseCurveToolOperation SAvaEaseCurveTool::GetToolOperation() const
+{
+	return EaseCurveToolWeak.IsValid() ? EaseCurveToolWeak.Pin()->GetToolOperation() : EAvaEaseCurveToolOperation::InOut;
+}
+
+FVector2D SAvaEaseCurveTool::GetEditorSize() const
+{
+	return FVector2D(CurrentGraphSize);
+}
+
+bool SAvaEaseCurveTool::ShowEqualValueKeyError() const
+{
+	return !EaseCurveToolWeak.Pin()->HasCachedKeysToEase();
+}
+
+bool SAvaEaseCurveTool::IsEaseCurveSelection() const
+{
+	return EaseCurveToolWeak.Pin()->AreAllEaseCurves();
+}
+
 void SAvaEaseCurveTool::HandleEditorTangentsChanged(const FAvaEaseCurveTangents& InTangents) const
 {
 	SetTangents(InTangents, ToolOperation.Get(), true, true, true);
@@ -187,11 +178,23 @@ void SAvaEaseCurveTool::HandleEditorTangentsChanged(const FAvaEaseCurveTangents&
 
 void SAvaEaseCurveTool::OnEditorDragStart() const
 {
+	const TSharedPtr<FAvaEaseCurveTool> EaseCurveTool = EaseCurveToolWeak.Pin();
+	if (!EaseCurveTool.IsValid())
+	{
+		return;
+	}
+
 	EaseCurveTool->BeginTransaction(LOCTEXT("EditorDragStartLabel", "Ease Curve Graph Drag"));
 }
 
 void SAvaEaseCurveTool::OnEditorDragEnd() const
 {
+	const TSharedPtr<FAvaEaseCurveTool> EaseCurveTool = EaseCurveToolWeak.Pin();
+	if (!EaseCurveTool.IsValid())
+	{
+		return;
+	}
+
 	EaseCurveTool->EndTransaction();
 
 	if (!EaseCurveTool->HasCachedKeysToEase())
@@ -209,9 +212,12 @@ void SAvaEaseCurveTool::SetTangents(const FAvaEaseCurveTangents& InTangents, EAv
 	}
 
 	// To change the graph UI tangents, we need to change the ease curve object tangents and the graph will reflect.
-	if (bInSetEaseCurve && EaseCurveTool.IsValid())
+	if (bInSetEaseCurve)
 	{
-		EaseCurveTool->SetEaseCurveTangents(InTangents, InOperation, bInBroadcastUpdate, bInSetSequencerTangents);
+        if (const TSharedPtr<FAvaEaseCurveTool> EaseCurveTool = EaseCurveToolWeak.Pin())
+        {
+        	EaseCurveTool->SetEaseCurveTangents(InTangents, InOperation, bInBroadcastUpdate, bInSetSequencerTangents);
+        }
 	}
 
 	if (GetDefault<UAvaEaseCurveToolSettings>()->GetAutoZoomToFit())
@@ -222,11 +228,21 @@ void SAvaEaseCurveTool::SetTangents(const FAvaEaseCurveTangents& InTangents, EAv
 
 FAvaEaseCurveTangents SAvaEaseCurveTool::GetTangents() const
 {
-	return EaseCurveTool->GetEaseCurveTangents();
+	if (const TSharedPtr<FAvaEaseCurveTool> EaseCurveTool = EaseCurveToolWeak.Pin())
+	{
+		return EaseCurveTool->GetEaseCurveTangents();
+	}
+	return FAvaEaseCurveTangents();
 }
 
 void SAvaEaseCurveTool::OnStartTangentSpinBoxChanged(const double InNewValue) const
 {
+	const TSharedPtr<FAvaEaseCurveTool> EaseCurveTool = EaseCurveToolWeak.Pin();
+	if (!EaseCurveTool.IsValid())
+	{
+		return;
+	}
+
 	FAvaEaseCurveTangents NewTangents = EaseCurveTool->GetEaseCurveTangents();
 	NewTangents.Start = InNewValue;
 	SetTangents(NewTangents, ToolOperation.Get(), true, true, true);
@@ -234,6 +250,12 @@ void SAvaEaseCurveTool::OnStartTangentSpinBoxChanged(const double InNewValue) co
 
 void SAvaEaseCurveTool::OnStartTangentWeightSpinBoxChanged(const double InNewValue) const
 {
+	const TSharedPtr<FAvaEaseCurveTool> EaseCurveTool = EaseCurveToolWeak.Pin();
+	if (!EaseCurveTool.IsValid())
+	{
+		return;
+	}
+
 	FAvaEaseCurveTangents NewTangents = EaseCurveTool->GetEaseCurveTangents();
 	NewTangents.StartWeight = InNewValue;
 	SetTangents(NewTangents, ToolOperation.Get(), true, true, true);
@@ -241,6 +263,12 @@ void SAvaEaseCurveTool::OnStartTangentWeightSpinBoxChanged(const double InNewVal
 
 void SAvaEaseCurveTool::OnEndTangentSpinBoxChanged(const double InNewValue) const
 {
+	const TSharedPtr<FAvaEaseCurveTool> EaseCurveTool = EaseCurveToolWeak.Pin();
+	if (!EaseCurveTool.IsValid())
+	{
+		return;
+	}
+
 	FAvaEaseCurveTangents NewTangents = EaseCurveTool->GetEaseCurveTangents();
 	NewTangents.End = InNewValue;
 	SetTangents(NewTangents, ToolOperation.Get(), true, true, true);
@@ -248,13 +276,47 @@ void SAvaEaseCurveTool::OnEndTangentSpinBoxChanged(const double InNewValue) cons
 
 void SAvaEaseCurveTool::OnEndTangentWeightSpinBoxChanged(const double InNewValue) const
 {
+	const TSharedPtr<FAvaEaseCurveTool> EaseCurveTool = EaseCurveToolWeak.Pin();
+	if (!EaseCurveTool.IsValid())
+	{
+		return;
+	}
+
 	FAvaEaseCurveTangents NewTangents = EaseCurveTool->GetEaseCurveTangents();
 	NewTangents.EndWeight = InNewValue;
 	SetTangents(NewTangents, ToolOperation.Get(), true, true, true);
 }
 
+void SAvaEaseCurveTool::OnBeginSliderMovement()
+{
+	const TSharedPtr<FAvaEaseCurveTool> EaseCurveTool = EaseCurveToolWeak.Pin();
+    if (!EaseCurveTool.IsValid())
+    {
+    	return;
+    }
+
+	EaseCurveTool->BeginTransaction(LOCTEXT("SliderDragStartLabel", "Ease Curve Slider Drag"));
+}
+
+void SAvaEaseCurveTool::OnEndSliderMovement(const double InNewValue)
+{
+	const TSharedPtr<FAvaEaseCurveTool> EaseCurveTool = EaseCurveToolWeak.Pin();
+	if (!EaseCurveTool.IsValid())
+	{
+		return;
+	}
+
+	EaseCurveTool->EndTransaction();
+}
+
 void SAvaEaseCurveTool::OnPresetChanged(const TSharedPtr<FAvaEaseCurvePreset>& InPreset) const
 {
+	const TSharedPtr<FAvaEaseCurveTool> EaseCurveTool = EaseCurveToolWeak.Pin();
+	if (!EaseCurveTool.IsValid())
+	{
+		return;
+	}
+
 	if (!EaseCurveTool->HasCachedKeysToEase())
 	{
 		ResetTangentsAndNotify();
@@ -289,138 +351,16 @@ void SAvaEaseCurveTool::OnQuickPresetChanged(const TSharedPtr<FAvaEaseCurvePrese
 	FSlateApplication::Get().SetAllUserFocus(CurveEaseEditorWidget);
 }
 
-void SAvaEaseCurveTool::BindCommands()
+bool SAvaEaseCurveTool::OnGetNewPresetTangents(FAvaEaseCurveTangents& OutTangents) const
 {
-	const FAvaEaseCurveToolCommands& EaseCurveToolCommands = FAvaEaseCurveToolCommands::Get();
-
-	const TSharedRef<FAvaEaseCurveTool> EaseCurveToolRef = EaseCurveTool.ToSharedRef();
-
-	CommandList = MakeShared<FUICommandList>();
-
-	CommandList->MapAction(FGenericCommands::Get().Undo, FExecuteAction::CreateSP(EaseCurveToolRef, &FAvaEaseCurveTool::UndoAction));
-	
-	CommandList->MapAction(FGenericCommands::Get().Redo, FExecuteAction::CreateSP(EaseCurveToolRef, &FAvaEaseCurveTool::RedoAction));
-
-	CommandList->MapAction(EaseCurveToolCommands.OpenToolSettings, FExecuteAction::CreateSP(EaseCurveToolRef, &FAvaEaseCurveTool::OpenToolSettings));
-
-	CommandList->MapAction(EaseCurveToolCommands.ResetToDefaultPresets, FExecuteAction::CreateSP(this, &SAvaEaseCurveTool::ResetToDefaultPresets));
-
-	CommandList->MapAction(EaseCurveToolCommands.Refresh, FExecuteAction::CreateSP(EaseCurveToolRef, &FAvaEaseCurveTool::UpdateEaseCurveFromSequencerKeySelections));
-	
-	CommandList->MapAction(EaseCurveToolCommands.Apply, FExecuteAction::CreateSP(this, &SAvaEaseCurveTool::ApplyTangents));
-
-	CommandList->MapAction(EaseCurveToolCommands.ZoomToFit, FExecuteAction::CreateSP(this, &SAvaEaseCurveTool::ZoomToFit));
-
-	CommandList->MapAction(EaseCurveToolCommands.ToggleGridSnap
-		, FExecuteAction::CreateUObject(GetMutableDefault<UAvaEaseCurveToolSettings>(), &UAvaEaseCurveToolSettings::ToggleGridSnap)
-		, FCanExecuteAction()
-		, FIsActionChecked::CreateUObject(GetDefault<UAvaEaseCurveToolSettings>(), &UAvaEaseCurveToolSettings::GetGridSnap));
-
-	CommandList->MapAction(EaseCurveToolCommands.ToggleAutoFlipTangents
-		, FExecuteAction::CreateUObject(GetMutableDefault<UAvaEaseCurveToolSettings>(), &UAvaEaseCurveToolSettings::ToggleAutoFlipTangents)
-		, FCanExecuteAction()
-		, FIsActionChecked::CreateUObject(GetDefault<UAvaEaseCurveToolSettings>(), &UAvaEaseCurveToolSettings::GetAutoFlipTangents));
-
-	CommandList->MapAction(EaseCurveToolCommands.ToggleAutoZoomToFit
-		, FExecuteAction::CreateUObject(GetMutableDefault<UAvaEaseCurveToolSettings>(), &UAvaEaseCurveToolSettings::ToggleAutoZoomToFit)
-		, FCanExecuteAction()
-		, FIsActionChecked::CreateUObject(GetDefault<UAvaEaseCurveToolSettings>(), &UAvaEaseCurveToolSettings::GetAutoZoomToFit));
-
-	CommandList->MapAction(EaseCurveToolCommands.SetOperationToEaseOut
-		, FExecuteAction::CreateSP(EaseCurveToolRef, &FAvaEaseCurveTool::SetToolOperation, EAvaEaseCurveToolOperation::Out)
-		, FCanExecuteAction()
-		, FIsActionChecked::CreateSP(EaseCurveToolRef, &FAvaEaseCurveTool::IsToolOperation, EAvaEaseCurveToolOperation::Out));
-
-	CommandList->MapAction(EaseCurveToolCommands.SetOperationToEaseInOut
-		, FExecuteAction::CreateSP(EaseCurveToolRef, &FAvaEaseCurveTool::SetToolOperation, EAvaEaseCurveToolOperation::InOut)
-		, FCanExecuteAction()
-		, FIsActionChecked::CreateSP(EaseCurveToolRef, &FAvaEaseCurveTool::IsToolOperation, EAvaEaseCurveToolOperation::InOut));
-
-	CommandList->MapAction(EaseCurveToolCommands.SetOperationToEaseIn
-		, FExecuteAction::CreateSP(EaseCurveToolRef, &FAvaEaseCurveTool::SetToolOperation, EAvaEaseCurveToolOperation::In)
-		, FCanExecuteAction()
-		, FIsActionChecked::CreateSP(EaseCurveToolRef, &FAvaEaseCurveTool::IsToolOperation, EAvaEaseCurveToolOperation::In));
-
-	CommandList->MapAction(EaseCurveToolCommands.ResetTangents
-		, FExecuteAction::CreateSP(EaseCurveToolRef, &FAvaEaseCurveTool::ResetEaseCurveTangents, EAvaEaseCurveToolOperation::InOut));
-
-	CommandList->MapAction(EaseCurveToolCommands.ResetStartTangent
-		, FExecuteAction::CreateSP(EaseCurveToolRef, &FAvaEaseCurveTool::ResetEaseCurveTangents, EAvaEaseCurveToolOperation::Out));
-
-	CommandList->MapAction(EaseCurveToolCommands.ResetEndTangent
-		, FExecuteAction::CreateSP(EaseCurveToolRef, &FAvaEaseCurveTool::ResetEaseCurveTangents, EAvaEaseCurveToolOperation::In));
-
-	CommandList->MapAction(EaseCurveToolCommands.FlattenTangents
-		, FExecuteAction::CreateSP(EaseCurveToolRef, &FAvaEaseCurveTool::FlattenOrStraightenTangents, EAvaEaseCurveToolOperation::InOut, true));
-
-	CommandList->MapAction(EaseCurveToolCommands.FlattenStartTangent
-		, FExecuteAction::CreateSP(EaseCurveToolRef, &FAvaEaseCurveTool::FlattenOrStraightenTangents, EAvaEaseCurveToolOperation::Out, true));
-
-	CommandList->MapAction(EaseCurveToolCommands.FlattenEndTangent
-		, FExecuteAction::CreateSP(EaseCurveToolRef, &FAvaEaseCurveTool::FlattenOrStraightenTangents, EAvaEaseCurveToolOperation::In, true));
-
-	CommandList->MapAction(EaseCurveToolCommands.StraightenTangents
-		, FExecuteAction::CreateSP(EaseCurveToolRef, &FAvaEaseCurveTool::FlattenOrStraightenTangents, EAvaEaseCurveToolOperation::InOut, false));
-
-	CommandList->MapAction(EaseCurveToolCommands.StraightenStartTangent
-		, FExecuteAction::CreateSP(EaseCurveToolRef, &FAvaEaseCurveTool::FlattenOrStraightenTangents, EAvaEaseCurveToolOperation::Out, false));
-
-	CommandList->MapAction(EaseCurveToolCommands.StraightenEndTangent
-		, FExecuteAction::CreateSP(EaseCurveToolRef, &FAvaEaseCurveTool::FlattenOrStraightenTangents, EAvaEaseCurveToolOperation::In, false));
-
-	CommandList->MapAction(EaseCurveToolCommands.CopyTangents
-		, FExecuteAction::CreateSP(EaseCurveToolRef, &FAvaEaseCurveTool::CopyTangentsToClipboard)
-		, FCanExecuteAction::CreateSP(EaseCurveToolRef, &FAvaEaseCurveTool::CanCopyTangentsToClipboard));
-
-	CommandList->MapAction(EaseCurveToolCommands.PasteTangents
-		, FExecuteAction::CreateSP(EaseCurveToolRef, &FAvaEaseCurveTool::PasteTangentsFromClipboard)
-		, FCanExecuteAction::CreateSP(EaseCurveToolRef, &FAvaEaseCurveTool::CanPasteTangentsFromClipboard));
+	const TSharedPtr<FAvaEaseCurveTool> EaseCurveTool = EaseCurveToolWeak.Pin();
+	if (!EaseCurveTool.IsValid())
+	{
+		return false;
+	}
 
-	CommandList->MapAction(EaseCurveToolCommands.CreateExternalCurveAsset
-		, FExecuteAction::CreateLambda([this]()
-			{
-				EaseCurveTool->CreateCurveAsset();
-			})
-		, FCanExecuteAction());
-
-	CommandList->MapAction(EaseCurveToolCommands.SetKeyInterpConstant,
-		FExecuteAction::CreateSP(EaseCurveToolRef, &FAvaEaseCurveTool::SetKeyInterpMode, RCIM_Constant, RCTM_Auto),
-		FCanExecuteAction(),
-		FIsActionChecked::CreateSP(EaseCurveToolRef, &FAvaEaseCurveTool::IsKeyInterpMode, RCIM_Constant, RCTM_Auto));
-
-	CommandList->MapAction(EaseCurveToolCommands.SetKeyInterpLinear,
-		FExecuteAction::CreateSP(EaseCurveToolRef, &FAvaEaseCurveTool::SetKeyInterpMode, RCIM_Linear, RCTM_Auto),
-		FCanExecuteAction(),
-		FIsActionChecked::CreateSP(EaseCurveToolRef, &FAvaEaseCurveTool::IsKeyInterpMode, RCIM_Linear, RCTM_Auto));
-
-	CommandList->MapAction(EaseCurveToolCommands.SetKeyInterpCubicAuto,
-		FExecuteAction::CreateSP(EaseCurveToolRef, &FAvaEaseCurveTool::SetKeyInterpMode, RCIM_Cubic, RCTM_Auto),
-		FCanExecuteAction(),
-		FIsActionChecked::CreateSP(EaseCurveToolRef, &FAvaEaseCurveTool::IsKeyInterpMode, RCIM_Cubic, RCTM_Auto));
-
-	CommandList->MapAction(EaseCurveToolCommands.SetKeyInterpCubicSmartAuto,
-		FExecuteAction::CreateSP(EaseCurveToolRef, &FAvaEaseCurveTool::SetKeyInterpMode, RCIM_Cubic, RCTM_SmartAuto),
-		FCanExecuteAction(),
-		FIsActionChecked::CreateSP(EaseCurveToolRef, &FAvaEaseCurveTool::IsKeyInterpMode, RCIM_Cubic, RCTM_SmartAuto));
-
-	CommandList->MapAction(EaseCurveToolCommands.SetKeyInterpCubicUser,
-		FExecuteAction::CreateSP(EaseCurveToolRef, &FAvaEaseCurveTool::SetKeyInterpMode, RCIM_Cubic, RCTM_User),
-		FCanExecuteAction(),
-		FIsActionChecked::CreateSP(EaseCurveToolRef, &FAvaEaseCurveTool::IsKeyInterpMode, RCIM_Cubic, RCTM_User));
-
-	CommandList->MapAction(EaseCurveToolCommands.SetKeyInterpCubicBreak,
-		FExecuteAction::CreateSP(EaseCurveToolRef, &FAvaEaseCurveTool::SetKeyInterpMode, RCIM_Cubic, RCTM_Break),
-		FCanExecuteAction(),
-		FIsActionChecked::CreateSP(EaseCurveToolRef, &FAvaEaseCurveTool::IsKeyInterpMode, RCIM_Cubic, RCTM_Break));
-
-	CommandList->MapAction(EaseCurveToolCommands.QuickEase
-		, FExecuteAction::CreateSP(EaseCurveToolRef, &FAvaEaseCurveTool::ApplyQuickEaseToSequencerKeySelections, EAvaEaseCurveToolOperation::InOut));
-
-	CommandList->MapAction(EaseCurveToolCommands.QuickEaseIn
-		, FExecuteAction::CreateSP(EaseCurveToolRef, &FAvaEaseCurveTool::ApplyQuickEaseToSequencerKeySelections, EAvaEaseCurveToolOperation::In));
-
-	CommandList->MapAction(EaseCurveToolCommands.QuickEaseOut
-		, FExecuteAction::CreateSP(EaseCurveToolRef, &FAvaEaseCurveTool::ApplyQuickEaseToSequencerKeySelections, EAvaEaseCurveToolOperation::Out));
+	OutTangents = EaseCurveTool->GetEaseCurveTangents();
+	return true;
 }
 
 void SAvaEaseCurveTool::UndoAction()
@@ -458,10 +398,18 @@ FKeyHandle SAvaEaseCurveTool::GetSelectedKeyHandle() const
 
 FReply SAvaEaseCurveTool::OnKeyDown(const FGeometry& MyGeometry, const FKeyEvent& InKeyEvent)
 {
-	if (CommandList->ProcessCommandBindings(InKeyEvent))
+	const TSharedPtr<FAvaEaseCurveTool> EaseCurveTool = EaseCurveToolWeak.Pin();
+	if (!EaseCurveTool.IsValid())
+	{
+		return FReply::Unhandled();
+	}
+
+	const TSharedPtr<FUICommandList> CommandList = EaseCurveTool->GetCommandList();
+	if (CommandList.IsValid() && CommandList->ProcessCommandBindings(InKeyEvent))
 	{
 		return FReply::Handled();
 	}
+
 	return FReply::Unhandled();
 }
 
@@ -493,24 +441,6 @@ FText SAvaEaseCurveTool::GetEndTooltipText() const
 		: LOCTEXT("LeaveTooltipText", "Leave");
 }
 
-void SAvaEaseCurveTool::ResetToDefaultPresets()
-{
-	const FText MessageBoxTitle = LOCTEXT("ResetToDefaultPresets", "Reset To Default Presets");
-	const EAppReturnType::Type Response = FMessageDialog::Open(EAppMsgType::YesNoCancel
-		, LOCTEXT("ConfirmResetToDefaultPresets", "Are you sure you want to reset to default presets?\n\n"
-			"*CAUTION* All directories and files inside '[Project]/Config/EaseCurves' will be lost!")
-		, MessageBoxTitle);
-	if (Response == EAppReturnType::Yes)
-	{
-		UAvaEaseCurveSubsystem::Get().ResetToDefaultPresets(false);
-	}
-}
-
-void SAvaEaseCurveTool::ApplyTangents()
-{
-	EaseCurveTool->SetEaseCurveTangents(EaseCurveTool->GetEaseCurveTangents(), EaseCurveTool->GetToolOperation(), /*bInBroadcastUpdate=*/true, /*bInSetSequencerTangents=*/true);
-}
-
 void SAvaEaseCurveTool::ResetTangentsAndNotify() const
 {
 	CurvePresetWidget->ClearSelection();
diff --git a/Engine/Plugins/Experimental/Avalanche/Source/AvalancheSequencer/Private/EaseCurveTool/Widgets/SAvaEaseCurveTool.h b/Engine/Plugins/Experimental/Avalanche/Source/AvalancheSequencer/Private/EaseCurveTool/Widgets/SAvaEaseCurveTool.h
--- a/Engine/Plugins/Experimental/Avalanche/Source/AvalancheSequencer/Private/EaseCurveTool/Widgets/SAvaEaseCurveTool.h
+++ b/Engine/Plugins/Experimental/Avalanche/Source/AvalancheSequencer/Private/EaseCurveTool/Widgets/SAvaEaseCurveTool.h
@@ -9,15 +9,10 @@
 #include "Widgets/SCompoundWidget.h"
 
 class FAvaEaseCurveToolContextMenu;
-class FCurveEditor;
 class FText;
-class FUICommandList;
 class SAvaEaseCurveEditor;
 class SAvaEaseCurvePreset;
-class UCurveBase;
-class UToolMenu;
 struct FKeyHandle;
-struct FRichCurve;
 
 class SAvaEaseCurveTool
 	: public SCompoundWidget
@@ -36,7 +31,7 @@ public:
 	SLATE_END_ARGS()
 
 	void Construct(const FArguments& InArgs, const TSharedRef<FAvaEaseCurveTool>& InEaseCurveTool);
-	
+
 	void SetTangents(const FAvaEaseCurveTangents& InTangents, EAvaEaseCurveToolOperation InOperation,
 		const bool bInSetEaseCurve, const bool bInBroadcastUpdate, const bool bInSetSequencerTangents) const;
 
@@ -49,17 +44,24 @@ public:
 protected:
 	TSharedRef<SWidget> ConstructCurveEditorPanel();
 
+	FFrameRate GetDisplayRate() const;
+	EAvaEaseCurveToolOperation GetToolOperation() const;
+	FVector2D GetEditorSize() const;
+	bool ShowEqualValueKeyError() const;
+	bool IsEaseCurveSelection() const;
+
 	void HandleEditorTangentsChanged(const FAvaEaseCurveTangents& InTangents) const;
 
 	void OnStartTangentSpinBoxChanged(const double InNewValue) const;
 	void OnStartTangentWeightSpinBoxChanged(const double InNewValue) const;
 	void OnEndTangentSpinBoxChanged(const double InNewValue) const;
 	void OnEndTangentWeightSpinBoxChanged(const double InNewValue) const;
+	void OnBeginSliderMovement();
+	void OnEndSliderMovement(const double InNewValue);
 
 	void OnPresetChanged(const TSharedPtr<FAvaEaseCurvePreset>& InPreset) const;
 	void OnQuickPresetChanged(const TSharedPtr<FAvaEaseCurvePreset>& InPreset) const;
-
-	void BindCommands();
+	bool OnGetNewPresetTangents(FAvaEaseCurveTangents& OutTangents) const;
 
 	void UndoAction();
 	void RedoAction();
@@ -72,19 +74,13 @@ protected:
 	FText GetEndText() const;
 	FText GetEndTooltipText() const;
 
-	void ResetToDefaultPresets();
-
-	void ApplyTangents();
-
 	void ResetTangentsAndNotify() const;
 
 	//~ Begin SWidget
 	virtual FReply OnKeyDown(const FGeometry& MyGeometry, const FKeyEvent& InKeyEvent) override;
 	//~ End SWidget
 
-	TSharedPtr<FUICommandList> CommandList;
-
-	TSharedPtr<FAvaEaseCurveTool> EaseCurveTool;
+	TWeakPtr<FAvaEaseCurveTool> EaseCurveToolWeak;
 
 	TAttribute<EAvaEaseCurveToolMode> ToolMode;
 	TAttribute<EAvaEaseCurveToolOperation> ToolOperation;
