MotionDesign : Shapes
- Fix export to static mesh asset, could not save asset due to material being referenced in motion design level

diff --git a/Engine/Plugins/Experimental/Avalanche/Source/AvalancheShapes/Private/DynamicMeshes/AvaShapeDynMeshBase.cpp b/Engine/Plugins/Experimental/Avalanche/Source/AvalancheShapes/Private/DynamicMeshes/AvaShapeDynMeshBase.cpp
--- a/Engine/Plugins/Experimental/Avalanche/Source/AvalancheShapes/Private/DynamicMeshes/AvaShapeDynMeshBase.cpp
+++ b/Engine/Plugins/Experimental/Avalanche/Source/AvalancheShapes/Private/DynamicMeshes/AvaShapeDynMeshBase.cpp
@@ -2349,7 +2349,9 @@ bool UAvaShapeDynamicMeshBase::ExportToStaticMesh(UStaticMesh* DestinationMesh)
 	{
 		return false;
 	}
+
 	UDynamicMesh* SourceMesh = GetShapeMeshComponent()->GetDynamicMesh();
+
 	// export options
 	FGeometryScriptCopyMeshToAssetOptions AssetOptions;
 	AssetOptions.bReplaceMaterials = true;
@@ -2357,13 +2359,26 @@ bool UAvaShapeDynamicMeshBase::ExportToStaticMesh(UStaticMesh* DestinationMesh)
 	AssetOptions.bEnableRecomputeTangents = false;
 	AssetOptions.bEnableRemoveDegenerates = true;
 	AssetOptions.NewMaterialSlotNames = GetMeshDataNames();
-	AssetOptions.NewMaterials = GetMeshDataMaterials();
+
+	// Only take asset materials instead of dynamic ones
+	UMaterial* DefaultMaterial = LoadObject<UMaterial>(nullptr, TEXT("/Script/Engine.Material'/Engine/EngineMaterials/WorldGridMaterial.WorldGridMaterial'"));
+	for (UMaterialInterface* Material : GetMeshDataMaterials())
+	{
+		if (IsValid(Material) && Material->IsAsset())
+		{
+			AssetOptions.NewMaterials.Add(Material);
+		}
+		else
+		{
+			AssetOptions.NewMaterials.Add(DefaultMaterial);
+		}
+	}
+
 	// LOD options
 	FGeometryScriptMeshWriteLOD TargetLOD;
 	TargetLOD.LODIndex = 0;
 
 	EGeometryScriptOutcomePins OutResult;
-
 	UGeometryScriptLibrary_StaticMeshFunctions::CopyMeshToStaticMesh(SourceMesh, DestinationMesh, AssetOptions, TargetLOD, OutResult);
 	DestinationMesh->GetBodySetup()->AggGeom = CachedComponent->GetBodySetup()->AggGeom;
 	return OutResult == EGeometryScriptOutcomePins::Success;
