MotionDesign : Shapes
- Fix shape scale to zero will stay invisible when scaled back up

diff --git a/Engine/Plugins/Experimental/Avalanche/Source/AvalancheShapes/Private/DynamicMeshes/AvaShape2DDynMeshBase.cpp b/Engine/Plugins/Experimental/Avalanche/Source/AvalancheShapes/Private/DynamicMeshes/AvaShape2DDynMeshBase.cpp
--- a/Engine/Plugins/Experimental/Avalanche/Source/AvalancheShapes/Private/DynamicMeshes/AvaShape2DDynMeshBase.cpp
+++ b/Engine/Plugins/Experimental/Avalanche/Source/AvalancheShapes/Private/DynamicMeshes/AvaShape2DDynMeshBase.cpp
@@ -192,7 +192,7 @@ void UAvaShape2DDynMeshBase::OnSizeChanged()
 
 	PixelSize2D = Size2D;
 
-	if (Size2D == PreEditSize2D)
+	if (Size2D.Equals(PreEditSize2D))
 	{
 		return;
 	}
diff --git a/Engine/Plugins/Experimental/Avalanche/Source/AvalancheShapes/Private/DynamicMeshes/AvaShapeDynMeshBase.cpp b/Engine/Plugins/Experimental/Avalanche/Source/AvalancheShapes/Private/DynamicMeshes/AvaShapeDynMeshBase.cpp
--- a/Engine/Plugins/Experimental/Avalanche/Source/AvalancheShapes/Private/DynamicMeshes/AvaShapeDynMeshBase.cpp
+++ b/Engine/Plugins/Experimental/Avalanche/Source/AvalancheShapes/Private/DynamicMeshes/AvaShapeDynMeshBase.cpp
@@ -211,6 +211,11 @@ void UAvaShapeDynamicMeshBase::PostLoad()
 void UAvaShapeDynamicMeshBase::OnComponentDestroyed(bool bDestroyingHierarchy)
 {
 	Super::OnComponentDestroyed(bDestroyingHierarchy);
+
+	if (UDynamicMeshComponent* MeshComponent = GetShapeMeshComponent())
+	{
+		MeshComponent->TransformUpdated.RemoveAll(this);
+	}
 }
 
 void UAvaShapeDynamicMeshBase::ToggleGizmo_Implementation(
@@ -479,6 +484,16 @@ void UAvaShapeDynamicMeshBase::OnUsePrimaryMaterialEverywhereChanged()
 	}
 }
 
+void UAvaShapeDynamicMeshBase::OnMeshTransformChanged(USceneComponent* InComponent, EUpdateTransformFlags InTransformFlags, ETeleportType InTeleport)
+{
+	const bool bIsZeroScaled = InComponent->GetComponentScale().IsNearlyZero();
+	if (bZeroScaled && !bIsZeroScaled)
+	{
+		OnSizeChanged();
+	}
+	bZeroScaled = bIsZeroScaled;
+}
+
 bool UAvaShapeDynamicMeshBase::SetOverridePrimaryUVParams(int32 MeshIndex, bool bOverride)
 {
 	if (MeshIndex == MESH_INDEX_PRIMARY)
@@ -2211,6 +2226,9 @@ void UAvaShapeDynamicMeshBase::InitializeDynamicMesh()
 
 		CachedMesh = EditMesh;
 	});
+
+	bZeroScaled = MeshComponent->GetComponentScale().IsNearlyZero();
+	MeshComponent->TransformUpdated.AddUObject(this, &UAvaShapeDynamicMeshBase::OnMeshTransformChanged);
 }
 
 TArray<FAvaSnapPoint> UAvaShapeDynamicMeshBase::GetLocalSnapPoints() const
diff --git a/Engine/Plugins/Experimental/Avalanche/Source/AvalancheShapes/Public/DynamicMeshes/AvaShapeDynMeshBase.h b/Engine/Plugins/Experimental/Avalanche/Source/AvalancheShapes/Public/DynamicMeshes/AvaShapeDynMeshBase.h
--- a/Engine/Plugins/Experimental/Avalanche/Source/AvalancheShapes/Public/DynamicMeshes/AvaShapeDynMeshBase.h
+++ b/Engine/Plugins/Experimental/Avalanche/Source/AvalancheShapes/Public/DynamicMeshes/AvaShapeDynMeshBase.h
@@ -399,6 +399,8 @@ private:
 	/** Apply primary material everywhere */
 	void OnUsePrimaryMaterialEverywhereChanged();
 
+	void OnMeshTransformChanged(USceneComponent* InComponent, EUpdateTransformFlags InTransformFlags, ETeleportType InTeleport);
+
 	// used to quickly restore mesh that has been modified after OnMeshUpdateFinished was called (eg extrude, masking)
 	TOptional<UE::Geometry::FDynamicMesh3> CachedMesh;
 
@@ -435,4 +437,7 @@ private:
 	bool bVerticesDirty;
 
 	bool bColorsDirty;
+
+	/** Is the mesh scaled to zero ? */
+	bool bZeroScaled = false;
 };
