MotionDesign : InteractiveTools
- Adding new viewport overlay toolbar that remembers last used command, and support context menu to list commands in the same category (shapes 2d, shapes 3d, actors, cloners, effectors, static meshes, lights), toolbar can be hidden/shown using CTRL+T command, toolbar position can be set to the right/left/top/bottom in editor settings. Toolbar only appears when motion design mode is selected, when switching viewports toolbar will carry on.
- Simplifying motion design interactive tools to auto register themselves and register more than one tool parameters, this allows shapes like rectangle to also register a square variant, and cone to register a cylinder variant, cloner and effector have their own category to auto register all cloner layouts and effector shapes
- Adding cyclorama tool in static meshes (may need to be reimported with a smaller scale)
- Removing Motion Design placement tab, it is now replaced by the viewport overlay toolbar that offers the same functionality
- Resizing icons to have the same 16x16 size in the toolbar

diff --git a/Engine/Plugins/Experimental/Avalanche/Source/AvalancheEditor/Private/AvaEditorCommands.cpp b/Engine/Plugins/Experimental/Avalanche/Source/AvalancheEditor/Private/AvaEditorCommands.cpp
--- a/Engine/Plugins/Experimental/Avalanche/Source/AvalancheEditor/Private/AvaEditorCommands.cpp
+++ b/Engine/Plugins/Experimental/Avalanche/Source/AvalancheEditor/Private/AvaEditorCommands.cpp
@@ -322,6 +322,12 @@ void FAvaEditorCommands::RegisterToolsCommands()
 		, EUserInterfaceActionType::ToggleButton
 		, FInputChord());
 
+	UI_COMMAND(CycloramaTool
+		, "Cyclorama"
+		, "Create a cyclorama actor in the viewport."
+		, EUserInterfaceActionType::ToggleButton
+		, FInputChord());
+
 	FUICommandInfo::MakeCommandInfo(AsShared()
 		, MediaPlateTool
 		, TEXT("MediaPlateTool")
diff --git a/Engine/Plugins/Experimental/Avalanche/Source/AvalancheEditor/Private/AvaEditorCommands.h b/Engine/Plugins/Experimental/Avalanche/Source/AvalancheEditor/Private/AvaEditorCommands.h
--- a/Engine/Plugins/Experimental/Avalanche/Source/AvalancheEditor/Private/AvaEditorCommands.h
+++ b/Engine/Plugins/Experimental/Avalanche/Source/AvalancheEditor/Private/AvaEditorCommands.h
@@ -75,6 +75,7 @@ public:
 	TSharedPtr<FUICommandInfo> CylinderTool;
 	TSharedPtr<FUICommandInfo> ConeTool;
 	TSharedPtr<FUICommandInfo> PlaneTool;
+	TSharedPtr<FUICommandInfo> CycloramaTool;
 
 	TSharedPtr<FUICommandInfo> MediaPlateTool;
 
diff --git a/Engine/Plugins/Experimental/Avalanche/Source/AvalancheEditor/Private/AvaEditorStyle.cpp b/Engine/Plugins/Experimental/Avalanche/Source/AvalancheEditor/Private/AvaEditorStyle.cpp
--- a/Engine/Plugins/Experimental/Avalanche/Source/AvalancheEditor/Private/AvaEditorStyle.cpp
+++ b/Engine/Plugins/Experimental/Avalanche/Source/AvalancheEditor/Private/AvaEditorStyle.cpp
@@ -24,10 +24,8 @@ FAvaEditorStyle::FAvaEditorStyle()
 	: FSlateStyleSet(TEXT("AvaEditor"))
 {
 	const FVector2f Icon16(16.f);
-	const FVector2f Icon20(20.f);
 	const FVector2f Icon22(22.f);
 	const FVector2f Icon25(25.f);
-	const FVector2f Icon32(32.f);
 	const FVector2f Icon40(40.f);
 
 	TSharedPtr<IPlugin> Plugin = IPluginManager::Get().FindPlugin(UE_PLUGIN_NAME);
@@ -160,11 +158,12 @@ FAvaEditorStyle::FAvaEditorStyle()
 	Set("AvaEditor.CameraToolsCategory",     new CORE_IMAGE_BRUSH_SVG("Starship/AssetIcons/CameraActor_16", Icon16));
 	Set("AvaEditor.LightsToolsCategory",     new CORE_IMAGE_BRUSH_SVG("Starship/AssetIcons/PointLight_16", Icon16));
 
-	Set("AvaEditor.CubeTool",     new IMAGE_BRUSH("Icons/ToolboxIcons/cube",     Icon16));
-	Set("AvaEditor.SphereTool",   new IMAGE_BRUSH("Icons/ToolboxIcons/sphere",   Icon16));
-	Set("AvaEditor.CylinderTool", new IMAGE_BRUSH("Icons/ToolboxIcons/cylinder", Icon16));
-	Set("AvaEditor.ConeTool",     new IMAGE_BRUSH("Icons/ToolboxIcons/cone",     Icon20));
-	Set("AvaEditor.PlaneTool",    new IMAGE_BRUSH("Icons/ToolboxIcons/plane",    Icon20));
+	Set("AvaEditor.CubeTool",         new IMAGE_BRUSH("Icons/ToolboxIcons/cube",          Icon16));
+	Set("AvaEditor.SphereTool",       new IMAGE_BRUSH("Icons/ToolboxIcons/sphere",        Icon16));
+	Set("AvaEditor.CylinderTool",     new IMAGE_BRUSH("Icons/ToolboxIcons/cylinder",      Icon16));
+	Set("AvaEditor.ConeTool",         new IMAGE_BRUSH("Icons/ToolboxIcons/cone",          Icon16));
+	Set("AvaEditor.PlaneTool",        new IMAGE_BRUSH("Icons/ToolboxIcons/plane",         Icon16));
+	Set("AvaEditor.CycloramaTool",    new IMAGE_BRUSH("Icons/ToolboxIcons/roundedsquare", Icon16));
 
 	Set("AvaEditor.CameraTool",               new CORE_IMAGE_BRUSH_SVG("Starship/AssetIcons/CameraActor_16",       Icon16));
 	Set("AvaEditor.CineCameraTool",           new CORE_IMAGE_BRUSH_SVG("Starship/AssetIcons/CineCameraActor_16",   Icon16));
diff --git a/Engine/Plugins/Experimental/Avalanche/Source/AvalancheEditor/Private/InteractiveTools/AvaInteractiveToolsExtension.cpp b/Engine/Plugins/Experimental/Avalanche/Source/AvalancheEditor/Private/InteractiveTools/AvaInteractiveToolsExtension.cpp
--- a/Engine/Plugins/Experimental/Avalanche/Source/AvalancheEditor/Private/InteractiveTools/AvaInteractiveToolsExtension.cpp
+++ b/Engine/Plugins/Experimental/Avalanche/Source/AvalancheEditor/Private/InteractiveTools/AvaInteractiveToolsExtension.cpp
@@ -117,6 +117,10 @@ void FAvaInteractiveToolsExtension::RegisterTools(IAvalancheInteractiveToolsModu
 	InModule->RegisterTool(StaticMeshCategory, UAvaInteractiveToolsStaticMeshActorToolBuilder::CreateToolParameters(StaticMeshCategory,
 		EditorCommands.PlaneTool, FString(TEXT("Static Mesh Plane Tool")), Priority += 1000,
 		TEXT("/Script/Engine.StaticMesh'/Engine/BasicShapes/Plane.Plane'")));
+	
+	InModule->RegisterTool(StaticMeshCategory, UAvaInteractiveToolsStaticMeshActorToolBuilder::CreateToolParameters(StaticMeshCategory,
+		EditorCommands.CycloramaTool, FString(TEXT("Static Mesh Cyclorama Tool")), Priority += 1000,
+		TEXT("/Script/Engine.StaticMesh'/Avalanche/Toolkit/SM_Cyc_02.SM_Cyc_02'")));
 
 	Priority = 0;
 
diff --git a/Engine/Plugins/Experimental/Avalanche/Source/AvalancheEffectorsEditor/Private/AvaEffectorsEditorCommands.cpp b/Engine/Plugins/Experimental/Avalanche/Source/AvalancheEffectorsEditor/Private/AvaEffectorsEditorCommands.cpp
--- a/Engine/Plugins/Experimental/Avalanche/Source/AvalancheEffectorsEditor/Private/AvaEffectorsEditorCommands.cpp
+++ b/Engine/Plugins/Experimental/Avalanche/Source/AvalancheEffectorsEditor/Private/AvaEffectorsEditorCommands.cpp
@@ -2,6 +2,9 @@
 
 #include "AvaEffectorsEditorCommands.h"
 #include "AvaEffectorsEditorStyle.h"
+#include "Styling/SlateIconFinder.h"
+#include "Utilities/CEClonerLibrary.h"
+#include "Utilities/CEEffectorLibrary.h"
 
 #define LOCTEXT_NAMESPACE "AvaEffectorsEditorCommands"
 
@@ -17,17 +20,43 @@ FAvaEffectorsEditorCommands::FAvaEffectorsEditorCommands()
 
 void FAvaEffectorsEditorCommands::RegisterCommands()
 {
-	UI_COMMAND(Tool_Actor_Effector
-		, "Effector Actor"
-		, "Create an Effector Actor in the viewport."
-		, EUserInterfaceActionType::ToggleButton
-		, FInputChord());
-
-	UI_COMMAND(Tool_Actor_Cloner
-		, "Cloner Actor"
-		, "Create a Cloner Actor in the viewport."
-		, EUserInterfaceActionType::ToggleButton
-		, FInputChord());
+	TSet<FName> TypeNames;
+	UCEEffectorLibrary::GetEffectorTypeNames(TypeNames);
+
+	for (FName TypeName : TypeNames)
+	{
+		TSharedPtr<FUICommandInfo>& TypeCommand = Tool_Actor_Effectors.FindOrAdd(TypeName);
+
+		FUICommandInfo::MakeCommandInfo(
+			AsShared(),
+			TypeCommand,
+			TypeName,
+			FText::FromName(TypeName),
+			LOCTEXT("EffectorCommandTooltip", "Create a Effector Actor in the viewport."),
+			FSlateIconFinder::FindIcon(TEXT("AvaEffectorsEditor.Tool_Actor_Effector")),
+			EUserInterfaceActionType::ToggleButton,
+			FInputChord()
+		);
+	}
+
+	TSet<FName> LayoutNames;
+	UCEClonerLibrary::GetClonerLayoutNames(LayoutNames);
+
+	for (FName LayoutName : LayoutNames)
+	{
+		TSharedPtr<FUICommandInfo>& LayoutCommand = Tool_Actor_Cloners.FindOrAdd(LayoutName);
+
+		FUICommandInfo::MakeCommandInfo(
+			AsShared(),
+			LayoutCommand,
+			LayoutName,
+			FText::FromName(LayoutName),
+			LOCTEXT("ClonerCommandTooltip", "Create a Cloner Actor in the viewport."),
+			FSlateIconFinder::FindIcon(TEXT("AvaEffectorsEditor.Tool_Actor_Cloner")),
+			EUserInterfaceActionType::ToggleButton,
+			FInputChord()
+		);
+	}
 }
 
 #undef LOCTEXT_NAMESPACE
diff --git a/Engine/Plugins/Experimental/Avalanche/Source/AvalancheEffectorsEditor/Private/AvaEffectorsEditorCommands.h b/Engine/Plugins/Experimental/Avalanche/Source/AvalancheEffectorsEditor/Private/AvaEffectorsEditorCommands.h
--- a/Engine/Plugins/Experimental/Avalanche/Source/AvalancheEffectorsEditor/Private/AvaEffectorsEditorCommands.h
+++ b/Engine/Plugins/Experimental/Avalanche/Source/AvalancheEffectorsEditor/Private/AvaEffectorsEditorCommands.h
@@ -13,6 +13,6 @@ public:
 	virtual void RegisterCommands() override;
 	//~ End TCommands
 
-	TSharedPtr<FUICommandInfo> Tool_Actor_Effector;
-	TSharedPtr<FUICommandInfo> Tool_Actor_Cloner;
+	TMap<FName, TSharedPtr<FUICommandInfo>> Tool_Actor_Cloners;
+	TMap<FName, TSharedPtr<FUICommandInfo>> Tool_Actor_Effectors;
 };
diff --git a/Engine/Plugins/Experimental/Avalanche/Source/AvalancheEffectorsEditor/Private/AvaEffectorsEditorStyle.cpp b/Engine/Plugins/Experimental/Avalanche/Source/AvalancheEffectorsEditor/Private/AvaEffectorsEditorStyle.cpp
--- a/Engine/Plugins/Experimental/Avalanche/Source/AvalancheEffectorsEditor/Private/AvaEffectorsEditorStyle.cpp
+++ b/Engine/Plugins/Experimental/Avalanche/Source/AvalancheEffectorsEditor/Private/AvaEffectorsEditorStyle.cpp
@@ -12,18 +12,15 @@
 FAvaEffectorsEditorStyle::FAvaEffectorsEditorStyle()
 	: FSlateStyleSet(TEXT("AvaEffectorsEditor"))
 {
-	const FVector2f Icon20x20(20.0f, 20.0f);
-	const FVector2f Icon32x32(32.0f, 32.0f);
+	const FVector2f Icon16x16(16.0f, 16.0f);
 	
 	const TSharedPtr<IPlugin> Plugin = IPluginManager::Get().FindPlugin(UE_PLUGIN_NAME);
 	check(Plugin.IsValid());
 	
 	SetContentRoot(Plugin->GetBaseDir() / TEXT("Resources"));
 
-	Set("AvaEffectorsEditor.Tool_Actor_Effector", new IMAGE_BRUSH_SVG("Icons/ToolboxIcons/effector", Icon20x20));
-	Set("AvaEffectorsEditor.Tool_Actor_Cloner", new IMAGE_BRUSH_SVG("Icons/ToolboxIcons/cloner", Icon20x20));
-	Set("Tool_Actor_Effector", new IMAGE_BRUSH_SVG("Icons/ToolboxIcons/effector", Icon20x20));
-	Set("Tool_Actor_Cloner", new IMAGE_BRUSH_SVG("Icons/ToolboxIcons/cloner", Icon20x20));
+	Set("AvaEffectorsEditor.Tool_Actor_Effector", new IMAGE_BRUSH_SVG("Icons/ToolboxIcons/effector", Icon16x16));
+	Set("AvaEffectorsEditor.Tool_Actor_Cloner", new IMAGE_BRUSH_SVG("Icons/ToolboxIcons/cloner", Icon16x16));
 	
 	FSlateStyleRegistry::RegisterSlateStyle(*this);
 }
diff --git a/Engine/Plugins/Experimental/Avalanche/Source/AvalancheEffectorsEditor/Private/AvalancheEffectorsEditorModule.cpp b/Engine/Plugins/Experimental/Avalanche/Source/AvalancheEffectorsEditor/Private/AvalancheEffectorsEditorModule.cpp
--- a/Engine/Plugins/Experimental/Avalanche/Source/AvalancheEffectorsEditor/Private/AvalancheEffectorsEditorModule.cpp
+++ b/Engine/Plugins/Experimental/Avalanche/Source/AvalancheEffectorsEditor/Private/AvalancheEffectorsEditorModule.cpp
@@ -3,13 +3,10 @@
 #include "AvalancheEffectorsEditorModule.h"
 #include "AvaEffectorsEditorCommands.h"
 #include "AvaEffectorsEditorStyle.h"
-#include "AvaInteractiveToolsDelegates.h"
-#include "Cloner/AvaClonerActorTool.h"
 #include "Cloner/AvaClonerActorVis.h"
 #include "Cloner/AvaClonerEditorOutlinerContextMenu.h"
 #include "Cloner/CEClonerComponent.h"
 #include "ComponentVisualizers.h"
-#include "Effector/AvaEffectorActorTool.h"
 #include "Effector/AvaEffectorActorVis.h"
 #include "Effector/AvaEffectorEditorOutlinerContextMenu.h"
 #include "Effector/CEEffectorComponent.h"
@@ -23,10 +20,6 @@ void FAvalancheEffectorsEditorModule::StartupModule()
 {
 	FCoreDelegates::OnPostEngineInit.AddRaw(this, &FAvalancheEffectorsEditorModule::PostEngineInit);
 
-	FAvaEffectorsEditorCommands::Register();
-
-	FAvaInteractiveToolsDelegates::GetRegisterToolsDelegate().AddRaw(this, &FAvalancheEffectorsEditorModule::RegisterTools);
-
 	RegisterOutlinerItems();
 }
 
@@ -36,8 +29,6 @@ void FAvalancheEffectorsEditorModule::ShutdownModule()
 
 	FAvaEffectorsEditorCommands::Unregister();
 
-	FAvaInteractiveToolsDelegates::GetRegisterToolsDelegate().RemoveAll(this);
-
 	UnregisterOutlinerItems();
 }
 
@@ -49,19 +40,6 @@ void FAvalancheEffectorsEditorModule::PostEngineInit()
 	}
 }
 
-void FAvalancheEffectorsEditorModule::RegisterTools(IAvalancheInteractiveToolsModule* InModule)
-{
-	InModule->RegisterTool(
-		IAvalancheInteractiveToolsModule::Get().CategoryNameActor,
-		GetDefault<UAvaEffectorActorTool>()->GetToolParameters()
-	);
-
-	InModule->RegisterTool(
-		IAvalancheInteractiveToolsModule::Get().CategoryNameActor,
-		GetDefault<UAvaClonerActorTool>()->GetToolParameters()
-	);
-}
-
 void FAvalancheEffectorsEditorModule::RegisterComponentVisualizers()
 {
 	IAvalancheComponentVisualizersModule::RegisterComponentVisualizer<UCEEffectorComponent, FAvaEffectorActorVisualizer>(&Visualizers);
diff --git a/Engine/Plugins/Experimental/Avalanche/Source/AvalancheEffectorsEditor/Private/AvalancheEffectorsEditorModule.h b/Engine/Plugins/Experimental/Avalanche/Source/AvalancheEffectorsEditor/Private/AvalancheEffectorsEditorModule.h
--- a/Engine/Plugins/Experimental/Avalanche/Source/AvalancheEffectorsEditor/Private/AvalancheEffectorsEditorModule.h
+++ b/Engine/Plugins/Experimental/Avalanche/Source/AvalancheEffectorsEditor/Private/AvalancheEffectorsEditorModule.h
@@ -21,8 +21,6 @@ public:
 private:
 	void PostEngineInit();
 
-	void RegisterTools(IAvalancheInteractiveToolsModule* InModule);
-
 	void RegisterComponentVisualizers();
 
 	void RegisterOutlinerItems();
diff --git a/Engine/Plugins/Experimental/Avalanche/Source/AvalancheEffectorsEditor/Private/Cloner/AvaClonerActorTool.cpp b/Engine/Plugins/Experimental/Avalanche/Source/AvalancheEffectorsEditor/Private/Cloner/AvaClonerActorTool.cpp
--- a/Engine/Plugins/Experimental/Avalanche/Source/AvalancheEffectorsEditor/Private/Cloner/AvaClonerActorTool.cpp
+++ b/Engine/Plugins/Experimental/Avalanche/Source/AvalancheEffectorsEditor/Private/Cloner/AvaClonerActorTool.cpp
@@ -2,33 +2,56 @@
 
 #include "AvaClonerActorTool.h"
 #include "AvaEffectorsEditorCommands.h"
-#include "AvaInteractiveToolsSettings.h"
 #include "Builders/AvaInteractiveToolsToolBuilder.h"
 #include "Cloner/CEClonerActor.h"
 #include "Cloner/CEClonerActorFactory.h"
+#include "Cloner/CEClonerComponent.h"
 
 UAvaClonerActorTool::UAvaClonerActorTool()
 {
 	ActorClass = ACEClonerActor::StaticClass();
 }
 
-FName UAvaClonerActorTool::GetCategoryName()
+void UAvaClonerActorTool::OnRegisterTool(IAvalancheInteractiveToolsModule* InAITModule)
 {
-	return IAvalancheInteractiveToolsModule::CategoryNameActor;
+	Super::OnRegisterTool(InAITModule);
+
+	// Register commands here, subsystem is init
+	FAvaEffectorsEditorCommands::Register();
+
+	const FString ToolIdentifier = TEXT("Clone Actor Tool ");
+	for (const TPair<FName, TSharedPtr<FUICommandInfo>>& ClonerCommandPair : FAvaEffectorsEditorCommands::Get().Tool_Actor_Cloners)
+	{
+		FAvaInteractiveToolsToolParameters ToolParameters =
+		{
+			ClonerCommandPair.Value,
+			ToolIdentifier + ClonerCommandPair.Key.ToString(),
+			3000,
+			FAvalancheInteractiveToolsCreateBuilder::CreateLambda(
+				[](UEdMode* InEdMode)
+				{
+					return UAvaInteractiveToolsToolBuilder::CreateToolBuilder<UAvaClonerActorTool>(InEdMode);
+				}),
+			ActorClass,
+			CreateActorFactory<UCEClonerActorFactory>(),
+			FText::FromString(ClonerCommandPair.Key.ToString()),
+		};
+
+		InAITModule->RegisterTool(IAvalancheInteractiveToolsModule::CategoryNameCloner, MoveTemp(ToolParameters));
+	}
 }
 
-FAvaInteractiveToolsToolParameters UAvaClonerActorTool::GetToolParameters() const
+void UAvaClonerActorTool::OnActorSpawned(AActor* InActor)
 {
-	return {
-		FAvaEffectorsEditorCommands::Get().Tool_Actor_Cloner,
-		TEXT("Clone Actor Tool"),
-		3000,
-		FAvalancheInteractiveToolsCreateBuilder::CreateLambda(
-			[](UEdMode* InEdMode)
-			{
-				return UAvaInteractiveToolsToolBuilder::CreateToolBuilder<UAvaClonerActorTool>(InEdMode);
-			}),
-		ActorClass,
-		CreateActorFactory<UCEClonerActorFactory>()
-	};
+	Super::OnActorSpawned(InActor);
+
+	if (const ACEClonerActor* ClonerActor = Cast<ACEClonerActor>(InActor))
+	{
+		const FAvaInteractiveToolsToolParameters* ActiveTool = GetActiveToolParameters();
+
+		if (ActiveTool && ActiveTool->Label.IsSet())
+		{
+			ClonerActor->GetClonerComponent()->SetLayoutName(FName(ActiveTool->Label.GetValue().ToString()));
+		}
+	}
 }
diff --git a/Engine/Plugins/Experimental/Avalanche/Source/AvalancheEffectorsEditor/Private/Cloner/AvaClonerActorTool.h b/Engine/Plugins/Experimental/Avalanche/Source/AvalancheEffectorsEditor/Private/Cloner/AvaClonerActorTool.h
--- a/Engine/Plugins/Experimental/Avalanche/Source/AvalancheEffectorsEditor/Private/Cloner/AvaClonerActorTool.h
+++ b/Engine/Plugins/Experimental/Avalanche/Source/AvalancheEffectorsEditor/Private/Cloner/AvaClonerActorTool.h
@@ -14,7 +14,7 @@ public:
 	UAvaClonerActorTool();
 
 	//~ Begin UAvaInteractiveToolsToolBase
-	virtual FName GetCategoryName() override;
-	virtual FAvaInteractiveToolsToolParameters GetToolParameters() const override;
+	virtual void OnRegisterTool(IAvalancheInteractiveToolsModule* InAITModule) override;
+	virtual void OnActorSpawned(AActor* InActor) override;
 	//~ End UAvaInteractiveToolsToolBase
 };
diff --git a/Engine/Plugins/Experimental/Avalanche/Source/AvalancheEffectorsEditor/Private/Effector/AvaEffectorActorTool.cpp b/Engine/Plugins/Experimental/Avalanche/Source/AvalancheEffectorsEditor/Private/Effector/AvaEffectorActorTool.cpp
--- a/Engine/Plugins/Experimental/Avalanche/Source/AvalancheEffectorsEditor/Private/Effector/AvaEffectorActorTool.cpp
+++ b/Engine/Plugins/Experimental/Avalanche/Source/AvalancheEffectorsEditor/Private/Effector/AvaEffectorActorTool.cpp
@@ -2,33 +2,56 @@
 
 #include "AvaEffectorActorTool.h"
 #include "AvaEffectorsEditorCommands.h"
-#include "AvaInteractiveToolsSettings.h"
 #include "Builders/AvaInteractiveToolsToolBuilder.h"
 #include "Effector/CEEffectorActor.h"
 #include "Effector/CEEffectorActorFactory.h"
+#include "Effector/CEEffectorComponent.h"
 
 UAvaEffectorActorTool::UAvaEffectorActorTool()
 {
 	ActorClass = ACEEffectorActor::StaticClass();
 }
 
-FName UAvaEffectorActorTool::GetCategoryName()
+void UAvaEffectorActorTool::OnRegisterTool(IAvalancheInteractiveToolsModule* InAITModule)
 {
-	return IAvalancheInteractiveToolsModule::CategoryNameActor;
+	Super::OnRegisterTool(InAITModule);
+
+	// Register commands here, subsystem is init
+	FAvaEffectorsEditorCommands::Register();
+
+	const FString ToolIdentifier = TEXT("Effector Actor Tool ");
+	for (const TPair<FName, TSharedPtr<FUICommandInfo>>& EffectorCommandPair : FAvaEffectorsEditorCommands::Get().Tool_Actor_Effectors)
+	{
+		FAvaInteractiveToolsToolParameters ToolParameters =
+		{
+			EffectorCommandPair.Value,
+			ToolIdentifier + EffectorCommandPair.Key.ToString(),
+			4000,
+			FAvalancheInteractiveToolsCreateBuilder::CreateLambda(
+				[](UEdMode* InEdMode)
+				{
+					return UAvaInteractiveToolsToolBuilder::CreateToolBuilder<UAvaEffectorActorTool>(InEdMode);
+				}),
+			ActorClass,
+			CreateActorFactory<UCEEffectorActorFactory>(),
+			FText::FromString(EffectorCommandPair.Key.ToString()),
+		};
+
+		InAITModule->RegisterTool(IAvalancheInteractiveToolsModule::CategoryNameEffector, MoveTemp(ToolParameters));
+	}
 }
 
-FAvaInteractiveToolsToolParameters UAvaEffectorActorTool::GetToolParameters() const
+void UAvaEffectorActorTool::OnActorSpawned(AActor* InActor)
 {
-	return {
-		FAvaEffectorsEditorCommands::Get().Tool_Actor_Effector,
-		TEXT("Effector Actor Tool"),
-		4000,
-		FAvalancheInteractiveToolsCreateBuilder::CreateLambda(
-			[](UEdMode* InEdMode)
-			{
-				return UAvaInteractiveToolsToolBuilder::CreateToolBuilder<UAvaEffectorActorTool>(InEdMode);
-			}),
-		ActorClass,
-		CreateActorFactory<UCEEffectorActorFactory>()
-	};
+	Super::OnActorSpawned(InActor);
+
+	if (const ACEEffectorActor* EffectorActor = Cast<ACEEffectorActor>(InActor))
+	{
+		const FAvaInteractiveToolsToolParameters* ActiveTool = GetActiveToolParameters();
+
+		if (ActiveTool && ActiveTool->Label.IsSet())
+		{
+			EffectorActor->GetEffectorComponent()->SetTypeName(FName(ActiveTool->Label.GetValue().ToString()));
+		}
+	}
 }
diff --git a/Engine/Plugins/Experimental/Avalanche/Source/AvalancheEffectorsEditor/Private/Effector/AvaEffectorActorTool.h b/Engine/Plugins/Experimental/Avalanche/Source/AvalancheEffectorsEditor/Private/Effector/AvaEffectorActorTool.h
--- a/Engine/Plugins/Experimental/Avalanche/Source/AvalancheEffectorsEditor/Private/Effector/AvaEffectorActorTool.h
+++ b/Engine/Plugins/Experimental/Avalanche/Source/AvalancheEffectorsEditor/Private/Effector/AvaEffectorActorTool.h
@@ -14,7 +14,7 @@ public:
 	UAvaEffectorActorTool();
 
 	//~ Begin UAvaInteractiveToolsToolBase
-	virtual FName GetCategoryName() override;
-	virtual FAvaInteractiveToolsToolParameters GetToolParameters() const override;
+	virtual void OnRegisterTool(IAvalancheInteractiveToolsModule* InAITModule) override;
+	virtual void OnActorSpawned(AActor* InActor) override;
 	//~ End UAvaInteractiveToolsToolBase
 };
diff --git a/Engine/Plugins/Experimental/Avalanche/Source/AvalancheInteractiveTools/Private/AvaInteractiveToolsCommands.cpp b/Engine/Plugins/Experimental/Avalanche/Source/AvalancheInteractiveTools/Private/AvaInteractiveToolsCommands.cpp
--- a/Engine/Plugins/Experimental/Avalanche/Source/AvalancheInteractiveTools/Private/AvaInteractiveToolsCommands.cpp
+++ b/Engine/Plugins/Experimental/Avalanche/Source/AvalancheInteractiveTools/Private/AvaInteractiveToolsCommands.cpp
@@ -19,10 +19,7 @@ void FAvaInteractiveToolsCommands::RegisterCommands()
 {
 	RegisterEdModeCommands();
 	RegisterCategoryCommands();
-	Register2DCommands();
-	Register3DCommands();
 	RegisterActorCommands();
-	RegisterLayoutCommands();
 }
 
 void FAvaInteractiveToolsCommands::RegisterEdModeCommands()
@@ -33,6 +30,20 @@ void FAvaInteractiveToolsCommands::RegisterEdModeCommands()
 		, EUserInterfaceActionType::Button
 		, FInputChord(EKeys::Escape)
 	);
+
+	UI_COMMAND(ToggleViewportToolbar
+		, "Toggle Toolbar"
+		, "Toggle viewport overlay toolbar visibility"
+		, EUserInterfaceActionType::ToggleButton
+		, FInputChord(EModifierKey::Control, EKeys::T)
+	);
+	
+	UI_COMMAND(OpenViewportToolbarSettings
+		, "Toolbar Settings"
+		, "Opens editor toolbar settings"
+		, EUserInterfaceActionType::Button
+		, FInputChord()
+	);
 }
 
 void FAvaInteractiveToolsCommands::RegisterCategoryCommands()
@@ -54,20 +65,18 @@ void FAvaInteractiveToolsCommands::RegisterCategoryCommands()
 		, "Tools for creating Actors in the viewport."
 		, EUserInterfaceActionType::ToggleButton
 		, FInputChord());
-
-	UI_COMMAND(Category_Layout
-		, "Layout Helpers"
-		, "Tools for creating Layout Helpers in the viewport."
+	
+	UI_COMMAND(Category_Cloner
+		, "Cloners"
+		, "Tools for creating cloners in the viewport."
 		, EUserInterfaceActionType::ToggleButton
 		, FInputChord());
-}
-
-void FAvaInteractiveToolsCommands::Register2DCommands()
-{
-}
 
-void FAvaInteractiveToolsCommands::Register3DCommands()
-{
+	UI_COMMAND(Category_Effector
+		, "Effectors"
+		, "Tools for creating effectors in the viewport."
+		, EUserInterfaceActionType::ToggleButton
+		, FInputChord());
 }
 
 void FAvaInteractiveToolsCommands::RegisterActorCommands()
@@ -85,8 +94,4 @@ void FAvaInteractiveToolsCommands::RegisterActorCommands()
 		, FInputChord());
 }
 
-void FAvaInteractiveToolsCommands::RegisterLayoutCommands()
-{
-}
-
 #undef LOCTEXT_NAMESPACE
diff --git a/Engine/Plugins/Experimental/Avalanche/Source/AvalancheInteractiveTools/Private/AvaInteractiveToolsCommands.h b/Engine/Plugins/Experimental/Avalanche/Source/AvalancheInteractiveTools/Private/AvaInteractiveToolsCommands.h
--- a/Engine/Plugins/Experimental/Avalanche/Source/AvalancheInteractiveTools/Private/AvaInteractiveToolsCommands.h
+++ b/Engine/Plugins/Experimental/Avalanche/Source/AvalancheInteractiveTools/Private/AvaInteractiveToolsCommands.h
@@ -15,19 +15,19 @@ public:
 
 	void RegisterEdModeCommands();
 	void RegisterCategoryCommands();
-	void Register2DCommands();
-	void Register3DCommands();
 	void RegisterActorCommands();
-	void RegisterLayoutCommands();
 
 	// Ed Mode
 	TSharedPtr<FUICommandInfo> CancelActiveTool;
+	TSharedPtr<FUICommandInfo> ToggleViewportToolbar;
+	TSharedPtr<FUICommandInfo> OpenViewportToolbarSettings;
 
 	// Categories
 	TSharedPtr<FUICommandInfo> Category_2D;
 	TSharedPtr<FUICommandInfo> Category_3D;
 	TSharedPtr<FUICommandInfo> Category_Actor;
-	TSharedPtr<FUICommandInfo> Category_Layout;
+	TSharedPtr<FUICommandInfo> Category_Cloner;
+	TSharedPtr<FUICommandInfo> Category_Effector;
 
 	TSharedPtr<FUICommandInfo> Tool_Actor_Null;
 	TSharedPtr<FUICommandInfo> Tool_Actor_Spline;
diff --git a/Engine/Plugins/Experimental/Avalanche/Source/AvalancheInteractiveTools/Private/AvaInteractiveToolsSettings.cpp b/Engine/Plugins/Experimental/Avalanche/Source/AvalancheInteractiveTools/Private/AvaInteractiveToolsSettings.cpp
--- a/Engine/Plugins/Experimental/Avalanche/Source/AvalancheInteractiveTools/Private/AvaInteractiveToolsSettings.cpp
+++ b/Engine/Plugins/Experimental/Avalanche/Source/AvalancheInteractiveTools/Private/AvaInteractiveToolsSettings.cpp
@@ -2,8 +2,24 @@
 
 #include "AvaInteractiveToolsSettings.h"
 
+#include "ISettingsModule.h"
+#include "Modules/ModuleManager.h"
+
+UAvaInteractiveToolsSettings* UAvaInteractiveToolsSettings::Get()
+{
+	return GetMutableDefault<UAvaInteractiveToolsSettings>();
+}
+
 UAvaInteractiveToolsSettings::UAvaInteractiveToolsSettings()
 {
 	CategoryName = TEXT("Motion Design");
 	SectionName = TEXT("Interactive Tools");
 }
+
+void UAvaInteractiveToolsSettings::OpenEditorSettingsWindow() const
+{
+	if (ISettingsModule* SettingsModule = FModuleManager::GetModulePtr<ISettingsModule>(TEXT("Settings")))
+	{
+		SettingsModule->ShowViewer(GetContainerName(), GetCategoryName(), GetSectionName());
+	}
+}
diff --git a/Engine/Plugins/Experimental/Avalanche/Source/AvalancheInteractiveTools/Private/AvaInteractiveToolsStyle.cpp b/Engine/Plugins/Experimental/Avalanche/Source/AvalancheInteractiveTools/Private/AvaInteractiveToolsStyle.cpp
--- a/Engine/Plugins/Experimental/Avalanche/Source/AvalancheInteractiveTools/Private/AvaInteractiveToolsStyle.cpp
+++ b/Engine/Plugins/Experimental/Avalanche/Source/AvalancheInteractiveTools/Private/AvaInteractiveToolsStyle.cpp
@@ -6,14 +6,16 @@
 #include "Interfaces/IPluginManager.h"
 #include "Math/MathFwd.h"
 #include "Misc/Paths.h"
+#include "Styling/AppStyle.h"
+#include "Styling/StyleColors.h"
 #include "Styling/SlateStyleMacros.h"
 #include "Styling/SlateStyleRegistry.h"
+#include "Styling/ToolBarStyle.h"
 
 FAvaInteractiveToolsStyle::FAvaInteractiveToolsStyle()
 	: FSlateStyleSet(TEXT("AvaInteractiveTools"))
 {
 	const FVector2f Icon16x16(16.0f, 16.0f);
-	const FVector2f Icon20x20(20.0f, 20.0f);
 
 	const TSharedPtr<IPlugin> Plugin = IPluginManager::Get().FindPlugin(UE_PLUGIN_NAME);
 	check(Plugin.IsValid());
@@ -21,19 +23,35 @@ FAvaInteractiveToolsStyle::FAvaInteractiveToolsStyle()
 	SetContentRoot(Plugin->GetBaseDir() / TEXT("Resources"));
 	SetCoreContentRoot(FPaths::EngineContentDir() / TEXT("Editor/Slate"));
 
+	FToolBarStyle ToolbarStyle = FAppStyle::Get().GetWidgetStyle<FToolBarStyle>("EditorViewportToolBar");
+	FTextBlockStyle TextStyle = FAppStyle::Get().GetWidgetStyle<FTextBlockStyle>("SmallText");
+
+	ToolbarStyle
+		.SetButtonPadding(FMargin(2.f))
+		.SetAllowWrapButton(true)
+		.SetIconSize(Icon16x16)
+		.SetIconPadding(FMargin(2.f))
+		.SetLabelPadding(FMargin(2.f))
+		.SetLabelStyle(TextStyle)
+		.SetBackground(FSlateColorBrush(FStyleColors::Panel));
+
+	Set("ViewportToolbar", ToolbarStyle);
+	Set("AvaInteractiveTools.ToggleViewportToolbar", new IMAGE_BRUSH_SVG("Icons/EditorIcons/visible", Icon16x16));
+	Set("AvaInteractiveTools.OpenViewportToolbarSettings", new IMAGE_BRUSH_SVG("Icons/EditorIcons/settings", Icon16x16));
+
 	Set("Icons.Toolbox", new IMAGE_BRUSH(TEXT("Icons/ToolboxIcons/toolbox"), Icon16x16));
 
 	// Categories
-	Set("AvaInteractiveTools.Category_2D",     new IMAGE_BRUSH("Icons/ToolboxIcons/rectangle", Icon20x20));
-	Set("AvaInteractiveTools.Category_3D",     new IMAGE_BRUSH("Icons/ToolboxIcons/cube", Icon20x20));
-	Set("AvaInteractiveTools.Category_Actor",  new CORE_IMAGE_BRUSH_SVG("Starship/AssetIcons/Actor_16", Icon16x16));
-	Set("AvaInteractiveTools.Category_Layout", new IMAGE_BRUSH("Icons/ToolboxIcons/layoutgrid", Icon20x20));
+	Set("AvaInteractiveTools.Category_2D",       new IMAGE_BRUSH("Icons/ToolboxIcons/rectangle", Icon16x16));
+	Set("AvaInteractiveTools.Category_3D",       new IMAGE_BRUSH("Icons/ToolboxIcons/cube", Icon16x16));
+	Set("AvaInteractiveTools.Category_Actor",    new CORE_IMAGE_BRUSH(TEXT("Icons/SequencerIcons/icon_Sequencer_Move_24x"), Icon16x16));
+	Set("AvaInteractiveTools.Category_Cloner",   new IMAGE_BRUSH_SVG("Icons/ToolboxIcons/cloner", Icon16x16));
+	Set("AvaInteractiveTools.Category_Effector", new IMAGE_BRUSH_SVG("Icons/ToolboxIcons/effector", Icon16x16));
+	Set("AvaInteractiveTools.Category_Layout",   new IMAGE_BRUSH("Icons/ToolboxIcons/layoutgrid", Icon16x16));
 
 	// Actor Tools
-	Set("AvaInteractiveTools.Tool_Actor_Null", new CORE_IMAGE_BRUSH(TEXT("Icons/SequencerIcons/icon_Sequencer_Move_24x"), Icon16x16));
-	Set("AvaInteractiveTools.Tool_Actor_Spline", new IMAGE_BRUSH_SVG("Icons/ToolboxIcons/Toolbar_Spline", Icon20x20));
-	Set("Tool_Actor_Null", new CORE_IMAGE_BRUSH(TEXT("Icons/SequencerIcons/icon_Sequencer_Move_24x"), Icon20x20));
-	Set("Tool_Actor_Spline", new IMAGE_BRUSH_SVG("Icons/ToolboxIcons/Toolbar_Spline", Icon20x20));
+	Set("AvaInteractiveTools.Tool_Actor_Null",   new CORE_IMAGE_BRUSH(TEXT("Icons/SequencerIcons/icon_Sequencer_Move_24x"), Icon16x16));
+	Set("AvaInteractiveTools.Tool_Actor_Spline", new IMAGE_BRUSH_SVG("Icons/ToolboxIcons/Toolbar_Spline", Icon16x16));
 
 	FSlateStyleRegistry::RegisterSlateStyle(*this);
 }
diff --git a/Engine/Plugins/Experimental/Avalanche/Source/AvalancheInteractiveTools/Private/AvalancheInteractiveToolsModule.cpp b/Engine/Plugins/Experimental/Avalanche/Source/AvalancheInteractiveTools/Private/AvalancheInteractiveToolsModule.cpp
--- a/Engine/Plugins/Experimental/Avalanche/Source/AvalancheInteractiveTools/Private/AvalancheInteractiveToolsModule.cpp
+++ b/Engine/Plugins/Experimental/Avalanche/Source/AvalancheInteractiveTools/Private/AvalancheInteractiveToolsModule.cpp
@@ -3,13 +3,11 @@
 #include "AvalancheInteractiveToolsModule.h"
 #include "AvaInteractiveToolsCommands.h"
 #include "AvaInteractiveToolsDelegates.h"
-#include "AvaInteractiveToolsSettings.h"
 #include "Builders/AvaInteractiveToolsActorToolBuilder.h"
 #include "Interfaces/IPluginManager.h"
 #include "IPlacementModeModule.h"
 #include "Styling/SlateIconFinder.h"
-#include "Tools/AvaInteractiveToolsActorToolNull.h"
-#include "Tools/AvaInteractiveToolsActorToolSpline.h"
+#include "UObject/UObjectIterator.h"
 
 #define LOCTEXT_NAMESPACE "AvalancheInteractiveTools"
 
@@ -23,7 +21,6 @@ namespace UE::AvaInteractiveTools::Private
 void FAvalancheInteractiveToolsModule::StartupModule()
 {
 	FAvaInteractiveToolsCommands::Register();
-	bHasActiveTool = false;
 
 	if (IPluginManager::Get().GetLastCompletedLoadingPhase() >= ELoadingPhase::PostEngineInit)
 	{
@@ -127,9 +124,41 @@ const TArray<FAvaInteractiveToolsToolParameters>* FAvalancheInteractiveToolsModu
 	return Tools.Find(InCategory);
 }
 
+const FAvaInteractiveToolsToolParameters* FAvalancheInteractiveToolsModule::GetTool(const FString& InToolIdentifier)
+{
+	for (const TPair<FName, TArray<FAvaInteractiveToolsToolParameters>>& CategoryPair : Tools)
+	{
+		for (const FAvaInteractiveToolsToolParameters& ToolParams : CategoryPair.Value)
+		{
+			if (ToolParams.ToolIdentifier.Equals(InToolIdentifier))
+			{
+				return &ToolParams;
+			}
+		}
+	}
+
+	return nullptr;
+}
+
+FName FAvalancheInteractiveToolsModule::GetToolCategory(const FString& InToolIdentifier)
+{
+	for (const TPair<FName, TArray<FAvaInteractiveToolsToolParameters>>& CategoryPair : Tools)
+	{
+		for (const FAvaInteractiveToolsToolParameters& ToolParams : CategoryPair.Value)
+		{
+			if (ToolParams.ToolIdentifier.Equals(InToolIdentifier))
+			{
+				return CategoryPair.Key;
+			}
+		}
+	}
+
+	return NAME_None;
+}
+
 bool FAvalancheInteractiveToolsModule::HasActiveTool() const
 {
-	return bHasActiveTool;
+	return ActiveToolIdentifier.IsSet();
 }
 
 void FAvalancheInteractiveToolsModule::AddReferencedObjects(FReferenceCollector& InCollector)
@@ -152,14 +181,24 @@ FString FAvalancheInteractiveToolsModule::GetReferencerName() const
 	return ReferencerName;
 }
 
-void FAvalancheInteractiveToolsModule::OnToolActivated()
+void FAvalancheInteractiveToolsModule::OnToolActivated(const FString& InToolIdentifier)
 {
-	bHasActiveTool = true;
+	ActiveToolIdentifier = InToolIdentifier;
+	OnToolActivationDelegate.Broadcast(InToolIdentifier);
 }
 
 void FAvalancheInteractiveToolsModule::OnToolDeactivated()
 {
-	bHasActiveTool = false;
+	if (HasActiveTool())
+	{
+		OnToolDeactivationDelegate.Broadcast(ActiveToolIdentifier.GetValue());
+		ActiveToolIdentifier.Reset();
+	}
+}
+
+FString FAvalancheInteractiveToolsModule::GetActiveToolIdentifier() const
+{
+	return ActiveToolIdentifier.IsSet() ? ActiveToolIdentifier.GetValue() : FString();
 }
 
 void FAvalancheInteractiveToolsModule::OnPostEngineInit()
@@ -201,21 +240,40 @@ void FAvalancheInteractiveToolsModule::BroadcastRegisterCategories()
 
 void FAvalancheInteractiveToolsModule::RegisterDefaultCategories()
 {
-	RegisterCategory(CategoryName2D,     FAvaInteractiveToolsCommands::Get().Category_2D, 41);
-	RegisterCategory(CategoryName3D,     FAvaInteractiveToolsCommands::Get().Category_3D, 42);
-	RegisterCategory(CategoryNameActor,  FAvaInteractiveToolsCommands::Get().Category_Actor, 43);
+	RegisterCategory(CategoryName2D, FAvaInteractiveToolsCommands::Get().Category_2D, 41);
+	RegisterCategory(CategoryName3D, FAvaInteractiveToolsCommands::Get().Category_3D, 42);
+	RegisterCategory(CategoryNameActor, FAvaInteractiveToolsCommands::Get().Category_Actor, 43);
+	RegisterCategory(CategoryNameCloner, FAvaInteractiveToolsCommands::Get().Category_Cloner, 44);
+	RegisterCategory(CategoryNameEffector, FAvaInteractiveToolsCommands::Get().Category_Effector, 45);
 }
 
-void FAvalancheInteractiveToolsModule::BroadcastRegisterTools()
+void FAvalancheInteractiveToolsModule::RegisterAutoRegisterTools()
 {
-	RegisterDefaultTools();
-	FAvaInteractiveToolsDelegates::GetRegisterToolsDelegate().Broadcast(this);
+	for (const UClass* const Class : TObjectRange<UClass>())
+	{
+		if (!Class->IsChildOf(UAvaInteractiveToolsToolBase::StaticClass())
+			|| Class->HasAnyClassFlags(CLASS_Abstract | CLASS_Deprecated | CLASS_NewerVersionExists))
+		{
+			continue;
+		}
+
+		UAvaInteractiveToolsToolBase* ToolCDO = Class->GetDefaultObject<UAvaInteractiveToolsToolBase>();
+
+		if (!ToolCDO || !ToolCDO->ShouldAutoRegister())
+		{
+			continue;
+		}
+
+		ToolCDO->OnRegisterTool(this);
+
+		UE_LOG(LogAvaInteractiveTools, Log, TEXT("Tool %s auto registered"), *Class->GetName())
+	}
 }
 
-void FAvalancheInteractiveToolsModule::RegisterDefaultTools()
+void FAvalancheInteractiveToolsModule::BroadcastRegisterTools()
 {
-	RegisterTool(CategoryNameActor, GetDefault<UAvaInteractiveToolsActorToolNull>()->GetToolParameters());
-	RegisterTool(CategoryNameActor, GetDefault<UAvaInteractiveToolsActorToolSpline>()->GetToolParameters());
+	RegisterAutoRegisterTools();
+	FAvaInteractiveToolsDelegates::GetRegisterToolsDelegate().Broadcast(this);
 }
 
 void FAvalancheInteractiveToolsModule::OnPlacementCategoryRefreshed(FName InCategory)
@@ -279,9 +337,15 @@ void FAvalancheInteractiveToolsModule::OnPlacementCategoryRefreshed(FName InCate
 		PlaceableItem->DisplayName = Tool.UICommand->GetLabel();
 		PlaceableItem->NativeName = Tool.ToolIdentifier;
 
-		if (FSlateIconFinder::FindIcon(Tool.UICommand->GetCommandName()).IsSet())
+		const FName IconName(Tool.UICommand->GetBindingContext().ToString() + TEXT(".") + Tool.UICommand->GetCommandName().ToString());
+		if (FSlateIconFinder::FindIcon(IconName).IsSet())
+		{
+			PlaceableItem->ClassThumbnailBrushOverride = IconName;
+			PlaceableItem->bAlwaysUseGenericThumbnail = false;
+		}
+		else if (Tool.UICommand->GetIcon().IsSet())
 		{
-			PlaceableItem->ClassThumbnailBrushOverride = Tool.UICommand->GetCommandName();
+			PlaceableItem->ClassThumbnailBrushOverride = Tool.UICommand->GetIcon().GetStyleName();
 			PlaceableItem->bAlwaysUseGenericThumbnail = false;
 		}
 
diff --git a/Engine/Plugins/Experimental/Avalanche/Source/AvalancheInteractiveTools/Private/AvalancheInteractiveToolsModule.h b/Engine/Plugins/Experimental/Avalanche/Source/AvalancheInteractiveTools/Private/AvalancheInteractiveToolsModule.h
--- a/Engine/Plugins/Experimental/Avalanche/Source/AvalancheInteractiveTools/Private/AvalancheInteractiveToolsModule.h
+++ b/Engine/Plugins/Experimental/Avalanche/Source/AvalancheInteractiveTools/Private/AvalancheInteractiveToolsModule.h
@@ -34,6 +34,8 @@ public:
 	virtual void RegisterTool(FName InCategory, FAvaInteractiveToolsToolParameters&& InToolParams) override;
 	virtual const TMap<FName, TSharedPtr<FUICommandInfo>>& GetCategories() override;
 	virtual const TArray<FAvaInteractiveToolsToolParameters>* GetTools(FName InCategory) override;
+	virtual const FAvaInteractiveToolsToolParameters* GetTool(const FString& InToolIdentifier) override;
+	virtual FName GetToolCategory(const FString& InToolIdentifier) override;
 	virtual bool HasActiveTool() const override;
 	//~ End IAvalancheInteractiveToolsModule
 
@@ -42,18 +44,34 @@ public:
 	virtual FString GetReferencerName() const override;
 	// ~End FGCObject Interface
 
-	void OnToolActivated();
+	void OnToolActivated(const FString& InToolIdentifier);
 	void OnToolDeactivated();
+	FString GetActiveToolIdentifier() const;
+
+	virtual FToolEvent::RegistrationType& OnToolActivation() override
+	{
+		return OnToolActivationDelegate;
+	}
+
+	virtual FToolEvent::RegistrationType& OnToolDeactivation() override
+	{
+		return OnToolDeactivationDelegate;
+	}
 
 private:
 	TMap<FName, TSharedPtr<FUICommandInfo>> Categories;
 	TMap<FName, TArray<FAvaInteractiveToolsToolParameters>> Tools;
-	bool bHasActiveTool = false;
+	TOptional<FString> ActiveToolIdentifier;
+	FToolEvent OnToolActivationDelegate;
+	FToolEvent OnToolDeactivationDelegate;
 
 	void OnPostEngineInit();
+
 	void BroadcastRegisterCategories();
 	void RegisterDefaultCategories();
+
+	void RegisterAutoRegisterTools();
 	void BroadcastRegisterTools();
-	void RegisterDefaultTools();
+
 	void OnPlacementCategoryRefreshed(FName InCategory);
 };
diff --git a/Engine/Plugins/Experimental/Avalanche/Source/AvalancheInteractiveTools/Private/EdMode/AvaInteractiveToolsEdMode.cpp b/Engine/Plugins/Experimental/Avalanche/Source/AvalancheInteractiveTools/Private/EdMode/AvaInteractiveToolsEdMode.cpp
--- a/Engine/Plugins/Experimental/Avalanche/Source/AvalancheInteractiveTools/Private/EdMode/AvaInteractiveToolsEdMode.cpp
+++ b/Engine/Plugins/Experimental/Avalanche/Source/AvalancheInteractiveTools/Private/EdMode/AvaInteractiveToolsEdMode.cpp
@@ -4,6 +4,7 @@
 
 #include "AvaInteractiveToolsCommands.h"
 #include "AvaInteractiveToolsEdModeToolkit.h"
+#include "AvaInteractiveToolsSettings.h"
 #include "AvaInteractiveToolsStyle.h"
 #include "ContextObjectStore.h"
 #include "EditorModeManager.h"
@@ -62,6 +63,37 @@ void UAvaInteractiveToolsEdMode::BindCommands()
 		FIsActionButtonVisible(),
 		EUIActionRepeatMode::RepeatDisabled
 	);
+
+	CommandList->MapAction(
+		ToolManagerCommands.ToggleViewportToolbar,
+		FExecuteAction::CreateWeakLambda(this, [this]()
+			{
+				if (Toolkit.IsValid())
+				{
+					FAvaInteractiveToolsEdModeToolkit* AITToolkit = static_cast<FAvaInteractiveToolsEdModeToolkit*>(Toolkit.Get());
+					AITToolkit->SetViewportToolbarVisibility(!AITToolkit->GetViewportToolbarVisibility());
+				}
+			}),
+		FCanExecuteAction(),
+		FGetActionCheckState(),
+		FIsActionButtonVisible(),
+		EUIActionRepeatMode::RepeatDisabled
+	);
+	
+	CommandList->MapAction(
+		ToolManagerCommands.OpenViewportToolbarSettings,
+		FExecuteAction::CreateWeakLambda(this, [this]()
+			{
+				if (const UAvaInteractiveToolsSettings* Settings = UAvaInteractiveToolsSettings::Get())
+				{
+					Settings->OpenEditorSettingsWindow();
+				}
+			}),
+		FCanExecuteAction(),
+		FGetActionCheckState(),
+		FIsActionButtonVisible(),
+		EUIActionRepeatMode::RepeatDisabled
+	);
 }
 
 void UAvaInteractiveToolsEdMode::Enter()
diff --git a/Engine/Plugins/Experimental/Avalanche/Source/AvalancheInteractiveTools/Private/EdMode/AvaInteractiveToolsEdMode.h b/Engine/Plugins/Experimental/Avalanche/Source/AvalancheInteractiveTools/Private/EdMode/AvaInteractiveToolsEdMode.h
--- a/Engine/Plugins/Experimental/Avalanche/Source/AvalancheInteractiveTools/Private/EdMode/AvaInteractiveToolsEdMode.h
+++ b/Engine/Plugins/Experimental/Avalanche/Source/AvalancheInteractiveTools/Private/EdMode/AvaInteractiveToolsEdMode.h
@@ -30,8 +30,8 @@ public:
 	void OnToolPaletteChanged(FName InPaletteName);
 	void OnToolSetup(UInteractiveTool* InTool);
 	void OnToolShutdown(UInteractiveTool* InTool, EToolShutdownType InShutdownType);
-	const FString GetLastActiveTool() const { return LastActiveTool; }
-	double GetLastToolActivateTime() { return LastToolActivateTime; }
+	FString GetLastActiveTool() const { return LastActiveTool; }
+	double GetLastToolActivateTime() const { return LastToolActivateTime; }
 	void OnToolActivateEnd();
 
 protected:
diff --git a/Engine/Plugins/Experimental/Avalanche/Source/AvalancheInteractiveTools/Private/EdMode/AvaInteractiveToolsEdModeToolkit.cpp b/Engine/Plugins/Experimental/Avalanche/Source/AvalancheInteractiveTools/Private/EdMode/AvaInteractiveToolsEdModeToolkit.cpp
--- a/Engine/Plugins/Experimental/Avalanche/Source/AvalancheInteractiveTools/Private/EdMode/AvaInteractiveToolsEdModeToolkit.cpp
+++ b/Engine/Plugins/Experimental/Avalanche/Source/AvalancheInteractiveTools/Private/EdMode/AvaInteractiveToolsEdModeToolkit.cpp
@@ -1,13 +1,16 @@
 // Copyright Epic Games, Inc. All Rights Reserved.
 
 #include "AvaInteractiveToolsEdModeToolkit.h"
+
+#include "AvaInteractiveToolsCommands.h"
 #include "AvaInteractiveToolsEdMode.h"
-#include "AvalancheInteractiveToolsModule.h"
+#include "AvaInteractiveToolsSettings.h"
+#include "AvaInteractiveToolsStyle.h"
 #include "IAvalancheInteractiveToolsModule.h"
-#include "Toolkits/AssetEditorModeUILayer.h"
 #include "Tools/UEdMode.h"
 #include "Widgets/Layout/SBorder.h"
 #include "Widgets/Layout/SBox.h"
+#include "Widgets/SAvaInteractiveToolsToolbarCategoryButton.h"
 
 #define LOCTEXT_NAMESPACE "AvaInteractiveToolsEdModeToolkit"
 
@@ -18,6 +21,15 @@ FAvaInteractiveToolsEdModeToolkit::FAvaInteractiveToolsEdModeToolkit()
 
 FAvaInteractiveToolsEdModeToolkit::~FAvaInteractiveToolsEdModeToolkit()
 {
+	if (ToolkitHost.IsValid())
+	{
+		GetToolkitHost()->RemoveViewportOverlayWidget(ViewportToolbarWidget.ToSharedRef());
+	}
+}
+
+void FAvaInteractiveToolsEdModeToolkit::SetViewportToolbarVisibility(bool bInShow)
+{
+	bViewportToolbarVisible = bInShow;
 }
 
 FName FAvaInteractiveToolsEdModeToolkit::GetToolkitFName() const
@@ -31,56 +43,11 @@ FText FAvaInteractiveToolsEdModeToolkit::GetBaseToolkitName() const
 	return LOCTEXT("DisplayName", "AvaInteractiveToolsEdMode Tool");
 }
 
-TSharedPtr<SWidget> FAvaInteractiveToolsEdModeToolkit::GetInlineContent() const
-{
-	return SNew(SVerticalBox)
-		+ SVerticalBox::Slot()
-		.FillHeight(1.0f)
-		.VAlign(VAlign_Fill)
-		[
-			ToolkitWidget.ToSharedRef()
-		];
-}
-
 void FAvaInteractiveToolsEdModeToolkit::Init(const TSharedPtr<IToolkitHost>& InitToolkitHost, TWeakObjectPtr<UEdMode> InOwningMode)
 {
 	FModeToolkit::Init(InitToolkitHost, InOwningMode);
 
-	RegisterPalettes();
-
-	SAssignNew(ToolkitWidget, SBorder)
-		.HAlign(HAlign_Fill)
-		.Padding(0)
-		.BorderImage(FAppStyle::GetBrush("ToolPanel.GroupBorder"))
-		[
-			ToolkitBuilder->GenerateWidget()->AsShared()
-		];
-}
-
-void FAvaInteractiveToolsEdModeToolkit::GetToolPaletteNames(TArray<FName>& InPaletteName) const
-{
-	const TMap<FName, TSharedPtr<FUICommandInfo>>& Categories = IAvalancheInteractiveToolsModule::Get().GetCategories();
-	Categories.GetKeys(InPaletteName);
-}
-
-FText FAvaInteractiveToolsEdModeToolkit::GetToolPaletteDisplayName(FName InPaletteName) const
-{
-	const TMap<FName, TSharedPtr<FUICommandInfo>>& Categories = IAvalancheInteractiveToolsModule::Get().GetCategories();
-
-	if (const TSharedPtr<FUICommandInfo>* CommandInfo = Categories.Find(InPaletteName))
-	{
-		return CommandInfo->Get()->GetLabel();
-	}
-
-	return FText::FromName(InPaletteName);
-}
-
-void FAvaInteractiveToolsEdModeToolkit::OnToolPaletteChanged(FName InPaletteName)
-{
-	if (UAvaInteractiveToolsEdMode* AvaInteractiveToolsEdMode = Cast<UAvaInteractiveToolsEdMode>(GetScriptableEditorMode().Get()))
-	{
-		AvaInteractiveToolsEdMode->OnToolPaletteChanged(InPaletteName);
-	}
+	MakeViewportOverlayToolbar();
 }
 
 void FAvaInteractiveToolsEdModeToolkit::OnToolStarted(UInteractiveToolManager* InManager, UInteractiveTool* InTool)
@@ -95,106 +62,139 @@ void FAvaInteractiveToolsEdModeToolkit::OnToolEnded(UInteractiveToolManager* InM
 
 void FAvaInteractiveToolsEdModeToolkit::InvokeUI()
 {
-	FModeToolkit::InvokeUI();
-
-	InlineContentHolder->SetContent(GetInlineContent().ToSharedRef());
+	// Nothing
 }
 
 void FAvaInteractiveToolsEdModeToolkit::RequestModeUITabs()
 {
-	if (const TSharedPtr<FAssetEditorModeUILayer> ModeUILayerPtr = ModeUILayer.Pin())
-	{
-		PrimaryTabInfo.OnSpawnTab = FOnSpawnTab::CreateSP(SharedThis(this), &FAvaInteractiveToolsEdModeToolkit::CreatePrimaryModePanel);
-		PrimaryTabInfo.TabLabel = LOCTEXT("MotionDesignToolboxTab", "Motion Design");
-		PrimaryTabInfo.TabTooltip = LOCTEXT("MotionDesignToolboxTabTooltipText", "Opens the Motion Design tab.");
-		ModeUILayerPtr->SetModePanelInfo(UAssetEditorUISubsystem::TopLeftTabID, PrimaryTabInfo);
-
-		if (!HasIntegratedToolPalettes() && !HasToolkitBuilder())
-		{
-			ToolbarInfo.OnSpawnTab = FOnSpawnTab::CreateSP(SharedThis(this), &FAvaInteractiveToolsEdModeToolkit::MakeModeToolbarTab);
-			ToolbarInfo.TabLabel = LOCTEXT("MotionDesignToolbarTab", "Motion Design Toolbar");
-			ToolbarInfo.TabTooltip = LOCTEXT("MotionDesignToolbarTabTooltipText", "Opens the toolbar for the Motion Design toolbox.");
-			ModeUILayerPtr->SetModePanelInfo(UAssetEditorUISubsystem::VerticalToolbarID, ToolbarInfo);
-		}
-	}
+	// Nothing
 }
 
-void FAvaInteractiveToolsEdModeToolkit::RegisterPalettes()
+void FAvaInteractiveToolsEdModeToolkit::MakeViewportOverlayToolbar()
 {
-	UAvaInteractiveToolsEdMode* AvaInteractiveToolsMode = Cast<UAvaInteractiveToolsEdMode>(GetScriptableEditorMode().Get());
+	if (ViewportToolbarWidget.IsValid())
+	{
+		GetToolkitHost()->RemoveViewportOverlayWidget(ViewportToolbarWidget.ToSharedRef());
+		ViewportToolbarWidget.Reset();
+	}
 
-	ModeDetailsView->SetIsPropertyVisibleDelegate(
-		FIsPropertyVisible::CreateLambda(
-			[](const FPropertyAndParent& InPropertyAndParent) -> bool
-			{
-				static const FString Material = FString(TEXT("Material"));
-				static const FName Category = FName(TEXT("Category"));
-				const FString CategoryMeta = InPropertyAndParent.Property.GetMetaData(Category);
-				return CategoryMeta != Material;
-			}
-	));
+	UAvaInteractiveToolsSettings* Settings = UAvaInteractiveToolsSettings::Get();
 
-	ToolkitSections = MakeShared<FToolkitSections>();
-	//ToolkitSections->DetailsView = ModeDetailsView;
+	if (!Settings)
+	{
+		return;
+	}
 
-	ToolkitSections->ToolWarningArea = SNew(STextBlock)
-		.Text(this, &FAvaInteractiveToolsEdModeToolkit::GetToolWarningText)
-		.AutoWrapText(true);
-		
-	ToolkitBuilder = MakeShared<FToolkitBuilder>(
-		AvaInteractiveToolsMode->GetModeInfo().ToolbarCustomizationName,
-		GetToolkitCommands(),
-		ToolkitSections);
+	// When switching editor settings, update widget
+	Settings->OnSettingChanged().RemoveAll(this);
+	Settings->OnSettingChanged().AddSP(this, &FAvaInteractiveToolsEdModeToolkit::OnSettingsChanged);
 
-	const TMap<FName, TSharedPtr<FUICommandInfo>>& Categories = IAvalancheInteractiveToolsModule::Get().GetCategories();
-	TSharedPtr<FUICommandInfo> FirstCategoryCommand = nullptr;
+	if (Settings->ViewportToolbarPosition == EAvaInteractiveToolsViewportToolbarPosition::None)
+	{
+		return;
+	}
+
+	// When switching viewport, overlay is reset
+	GetToolkitHost()->OnActiveViewportChanged().RemoveAll(this);
+	GetToolkitHost()->OnActiveViewportChanged().AddSP(this, &FAvaInteractiveToolsEdModeToolkit::OnViewportChanged);
+
+	const TSharedRef<FUICommandList> ToolkitCommandList = GetToolkitCommands();
 
-	for (const TPair<FName, TSharedPtr<FUICommandInfo>>& Category : Categories)
+	auto SetupToolbar = [&ToolkitCommandList, &Settings](FToolBarBuilder& InToolbarBuilder)
 	{
-		if (Category.Value.IsValid())
+		InToolbarBuilder.SetLabelVisibility(EVisibility::Collapsed);
+		InToolbarBuilder.SetStyle(&FAvaInteractiveToolsStyle::Get(), "ViewportToolbar");
+		for (const TPair<FName, TSharedPtr<FUICommandInfo>>& Category : IAvalancheInteractiveToolsModule::Get().GetCategories())
 		{
-			TArray<TSharedPtr<FUICommandInfo>> CategoryCommands;
-			const TArray<FAvaInteractiveToolsToolParameters>* ToolList = IAvalancheInteractiveToolsModule::Get().GetTools(Category.Key);
-
-			if (ToolList)
+			if (!Category.Value.IsValid())
 			{
-				for (const FAvaInteractiveToolsToolParameters& Tool : *ToolList)
-				{
-					if (Tool.UICommand.IsValid())
-					{
-						CategoryCommands.Add(Tool.UICommand);
-					}
-				}
+				continue;
 			}
 
-			if (CategoryCommands.IsEmpty() == false)
-			{
-				ToolkitBuilder->AddPalette(MakeShared<FToolPalette>(Category.Value.ToSharedRef(), CategoryCommands));
-
-				if (FirstCategoryCommand.IsValid() == false)
-				{
-					FirstCategoryCommand = Category.Value;
-				}
-			}
+			InToolbarBuilder.AddWidget(
+				SNew(SAvaInteractiveToolsToolbarCategoryButton)
+				.CommandList(ToolkitCommandList)
+				.ShowLabel(Settings->bViewportToolbarLabelEnabled)
+				.Category(Category.Key)
+			);
 		}
-	}
 
-	if (FirstCategoryCommand.IsValid())
+		InToolbarBuilder.AddSeparator();
+		InToolbarBuilder.AddToolBarButton(FAvaInteractiveToolsCommands::Get().ToggleViewportToolbar);
+		InToolbarBuilder.AddToolBarButton(FAvaInteractiveToolsCommands::Get().OpenViewportToolbarSettings);
+	};
+
+	TSharedPtr<SWidget> ToolbarWidget;
+	EHorizontalAlignment HAlign;
+	EVerticalAlignment VAlign;
+	if (Settings->ViewportToolbarPosition == EAvaInteractiveToolsViewportToolbarPosition::Bottom
+		|| Settings->ViewportToolbarPosition == EAvaInteractiveToolsViewportToolbarPosition::Top)
 	{
-		ToolkitBuilder->SetActivePaletteOnLoad(FirstCategoryCommand.Get());
+		FSlimHorizontalToolBarBuilder ToolbarBuilder(
+			ToolkitCommandList,
+			FMultiBoxCustomization::None,
+			TSharedPtr<FExtender>(),
+			/** ForceSmallIcon */false);
+
+		SetupToolbar(ToolbarBuilder);
+		ToolbarWidget = ToolbarBuilder.MakeWidget();
+		
+		HAlign = HAlign_Center;
+		VAlign = Settings->ViewportToolbarPosition == EAvaInteractiveToolsViewportToolbarPosition::Bottom ? VAlign_Bottom : VAlign_Top;
 	}
+	else
+	{
+		FVerticalToolBarBuilder ToolbarBuilder(
+			ToolkitCommandList,
+			FMultiBoxCustomization::None,
+			TSharedPtr<FExtender>(),
+			/** ForceSmallIcon */false);
+
+		SetupToolbar(ToolbarBuilder);
+		ToolbarWidget = ToolbarBuilder.MakeWidget();
 
-	ToolkitBuilder->UpdateWidget();
+		HAlign = Settings->ViewportToolbarPosition == EAvaInteractiveToolsViewportToolbarPosition::Left ? HAlign_Left : HAlign_Right;
+		VAlign = VAlign_Center;
+	}
+
+	SAssignNew(ViewportToolbarWidget, SBorder)
+		.HAlign(HAlign)
+		.VAlign(VAlign)
+		.BorderBackgroundColor(FSlateColor(EStyleColor::Panel))
+		.Padding(0.f)
+		[
+			SNew(SBox)
+			.HAlign(HAlign_Center)
+			.VAlign(VAlign_Center)
+			.Padding(5.f)
+			[
+				ToolbarWidget.ToSharedRef()
+			]
+		];
+
+	ViewportToolbarWidget->SetVisibility(TAttribute<EVisibility>::CreateLambda([this]()
+	{
+		// Only show when motion design mode is selected
+		if (UAvaInteractiveToolsEdMode* AvaMode = Cast<UAvaInteractiveToolsEdMode>(GetScriptableEditorMode()))
+		{
+			// to avoid affecting other viewport overlay buttons
+			return bViewportToolbarVisible ? EVisibility::SelfHitTestInvisible : EVisibility::Collapsed;
+		}
+
+		return EVisibility::Collapsed;
+	}));
+
+	GetToolkitHost()->AddViewportOverlayWidget(ViewportToolbarWidget.ToSharedRef());
 }
 
-FText FAvaInteractiveToolsEdModeToolkit::GetToolWarningText() const
+void FAvaInteractiveToolsEdModeToolkit::OnSettingsChanged(UObject* InSettings, FPropertyChangedEvent& InEvent)
 {
-	static const FText ToolActive = LOCTEXT("ActiveToolWarning", "Tool Active.\n\nSelect the tool again to perform the default action (if supported). This will spawn a predefined actor in the middle of the viewport.\n\nRight click or press escape to cancel.");
-	static const FText ToolInactive = LOCTEXT("InactiveToolWarning", "Select a tool to start drawing.\n\nShift click will show tool presets, if available.");
+	MakeViewportOverlayToolbar();
+}
 
-	return FAvalancheInteractiveToolsModule::Get().HasActiveTool()
-		? ToolActive
-		: ToolInactive;
+void FAvaInteractiveToolsEdModeToolkit::OnViewportChanged(TSharedPtr<IAssetViewport> InOldViewport, TSharedPtr<IAssetViewport> InNewViewport)
+{
+	MakeViewportOverlayToolbar();
 }
 
 #undef LOCTEXT_NAMESPACE
diff --git a/Engine/Plugins/Experimental/Avalanche/Source/AvalancheInteractiveTools/Private/EdMode/AvaInteractiveToolsEdModeToolkit.h b/Engine/Plugins/Experimental/Avalanche/Source/AvalancheInteractiveTools/Private/EdMode/AvaInteractiveToolsEdModeToolkit.h
--- a/Engine/Plugins/Experimental/Avalanche/Source/AvalancheInteractiveTools/Private/EdMode/AvaInteractiveToolsEdModeToolkit.h
+++ b/Engine/Plugins/Experimental/Avalanche/Source/AvalancheInteractiveTools/Private/EdMode/AvaInteractiveToolsEdModeToolkit.h
@@ -4,23 +4,29 @@
 
 #include "Toolkits/BaseToolkit.h"
 
+class UObject;
+struct FPropertyChangedEvent;
+
 class FAvaInteractiveToolsEdModeToolkit : public FModeToolkit
 {
 public:
 	FAvaInteractiveToolsEdModeToolkit();
 	virtual ~FAvaInteractiveToolsEdModeToolkit() override;
 
+	void SetViewportToolbarVisibility(bool bInShow);
+	bool GetViewportToolbarVisibility() const
+	{
+		return bViewportToolbarVisible;
+	}
+
+protected:
 	//~ Begin IToolkit
 	virtual FName GetToolkitFName() const override;
 	virtual FText GetBaseToolkitName() const override;
-	virtual TSharedPtr<class SWidget> GetInlineContent() const override;
 	//~ End IToolkit
 
 	//~ Begin FModeToolkit
 	virtual void Init(const TSharedPtr<IToolkitHost>& InitToolkitHost, TWeakObjectPtr<UEdMode> InOwningMode) override;
-	virtual void GetToolPaletteNames(TArray<FName>& InPaletteName) const override;
-	virtual FText GetToolPaletteDisplayName(FName InPaletteName) const override;
-	virtual void OnToolPaletteChanged(FName InPaletteName) override;
 	virtual bool HasIntegratedToolPalettes() const override { return false; }
 	virtual bool HasExclusiveToolPalettes() const override { return false; }
 	virtual void OnToolStarted(UInteractiveToolManager* InManager, UInteractiveTool* InTool) override;
@@ -28,16 +34,17 @@ public:
 	virtual void InvokeUI() override;
 	//~ End FModeToolkit
 
-protected:
 	//~ Begin FModeToolkit
 	virtual void RequestModeUITabs() override;
 	//~ End FModeToolkit
 
 private:
-	TSharedPtr<SWidget> ToolkitWidget;
+	/** Creates an viewport overlay toolbar */
+	void MakeViewportOverlayToolbar();
 
-	/** A utility function to register the tool palettes with the ToolkitBuilder */
-	void RegisterPalettes();
+	void OnSettingsChanged(UObject* InSettings, FPropertyChangedEvent& InEvent);
+	void OnViewportChanged(TSharedPtr<IAssetViewport> InOldViewport, TSharedPtr<IAssetViewport> InNewViewport);
 
-	FText GetToolWarningText() const;
+	TSharedPtr<SWidget> ViewportToolbarWidget;
+	bool bViewportToolbarVisible = true;
 };
diff --git a/Engine/Plugins/Experimental/Avalanche/Source/AvalancheInteractiveTools/Private/Tools/AvaInteractiveToolsActorAreaToolBase.cpp b/Engine/Plugins/Experimental/Avalanche/Source/AvalancheInteractiveTools/Private/Tools/AvaInteractiveToolsActorAreaToolBase.cpp
--- a/Engine/Plugins/Experimental/Avalanche/Source/AvalancheInteractiveTools/Private/Tools/AvaInteractiveToolsActorAreaToolBase.cpp
+++ b/Engine/Plugins/Experimental/Avalanche/Source/AvalancheInteractiveTools/Private/Tools/AvaInteractiveToolsActorAreaToolBase.cpp
@@ -12,7 +12,7 @@ UAvaInteractiveToolsActorAreaToolBase::UAvaInteractiveToolsActorAreaToolBase()
 }
 
 AActor* UAvaInteractiveToolsActorAreaToolBase::SpawnActor(TSubclassOf<AActor> InActorClass, EAvaViewportStatus InViewportStatus,
-	const FVector2f& InViewportPosition, bool bInPreview, FString* InActorLabelOverride) const
+	const FVector2f& InViewportPosition, bool bInPreview, FString* InActorLabelOverride)
 {
 	AActor* NewActor = Super::SpawnActor(InActorClass, InViewportStatus, InViewportPosition, bInPreview, InActorLabelOverride);
 
diff --git a/Engine/Plugins/Experimental/Avalanche/Source/AvalancheInteractiveTools/Private/Tools/AvaInteractiveToolsActorTool.cpp b/Engine/Plugins/Experimental/Avalanche/Source/AvalancheInteractiveTools/Private/Tools/AvaInteractiveToolsActorTool.cpp
--- a/Engine/Plugins/Experimental/Avalanche/Source/AvalancheInteractiveTools/Private/Tools/AvaInteractiveToolsActorTool.cpp
+++ b/Engine/Plugins/Experimental/Avalanche/Source/AvalancheInteractiveTools/Private/Tools/AvaInteractiveToolsActorTool.cpp
@@ -8,15 +8,18 @@ UAvaInteractiveToolsActorTool::UAvaInteractiveToolsActorTool()
 {
 }
 
-FAvaInteractiveToolsToolParameters UAvaInteractiveToolsActorTool::GetToolParameters() const
+void UAvaInteractiveToolsActorTool::OnRegisterTool(IAvalancheInteractiveToolsModule* InAITModule)
 {
-	// This is never really needed, it should work regardless. OOP consistency!
-	return UAvaInteractiveToolsActorToolBuilder::CreateToolParameters(
-		Category,
-		Command,
-		Identifier,
-		Priority,
-		ActorClass,
-		GetClass()
-	);
-}
+	Super::OnRegisterTool(InAITModule);
+
+	FAvaInteractiveToolsToolParameters ToolParameters = UAvaInteractiveToolsActorToolBuilder::CreateToolParameters(
+	    Category,
+	    Command,
+	    Identifier,
+	    Priority,
+	    ActorClass,
+	    GetClass()
+    );
+	
+	InAITModule->RegisterTool(Category, MoveTemp(ToolParameters));
+}
\ No newline at end of file
diff --git a/Engine/Plugins/Experimental/Avalanche/Source/AvalancheInteractiveTools/Private/Tools/AvaInteractiveToolsActorToolNull.cpp b/Engine/Plugins/Experimental/Avalanche/Source/AvalancheInteractiveTools/Private/Tools/AvaInteractiveToolsActorToolNull.cpp
--- a/Engine/Plugins/Experimental/Avalanche/Source/AvalancheInteractiveTools/Private/Tools/AvaInteractiveToolsActorToolNull.cpp
+++ b/Engine/Plugins/Experimental/Avalanche/Source/AvalancheInteractiveTools/Private/Tools/AvaInteractiveToolsActorToolNull.cpp
@@ -11,14 +11,12 @@ UAvaInteractiveToolsActorToolNull::UAvaInteractiveToolsActorToolNull()
 	ActorClass = AAvaNullActor::StaticClass();
 }
 
-FName UAvaInteractiveToolsActorToolNull::GetCategoryName()
+void UAvaInteractiveToolsActorToolNull::OnRegisterTool(IAvalancheInteractiveToolsModule* InAITModule)
 {
-	return FAvalancheInteractiveToolsModule::CategoryNameActor;
-}
+	Super::OnRegisterTool(InAITModule);
 
-FAvaInteractiveToolsToolParameters UAvaInteractiveToolsActorToolNull::GetToolParameters() const
-{
-	return {
+	FAvaInteractiveToolsToolParameters ToolParameters =
+	{
 		FAvaInteractiveToolsCommands::Get().Tool_Actor_Null,
 		TEXT("Null Actor Tool"),
 		1000,
@@ -29,4 +27,6 @@ FAvaInteractiveToolsToolParameters UAvaInteractiveToolsActorToolNull::GetToolPar
 			}),
 		ActorClass
 	};
+
+	InAITModule->RegisterTool(IAvalancheInteractiveToolsModule::CategoryNameActor, MoveTemp(ToolParameters));
 }
diff --git a/Engine/Plugins/Experimental/Avalanche/Source/AvalancheInteractiveTools/Private/Tools/AvaInteractiveToolsActorToolNull.h b/Engine/Plugins/Experimental/Avalanche/Source/AvalancheInteractiveTools/Private/Tools/AvaInteractiveToolsActorToolNull.h
--- a/Engine/Plugins/Experimental/Avalanche/Source/AvalancheInteractiveTools/Private/Tools/AvaInteractiveToolsActorToolNull.h
+++ b/Engine/Plugins/Experimental/Avalanche/Source/AvalancheInteractiveTools/Private/Tools/AvaInteractiveToolsActorToolNull.h
@@ -15,7 +15,6 @@ public:
 
 	//~ Begin UAvaInteractiveToolsToolBase
 	virtual bool UseIdentityRotation() const override { return true; }
-	virtual FName GetCategoryName() override;
-	virtual FAvaInteractiveToolsToolParameters GetToolParameters() const override;
+	virtual void OnRegisterTool(IAvalancheInteractiveToolsModule* InAITModule) override;
 	//~ End UAvaInteractiveToolsToolBase
 };
diff --git a/Engine/Plugins/Experimental/Avalanche/Source/AvalancheInteractiveTools/Private/Tools/AvaInteractiveToolsActorToolSpline.cpp b/Engine/Plugins/Experimental/Avalanche/Source/AvalancheInteractiveTools/Private/Tools/AvaInteractiveToolsActorToolSpline.cpp
--- a/Engine/Plugins/Experimental/Avalanche/Source/AvalancheInteractiveTools/Private/Tools/AvaInteractiveToolsActorToolSpline.cpp
+++ b/Engine/Plugins/Experimental/Avalanche/Source/AvalancheInteractiveTools/Private/Tools/AvaInteractiveToolsActorToolSpline.cpp
@@ -11,14 +11,12 @@ UAvaInteractiveToolsActorToolSpline::UAvaInteractiveToolsActorToolSpline()
 	ActorClass = AAvaSplineActor::StaticClass();
 }
 
-FName UAvaInteractiveToolsActorToolSpline::GetCategoryName()
+void UAvaInteractiveToolsActorToolSpline::OnRegisterTool(IAvalancheInteractiveToolsModule* InAITModule)
 {
-	return FAvalancheInteractiveToolsModule::CategoryNameActor;
-}
+	Super::OnRegisterTool(InAITModule);
 
-FAvaInteractiveToolsToolParameters UAvaInteractiveToolsActorToolSpline::GetToolParameters() const
-{
-	return {
+	FAvaInteractiveToolsToolParameters ToolParameters =
+	{
 		FAvaInteractiveToolsCommands::Get().Tool_Actor_Spline,
 		TEXT("Spline Actor Tool"),
 		5000,
@@ -29,4 +27,6 @@ FAvaInteractiveToolsToolParameters UAvaInteractiveToolsActorToolSpline::GetToolP
 			}),
 		ActorClass
 	};
+
+	InAITModule->RegisterTool(IAvalancheInteractiveToolsModule::CategoryNameActor, MoveTemp(ToolParameters));
 }
diff --git a/Engine/Plugins/Experimental/Avalanche/Source/AvalancheInteractiveTools/Private/Tools/AvaInteractiveToolsActorToolSpline.h b/Engine/Plugins/Experimental/Avalanche/Source/AvalancheInteractiveTools/Private/Tools/AvaInteractiveToolsActorToolSpline.h
--- a/Engine/Plugins/Experimental/Avalanche/Source/AvalancheInteractiveTools/Private/Tools/AvaInteractiveToolsActorToolSpline.h
+++ b/Engine/Plugins/Experimental/Avalanche/Source/AvalancheInteractiveTools/Private/Tools/AvaInteractiveToolsActorToolSpline.h
@@ -15,7 +15,6 @@ public:
 
 	//~ Begin UAvaInteractiveToolsToolBase
 	virtual bool UseIdentityRotation() const override { return false; }
-	virtual FName GetCategoryName() override;
-	virtual FAvaInteractiveToolsToolParameters GetToolParameters() const override;
+	virtual void OnRegisterTool(IAvalancheInteractiveToolsModule* InAITModule) override;
 	//~ End UAvaInteractiveToolsToolBase
 };
diff --git a/Engine/Plugins/Experimental/Avalanche/Source/AvalancheInteractiveTools/Private/Tools/AvaInteractiveToolsStaticMeshActorTool.cpp b/Engine/Plugins/Experimental/Avalanche/Source/AvalancheInteractiveTools/Private/Tools/AvaInteractiveToolsStaticMeshActorTool.cpp
--- a/Engine/Plugins/Experimental/Avalanche/Source/AvalancheInteractiveTools/Private/Tools/AvaInteractiveToolsStaticMeshActorTool.cpp
+++ b/Engine/Plugins/Experimental/Avalanche/Source/AvalancheInteractiveTools/Private/Tools/AvaInteractiveToolsStaticMeshActorTool.cpp
@@ -11,10 +11,9 @@ UAvaInteractiveToolsStaticMeshActorTool::UAvaInteractiveToolsStaticMeshActorTool
 	ActorClass = AStaticMeshActor::StaticClass();
 }
 
-FAvaInteractiveToolsToolParameters UAvaInteractiveToolsStaticMeshActorTool::GetToolParameters() const
+void UAvaInteractiveToolsStaticMeshActorTool::OnRegisterTool(IAvalancheInteractiveToolsModule* InAITModule)
 {
-	// This is never really needed, it should work regardless. OOP consistency!
-	return UAvaInteractiveToolsStaticMeshActorToolBuilder::CreateToolParameters(
+	FAvaInteractiveToolsToolParameters ToolParameters = UAvaInteractiveToolsStaticMeshActorToolBuilder::CreateToolParameters(
 		Category,
 		Command,
 		Identifier,
@@ -22,10 +21,12 @@ FAvaInteractiveToolsToolParameters UAvaInteractiveToolsStaticMeshActorTool::GetT
 		StaticMesh,
 		GetClass()
 	);
+
+	InAITModule->RegisterTool(Category, MoveTemp(ToolParameters));
 }
 
 AActor* UAvaInteractiveToolsStaticMeshActorTool::SpawnActor(TSubclassOf<AActor> InActorClass, EAvaViewportStatus InViewportStatus,
-	const FVector2f& InViewportPosition, bool bInPreview, FString* InActorLabelOverride) const
+	const FVector2f& InViewportPosition, bool bInPreview, FString* InActorLabelOverride)
 {
 	AActor* Actor = Super::SpawnActor(InActorClass, InViewportStatus, InViewportPosition, bInPreview, InActorLabelOverride);
 
diff --git a/Engine/Plugins/Experimental/Avalanche/Source/AvalancheInteractiveTools/Private/Tools/AvaInteractiveToolsToolBase.cpp b/Engine/Plugins/Experimental/Avalanche/Source/AvalancheInteractiveTools/Private/Tools/AvaInteractiveToolsToolBase.cpp
--- a/Engine/Plugins/Experimental/Avalanche/Source/AvalancheInteractiveTools/Private/Tools/AvaInteractiveToolsToolBase.cpp
+++ b/Engine/Plugins/Experimental/Avalanche/Source/AvalancheInteractiveTools/Private/Tools/AvaInteractiveToolsToolBase.cpp
@@ -112,21 +112,10 @@ void UAvaInteractiveToolsToolBase::Setup()
 		}
 	}
 
-	if (!CanActivate(bReactivated))
-	{
-		ToolManager->PostActiveToolShutdownRequest(this, EToolShutdownType::Cancel);
-		return;
-	}
-
 	if (AvaInteractiveToolsEdMode)
 	{
 		AvaInteractiveToolsEdMode->OnToolSetup(this);
-	}
-
-	Activate(bReactivated);
-
-	if (AvaInteractiveToolsEdMode)
-	{
+		Activate(AvaInteractiveToolsEdMode->GetLastActiveTool(), bReactivated);
 		AvaInteractiveToolsEdMode->OnToolActivateEnd();
 	}
 }
@@ -213,7 +202,7 @@ bool UAvaInteractiveToolsToolBase::ConditionalIdentityRotation() const
 	}
 }
 
-AActor* UAvaInteractiveToolsToolBase::SpawnActor(TSubclassOf<AActor> InActorClass, bool bInPreview, FString* InActorLabelOverride) const
+AActor* UAvaInteractiveToolsToolBase::SpawnActor(TSubclassOf<AActor> InActorClass, bool bInPreview, FString* InActorLabelOverride)
 {
 	if (IToolsContextQueriesAPI* ContextAPI = GetToolManager()->GetContextQueriesAPI())
 	{
@@ -237,7 +226,7 @@ AActor* UAvaInteractiveToolsToolBase::SpawnActor(TSubclassOf<AActor> InActorClas
 }
 
 AActor* UAvaInteractiveToolsToolBase::SpawnActor(TSubclassOf<AActor> InActorClass, EAvaViewportStatus InViewportStatus, 
-	const FVector2f& InViewportPosition, bool bInPreview, FString* InActorLabelOverride) const
+	const FVector2f& InViewportPosition, bool bInPreview, FString* InActorLabelOverride)
 {
 	UWorld* World = nullptr;;
 	FVector SpawnLocation = FVector::ZeroVector;
@@ -325,6 +314,8 @@ AActor* UAvaInteractiveToolsToolBase::SpawnActor(TSubclassOf<AActor> InActorClas
 		SetToolkitSettingsObject(DetailsObject);
 	}
 
+	OnActorSpawned(NewActor);
+
 	return NewActor;
 }
 
@@ -472,6 +463,16 @@ void UAvaInteractiveToolsToolBase::OnClicked(const FInputDeviceRay& InClickPos)
 	RequestShutdown(EToolShutdownType::Cancel);
 }
 
+FString UAvaInteractiveToolsToolBase::GetActiveToolIdentifier() const
+{
+	return FAvalancheInteractiveToolsModule::Get().GetActiveToolIdentifier();
+}
+
+const FAvaInteractiveToolsToolParameters* UAvaInteractiveToolsToolBase::GetActiveToolParameters() const
+{
+	return FAvalancheInteractiveToolsModule::Get().GetTool(GetActiveToolIdentifier());
+}
+
 void UAvaInteractiveToolsToolBase::OnUpdateModifierState(int InModifierID, bool bInIsOn)
 {
 	switch (InModifierID)
@@ -486,7 +487,7 @@ void UAvaInteractiveToolsToolBase::OnUpdateModifierState(int InModifierID, bool
 	}
 }
 
-bool UAvaInteractiveToolsToolBase::CanActivate(bool bInReactivate) const
+bool UAvaInteractiveToolsToolBase::CanActivate(const FString& InToolIdentifier, bool bInReactivate) const
 {
 	if (!bInReactivate)
 	{
@@ -497,25 +498,20 @@ bool UAvaInteractiveToolsToolBase::CanActivate(bool bInReactivate) const
 	return SupportsDefaultAction();
 }
 
-void UAvaInteractiveToolsToolBase::Activate(bool bInReactivate)
+void UAvaInteractiveToolsToolBase::Activate(const FString& InToolIdentifier, bool bInReactivate)
 {
-	if (!CanActivate(bInReactivate))
+	if (!CanActivate(InToolIdentifier, bInReactivate))
 	{
 		RequestShutdown(EToolShutdownType::Cancel);
 		return;
 	}
 
-	FAvalancheInteractiveToolsModule::Get().OnToolActivated();
-
-	const bool bUsePresetMenu = ShouldUsePresetMenu();
-	const bool bSupportsDefaultAction = !bUsePresetMenu && SupportsDefaultAction();
-	const bool bForceDefaultAction = !bUsePresetMenu && ShouldForceDefaultAction();
+	FAvalancheInteractiveToolsModule::Get().OnToolActivated(InToolIdentifier);
+	
+	const bool bSupportsDefaultAction = SupportsDefaultAction();
+	const bool bForceDefaultAction = ShouldForceDefaultAction();
 
-	if (bUsePresetMenu)
-	{
-		ShowPresetMenu();
-	}
-	else if (!bForceDefaultAction && (!bInReactivate || !bSupportsDefaultAction))
+	if (!bForceDefaultAction && (!bInReactivate || !bSupportsDefaultAction))
 	{
 		OnActivate();
 
@@ -673,116 +669,4 @@ bool UAvaInteractiveToolsToolBase::ShouldForceDefaultAction() const
 	return FSlateApplication::Get().GetModifierKeys().IsAltDown();
 }
 
-bool UAvaInteractiveToolsToolBase::ShouldUsePresetMenu() const
-{
-	const FAvaInteractiveToolsToolParameters ToolParams = GetToolParameters();
-
-	if (ToolParams.Presets.IsEmpty())
-	{
-		return false;
-	}
-
-	return SupportsDefaultAction() && FSlateApplication::Get().GetModifierKeys().IsShiftDown();
-}
-
-void UAvaInteractiveToolsToolBase::ShowPresetMenu()
-{
-	using namespace UE::AvaInteractiveTools::Private;
-
-	RegisterPresetMenu();
-
-	TSharedPtr<SWidget> Parent = FSlateApplication::Get().FindBestParentWindowForDialogs(nullptr);
-
-	if (!Parent.IsValid())
-	{
-		return;
-	}
-
-	FToolMenuContext Context;
-	Context.AddObject(this);
-
-	FSlateApplication::Get().PushMenu(
-		Parent.ToSharedRef(),
-		FWidgetPath(),
-		UToolMenus::Get()->GenerateWidget(AvaITFToolPresetMenuName, Context),
-		FSlateApplication::Get().GetCursorPos(),
-		FPopupTransitionEffect::ContextMenu
-	);
-}
-
-void UAvaInteractiveToolsToolBase::OnPresetSelected(TStrongObjectPtr<UAvaInteractiveToolsToolBase> InToolPtr, FName InPresetName)
-{
-	UAvaInteractiveToolsToolBase* Tool = InToolPtr.Get();
-
-	if (!Tool)
-	{
-		return;
-	}
-
-	const FAvaInteractiveToolsToolParameters ToolParams = Tool->GetToolParameters();
-	const TSharedRef<FAvaInteractiveToolsToolPresetBase>* Preset = ToolParams.Presets.Find(InPresetName);
-
-	if (!Preset)
-	{
-		return;
-	}
-
-	Tool->bPerformingDefaultAction = true;
-	Tool->DefaultAction();
-
-	(*Preset)->ApplyPreset(Tool->GetSpawnedActor());
-}
-
-void UAvaInteractiveToolsToolBase::RegisterPresetMenu()
-{
-	using namespace UE::AvaInteractiveTools::Private;
-
-	UToolMenus* ToolMenus = UToolMenus::Get();
-
-	if (ToolMenus->IsMenuRegistered(AvaITFToolPresetMenuName))
-	{
-		return;
-	}
-
-	UToolMenu* PresetMenu = ToolMenus->RegisterMenu(AvaITFToolPresetMenuName);
-
-	static const FName PresetSectionName = "Preset";
-
-	PresetMenu->AddDynamicSection(
-		PresetSectionName,
-		FNewSectionConstructChoice(FNewToolMenuDelegate::CreateLambda(
-			[](UToolMenu* InMenu)
-			{
-				UAvaInteractiveToolsToolBase* Tool = InMenu->FindContext<UAvaInteractiveToolsToolBase>();
-
-				if (!Tool)
-				{
-					return;
-				}
-
-				const FAvaInteractiveToolsToolParameters ToolParams = Tool->GetToolParameters();
-
-				if (ToolParams.Presets.IsEmpty())
-				{
-					return;
-				}
-
-				TStrongObjectPtr<UAvaInteractiveToolsToolBase> ToolPtr = TStrongObjectPtr<UAvaInteractiveToolsToolBase>(Tool);
-				FToolMenuSection& Section = InMenu->FindOrAddSection("Presets");
-
-				for (const TPair<FName, TSharedRef<FAvaInteractiveToolsToolPresetBase>>& PresetPair : ToolParams.Presets)
-				{
-					Section.AddMenuEntry(
-						PresetPair.Key,
-						PresetPair.Value->GetName(),
-						PresetPair.Value->GetDescription(),
-						ToolParams.UICommand->GetIcon(),
-						FToolUIActionChoice(FExecuteAction::CreateStatic(&UAvaInteractiveToolsToolBase::OnPresetSelected, ToolPtr, PresetPair.Key))
-					);
-				}
-			}
-		))
-	);
-}
-
 #undef LOCTEXT_NAMESPACE
diff --git a/Engine/Plugins/Experimental/Avalanche/Source/AvalancheInteractiveTools/Private/Widgets/SAvaInteractiveToolsToolbarCategoryButton.cpp b/Engine/Plugins/Experimental/Avalanche/Source/AvalancheInteractiveTools/Private/Widgets/SAvaInteractiveToolsToolbarCategoryButton.cpp
new file mode 100644
--- /dev/null
+++ b/Engine/Plugins/Experimental/Avalanche/Source/AvalancheInteractiveTools/Private/Widgets/SAvaInteractiveToolsToolbarCategoryButton.cpp
@@ -0,0 +1,137 @@
+// Copyright Epic Games, Inc. All Rights Reserved.
+
+#include "Widgets/SAvaInteractiveToolsToolbarCategoryButton.h"
+
+#include "Application/ThrottleManager.h"
+#include "AvaInteractiveToolsStyle.h"
+#include "Framework/Application/IMenu.h"
+#include "Framework/Application/SlateApplication.h"
+#include "Framework/Commands/UICommandList.h"
+#include "Framework/MultiBox/MultiBoxBuilder.h"
+#include "IAvalancheInteractiveToolsModule.h"
+
+void SAvaInteractiveToolsToolbarCategoryButton::Construct(const FArguments& InArgs)
+{
+	CommandList = InArgs._CommandList;
+	bShowLabel = InArgs._ShowLabel;
+	Category = InArgs._Category;
+
+	IAvalancheInteractiveToolsModule& AITModule = IAvalancheInteractiveToolsModule::Get();
+
+	check(CommandList.IsValid() && AITModule.GetCategories().Contains(Category) && !AITModule.GetTools(Category)->IsEmpty())
+
+	ActiveToolIdentifier = (*AITModule.GetTools(Category))[0].ToolIdentifier;
+
+	CreateActiveCommandWidget();
+
+	IAvalancheInteractiveToolsModule::Get().OnToolActivation().AddSP(this, &SAvaInteractiveToolsToolbarCategoryButton::OnToolActivated);
+}
+
+SAvaInteractiveToolsToolbarCategoryButton::~SAvaInteractiveToolsToolbarCategoryButton()
+{
+	IAvalancheInteractiveToolsModule::Get().OnToolActivation().RemoveAll(this);
+}
+
+FReply SAvaInteractiveToolsToolbarCategoryButton::OnMouseButtonDown(const FGeometry& InGeometry, const FPointerEvent& InMouseEvent)
+{
+	if (InMouseEvent.GetEffectingButton() == EKeys::RightMouseButton)
+	{
+		ShowCommandsContextMenu();
+		return FReply::Handled();
+	}
+
+	return SCompoundWidget::OnMouseButtonDown(InGeometry, InMouseEvent);
+}
+
+void SAvaInteractiveToolsToolbarCategoryButton::CreateActiveCommandWidget()
+{
+	IAvalancheInteractiveToolsModule& AITModule = IAvalancheInteractiveToolsModule::Get();
+	const FAvaInteractiveToolsToolParameters* Tool = AITModule.GetTool(ActiveToolIdentifier);
+
+	if (Tool && Tool->UICommand.IsValid())
+	{
+		FVerticalToolBarBuilder CategoryToolbarBuilder(
+		CommandList,
+		FMultiBoxCustomization::None,
+		TSharedPtr<FExtender>()
+		, /* ForceSmallIcons */ true);
+
+		CategoryToolbarBuilder.SetStyle(&FAvaInteractiveToolsStyle::Get(), "ViewportToolbar");
+		CategoryToolbarBuilder.SetLabelVisibility(bShowLabel ? EVisibility::Visible : EVisibility::Collapsed);
+
+		TAttribute<FText> LabelAttribute = Tool->Label.IsSet() ? Tool->Label.GetValue() : TAttribute<FText>();
+		TAttribute<FText> TooltipAttribute = Tool->Tooltip.IsSet() ? Tool->Tooltip.GetValue() : TAttribute<FText>();
+		TAttribute<FSlateIcon> IconAttribute = Tool->Icon.IsSet() ? Tool->Icon.GetValue() : TAttribute<FSlateIcon>();
+		CategoryToolbarBuilder.AddToolBarButton(Tool->UICommand, NAME_None, LabelAttribute, TooltipAttribute, IconAttribute);
+
+		ChildSlot
+		[
+			CategoryToolbarBuilder.MakeWidget()
+		];
+	}
+}
+
+void SAvaInteractiveToolsToolbarCategoryButton::ShowCommandsContextMenu()
+{
+	HideCommandsContextMenu();
+
+	// Needed otherwise cannot select entries in context toolbar below
+	FSlateThrottleManager::Get().DisableThrottle(true);
+
+	const FVector2D MousePosition = FSlateApplication::Get().GetLastCursorPos();
+
+	FVerticalToolBarBuilder CategoryToolbarBuilder(
+		CommandList,
+		FMultiBoxCustomization::None,
+		TSharedPtr<FExtender>()
+		, /* ForceSmallIcons */ false);
+
+	CategoryToolbarBuilder.SetStyle(&FAvaInteractiveToolsStyle::Get(), "ViewportToolbar");
+
+	IAvalancheInteractiveToolsModule& AITModule = IAvalancheInteractiveToolsModule::Get();
+	for (const FAvaInteractiveToolsToolParameters& Tool : (*AITModule.GetTools(Category)))
+	{
+		if (Tool.UICommand.IsValid())
+		{
+			TAttribute<FText> LabelAttribute = Tool.Label.IsSet() ? Tool.Label.GetValue() : TAttribute<FText>();
+			TAttribute<FText> TooltipAttribute = Tool.Tooltip.IsSet() ? Tool.Tooltip.GetValue() : TAttribute<FText>();
+			TAttribute<FSlateIcon> IconAttribute = Tool.Icon.IsSet() ? Tool.Icon.GetValue() : TAttribute<FSlateIcon>();
+
+			CategoryToolbarBuilder.AddToolBarButton(Tool.UICommand, NAME_None, LabelAttribute, TooltipAttribute, IconAttribute);
+		}
+	}
+
+	ContextMenu = FSlateApplication::Get().PushMenu(
+		FSlateApplication::Get().GetActiveTopLevelWindow().ToSharedRef(),
+		FWidgetPath(),
+		CategoryToolbarBuilder.MakeWidget(),
+		MousePosition,
+		FPopupTransitionEffect(FPopupTransitionEffect::None)
+	);
+	
+	ContextMenu->GetOnMenuDismissed().AddLambda([this](TSharedRef<IMenu>)
+	{
+		FSlateThrottleManager::Get().DisableThrottle(false);
+		HideCommandsContextMenu();
+	});
+}
+
+void SAvaInteractiveToolsToolbarCategoryButton::HideCommandsContextMenu()
+{
+	if (ContextMenu.IsValid())
+	{
+		ContextMenu->Dismiss();
+		ContextMenu.Reset();
+	}
+}
+
+void SAvaInteractiveToolsToolbarCategoryButton::OnToolActivated(const FString& InToolIdentifier)
+{
+	const FName ToolCategory = IAvalancheInteractiveToolsModule::Get().GetToolCategory(InToolIdentifier);
+	if (ToolCategory.IsEqual(Category))
+	{
+		ActiveToolIdentifier = InToolIdentifier;
+		CreateActiveCommandWidget();
+		HideCommandsContextMenu();
+	}
+}
\ No newline at end of file
diff --git a/Engine/Plugins/Experimental/Avalanche/Source/AvalancheInteractiveTools/Private/Widgets/SAvaInteractiveToolsToolbarCategoryButton.h b/Engine/Plugins/Experimental/Avalanche/Source/AvalancheInteractiveTools/Private/Widgets/SAvaInteractiveToolsToolbarCategoryButton.h
new file mode 100644
--- /dev/null
+++ b/Engine/Plugins/Experimental/Avalanche/Source/AvalancheInteractiveTools/Private/Widgets/SAvaInteractiveToolsToolbarCategoryButton.h
@@ -0,0 +1,44 @@
+// Copyright Epic Games, Inc. All Rights Reserved.
+
+#pragma once
+
+#include "Templates/SharedPointer.h"
+#include "Widgets/SCompoundWidget.h"
+
+class FUICommandList;
+class IMenu;
+
+class SAvaInteractiveToolsToolbarCategoryButton : public SCompoundWidget
+{
+public:
+	SLATE_BEGIN_ARGS(SAvaInteractiveToolsToolbarCategoryButton)
+	{}
+		/** Category to retrieve tools from */
+		SLATE_ARGUMENT(FName, Category)
+		/** Command list to execute selected command */
+		SLATE_ARGUMENT(TSharedPtr<FUICommandList>, CommandList)
+		/** Show label under category button */
+		SLATE_ARGUMENT(bool, ShowLabel)
+	SLATE_END_ARGS()
+
+	void Construct(const FArguments& InArgs);
+	virtual ~SAvaInteractiveToolsToolbarCategoryButton() override;
+
+protected:
+	//~ Begin SButton
+	virtual FReply OnMouseButtonDown(const FGeometry& InGeometry, const FPointerEvent& InMouseEvent) override;
+	//~ End SButton
+
+	void CreateActiveCommandWidget();
+	void ShowCommandsContextMenu();
+	void HideCommandsContextMenu();
+
+	void OnToolActivated(const FString& InToolIdentifier);
+
+private:
+	TSharedPtr<FUICommandList> CommandList;
+	TSharedPtr<IMenu> ContextMenu;
+	bool bShowLabel = false;
+	FName Category = NAME_None;
+	FString ActiveToolIdentifier;
+};
\ No newline at end of file
diff --git a/Engine/Plugins/Experimental/Avalanche/Source/AvalancheInteractiveTools/Public/AvaInteractiveToolsSettings.h b/Engine/Plugins/Experimental/Avalanche/Source/AvalancheInteractiveTools/Public/AvaInteractiveToolsSettings.h
--- a/Engine/Plugins/Experimental/Avalanche/Source/AvalancheInteractiveTools/Public/AvaInteractiveToolsSettings.h
+++ b/Engine/Plugins/Experimental/Avalanche/Source/AvalancheInteractiveTools/Public/AvaInteractiveToolsSettings.h
@@ -6,24 +6,52 @@
 #include "AvaInteractiveToolsSettings.generated.h"
 
 UENUM()
-enum class EAvaInteractiveToolsDefaultActionAlignment
+enum class EAvaInteractiveToolsDefaultActionAlignment : uint8
 {
 	Axis,
 	Camera
 };
 
+UENUM()
+enum class EAvaInteractiveToolsViewportToolbarPosition : uint8
+{
+	/** Disabled */
+	None,
+	/** Horizontal */
+	Bottom,
+	/** Horizontal */
+	Top,
+	/** Vertical */
+	Left,
+	/** Vertical */
+	Right
+};
+
 UCLASS(Config = EditorPerProjectUserSettings, meta = (DisplayName = "Interactive Tools"))
 class AVALANCHEINTERACTIVETOOLS_API UAvaInteractiveToolsSettings : public UDeveloperSettings
 {
 	GENERATED_BODY()
 
 public:
+	static UAvaInteractiveToolsSettings* Get();
+
 	UAvaInteractiveToolsSettings();
 
+	/** Open these settings in the editor */
+	void OpenEditorSettingsWindow() const;
+
 	/** Distance from the camera at which actors are created. */
 	UPROPERTY(Config, EditAnywhere, Category = "Motion Design")
 	float CameraDistance = 500.f;
 
 	UPROPERTY(Config, EditAnywhere, Category = "Motion Design")
 	EAvaInteractiveToolsDefaultActionAlignment DefaultActionActorAlignment = EAvaInteractiveToolsDefaultActionAlignment::Axis;
+
+	/** Position of the viewport overlay toolbar position */
+	UPROPERTY(Config, EditAnywhere, Category = "Motion Design")
+	EAvaInteractiveToolsViewportToolbarPosition ViewportToolbarPosition = EAvaInteractiveToolsViewportToolbarPosition::Bottom;
+
+	/** Show viewport overlay toolbar label under each item */
+	UPROPERTY(Config, EditAnywhere, Category = "Motion Design")
+	bool bViewportToolbarLabelEnabled = false;
 };
diff --git a/Engine/Plugins/Experimental/Avalanche/Source/AvalancheInteractiveTools/Public/IAvalancheInteractiveToolsModule.h b/Engine/Plugins/Experimental/Avalanche/Source/AvalancheInteractiveTools/Public/IAvalancheInteractiveToolsModule.h
--- a/Engine/Plugins/Experimental/Avalanche/Source/AvalancheInteractiveTools/Public/IAvalancheInteractiveToolsModule.h
+++ b/Engine/Plugins/Experimental/Avalanche/Source/AvalancheInteractiveTools/Public/IAvalancheInteractiveToolsModule.h
@@ -5,18 +5,19 @@
 #include "Containers/Map.h"
 #include "Containers/UnrealString.h"
 #include "Editor.h"
+#include "Misc/Optional.h"
 #include "Modules/ModuleInterface.h"
 #include "Modules/ModuleManager.h"
 #include "Templates/SharedPointer.h"
-#include "Tools/AvaInteractiveToolsToolPresetBase.h"
+#include "Textures/SlateIcon.h"
 
 class FName;
 class FUICommandInfo;
+class FText;
 class IAvalancheInteractiveToolsModule;
 class UEdMode;
 class UInteractiveToolBuilder;
 enum class EToolsContextScope;
-struct FSlateIcon;
 
 DECLARE_DELEGATE_RetVal_OneParam(UInteractiveToolBuilder*, FAvalancheInteractiveToolsCreateBuilder, UEdMode*)
 
@@ -28,7 +29,9 @@ struct FAvaInteractiveToolsToolParameters
 	FAvalancheInteractiveToolsCreateBuilder CreateBuilder;
 	TSubclassOf<AActor> FactoryClass = nullptr;
 	TObjectPtr<UActorFactory> Factory = nullptr;
-	TMap<FName, TSharedRef<FAvaInteractiveToolsToolPresetBase>> Presets;
+	TOptional<FText> Label;
+	TOptional<FText> Tooltip;
+	TOptional<FSlateIcon> Icon;
 };
 
 namespace UE::AvaInteractiveTools
@@ -45,6 +48,8 @@ public:
 	static const inline FName CategoryName2D = FName(TEXT("AvaTools2D"));
 	static const inline FName CategoryName3D = FName(TEXT("AvaTools3D"));
 	static const inline FName CategoryNameActor = FName(TEXT("AvaToolsActor"));
+	static const inline FName CategoryNameCloner = FName(TEXT("AvaToolsCloner"));
+	static const inline FName CategoryNameEffector = FName(TEXT("AvaToolsEffector"));
 
 	static IAvalancheInteractiveToolsModule& Get()
 	{
@@ -67,6 +72,14 @@ public:
 	virtual void RegisterTool(FName InCategory, FAvaInteractiveToolsToolParameters&& InToolParams) = 0;
 
 	virtual const TArray<FAvaInteractiveToolsToolParameters>* GetTools(FName InCategory) = 0;
+	
+	virtual const FAvaInteractiveToolsToolParameters* GetTool(const FString& InToolIdentifier) = 0;
+
+	virtual FName GetToolCategory(const FString& InToolIdentifier) = 0;
 
 	virtual bool HasActiveTool() const = 0;
+
+	DECLARE_MULTICAST_DELEGATE_OneParam(FToolEvent, const FString& /** ToolIdentifier */)
+	virtual FToolEvent::RegistrationType& OnToolActivation() = 0;
+	virtual FToolEvent::RegistrationType& OnToolDeactivation() = 0;
 };
diff --git a/Engine/Plugins/Experimental/Avalanche/Source/AvalancheInteractiveTools/Public/Tools/AvaInteractiveToolsActorAreaToolBase.h b/Engine/Plugins/Experimental/Avalanche/Source/AvalancheInteractiveTools/Public/Tools/AvaInteractiveToolsActorAreaToolBase.h
--- a/Engine/Plugins/Experimental/Avalanche/Source/AvalancheInteractiveTools/Public/Tools/AvaInteractiveToolsActorAreaToolBase.h
+++ b/Engine/Plugins/Experimental/Avalanche/Source/AvalancheInteractiveTools/Public/Tools/AvaInteractiveToolsActorAreaToolBase.h
@@ -27,6 +27,6 @@ protected:
 
 	//~ Begin UAvaInteractiveToolsToolBase
 	virtual AActor* SpawnActor(TSubclassOf<AActor> InActorClass, EAvaViewportStatus InViewportStatus, 
-		const FVector2f& InViewportPosition, bool bInPreview, FString* InActorLabelOverride = nullptr) const override;
+		const FVector2f& InViewportPosition, bool bInPreview, FString* InActorLabelOverride = nullptr) override;
 	//~ End UAvaInteractiveToolsToolBase
 };
diff --git a/Engine/Plugins/Experimental/Avalanche/Source/AvalancheInteractiveTools/Public/Tools/AvaInteractiveToolsActorTool.h b/Engine/Plugins/Experimental/Avalanche/Source/AvalancheInteractiveTools/Public/Tools/AvaInteractiveToolsActorTool.h
--- a/Engine/Plugins/Experimental/Avalanche/Source/AvalancheInteractiveTools/Public/Tools/AvaInteractiveToolsActorTool.h
+++ b/Engine/Plugins/Experimental/Avalanche/Source/AvalancheInteractiveTools/Public/Tools/AvaInteractiveToolsActorTool.h
@@ -22,13 +22,12 @@ public:
 	UAvaInteractiveToolsActorTool();
 
 	//~ Begin UAvaInteractiveToolsToolBase
-	virtual FName GetCategoryName() override { return Category; }
-	virtual FAvaInteractiveToolsToolParameters GetToolParameters() const override;
+	virtual void OnRegisterTool(IAvalancheInteractiveToolsModule* InAITModule) override;
 	//~ End UAvaInteractiveToolsToolBase
 
 protected:
 	FName Category;
 	TSharedPtr<FUICommandInfo> Command;
 	FString Identifier;
-	int32 Priority;
+	int32 Priority = 0;
 };
diff --git a/Engine/Plugins/Experimental/Avalanche/Source/AvalancheInteractiveTools/Public/Tools/AvaInteractiveToolsStaticMeshActorTool.h b/Engine/Plugins/Experimental/Avalanche/Source/AvalancheInteractiveTools/Public/Tools/AvaInteractiveToolsStaticMeshActorTool.h
--- a/Engine/Plugins/Experimental/Avalanche/Source/AvalancheInteractiveTools/Public/Tools/AvaInteractiveToolsStaticMeshActorTool.h
+++ b/Engine/Plugins/Experimental/Avalanche/Source/AvalancheInteractiveTools/Public/Tools/AvaInteractiveToolsStaticMeshActorTool.h
@@ -21,7 +21,7 @@ public:
 	UAvaInteractiveToolsStaticMeshActorTool();
 
 	//~ Begin UAvaInteractiveToolsToolBase
-	virtual FAvaInteractiveToolsToolParameters GetToolParameters() const override;
+	virtual void OnRegisterTool(IAvalancheInteractiveToolsModule* InAITModule) override;
 	//~ End UAvaInteractiveToolsToolBase
 
 protected:
@@ -30,6 +30,6 @@ protected:
 
 	//~ Begin UAvaInteractiveToolsToolBase
 	virtual AActor* SpawnActor(TSubclassOf<AActor> InActorClass, EAvaViewportStatus InViewportStatus,
-		const FVector2f& InViewportPosition, bool bInPreview, FString* InActorLabelOverride = nullptr) const override;
+		const FVector2f& InViewportPosition, bool bInPreview, FString* InActorLabelOverride = nullptr) override;
 	//~ End UAvaInteractiveToolsToolBase
 };
diff --git a/Engine/Plugins/Experimental/Avalanche/Source/AvalancheInteractiveTools/Public/Tools/AvaInteractiveToolsToolBase.h b/Engine/Plugins/Experimental/Avalanche/Source/AvalancheInteractiveTools/Public/Tools/AvaInteractiveToolsToolBase.h
--- a/Engine/Plugins/Experimental/Avalanche/Source/AvalancheInteractiveTools/Public/Tools/AvaInteractiveToolsToolBase.h
+++ b/Engine/Plugins/Experimental/Avalanche/Source/AvalancheInteractiveTools/Public/Tools/AvaInteractiveToolsToolBase.h
@@ -66,11 +66,8 @@ public:
 	/** Returns true if this supports a default action. */
 	virtual bool SupportsDefaultAction() const;
 
-	/** Designed to be called on the CDO. */
-	virtual FName GetCategoryName() PURE_VIRTUAL(UAvaInteractiveToolsToolBase::GetCategory, return NAME_None;)
-
-	/** Designed to be called on the CDO. */
-	virtual FAvaInteractiveToolsToolParameters GetToolParameters() const PURE_VIRTUAL(UAvaInteractiveToolsToolBase::GetCommand, return {};)
+	FString GetActiveToolIdentifier() const;
+	const FAvaInteractiveToolsToolParameters* GetActiveToolParameters() const;
 
 	//~ Begin IModifierToggleBehaviorTarget
 	virtual void OnUpdateModifierState(int InModifierID, bool bInIsOn) override;
@@ -93,8 +90,8 @@ public:
 	AActor* GetSpawnedActor() const { return SpawnedActor; }
 
 	virtual void OnViewportPlannerUpdate() {}
-
 	virtual void OnViewportPlannerComplete();
+	virtual void OnActorSpawned(AActor* InActor) {}
 
 	/** Calls the below function with the default camera distance. */
 	bool ViewportPositionToWorldPositionAndOrientation(EAvaViewportStatus InViewportStatus, const FVector2f& InViewportPosition,
@@ -106,6 +103,13 @@ public:
 
 	FViewport* GetViewport(EAvaViewportStatus InViewportStatus) const;
 
+	bool ShouldAutoRegister() const
+	{
+		return bAutoRegister;
+	}
+
+	virtual void OnRegisterTool(IAvalancheInteractiveToolsModule* InAITModule) {}
+
 protected:
 	/** Helper to create an Actor Factory of a given class */
 	template<typename InActorFactoryClass
@@ -138,8 +142,10 @@ protected:
 
 	bool bPerformingDefaultAction = false;
 
-	virtual bool CanActivate(bool bInReactivate) const;
-	virtual void Activate(bool bInReactivate);
+	bool bAutoRegister = true;
+
+	virtual bool CanActivate(const FString& InToolIdentifier, bool bInReactivate) const;
+	virtual void Activate(const FString& InToolIdentifier, bool bInReactivate);
 
 	/** Called when the tool is selected in the InteractiveTools tab. */
 	virtual void OnActivate();
@@ -176,23 +182,15 @@ protected:
 	bool ConditionalIdentityRotation() const;
 
 	/** Calls the other spawn actor with the centre of the last active viewport as the position. */
-	AActor* SpawnActor(TSubclassOf<AActor> InActorClass, bool bInPreview, FString* InActorLabelOverride = nullptr) const;
+	AActor* SpawnActor(TSubclassOf<AActor> InActorClass, bool bInPreview, FString* InActorLabelOverride = nullptr);
 
 	/** Spawns an actor of the given class at the correct distance from the camera in the given viewport position. */
 	virtual AActor* SpawnActor(TSubclassOf<AActor> InActorClass, EAvaViewportStatus InViewportStatus, 
-		const FVector2f& InViewportPosition, bool bInPreview, FString* InActorLabelOverride = nullptr) const;
+		const FVector2f& InViewportPosition, bool bInPreview, FString* InActorLabelOverride = nullptr);
 
 	void SetToolkitSettingsObject(UObject* InObject) const;
 
 	bool IsMotionDesignViewport() const;
 
 	bool ShouldForceDefaultAction() const;
-
-	bool ShouldUsePresetMenu() const;
-
-	void ShowPresetMenu();
-
-	static void OnPresetSelected(TStrongObjectPtr<UAvaInteractiveToolsToolBase> InToolPtr, FName InPresetName);
-
-	static void RegisterPresetMenu();
 };
diff --git a/Engine/Plugins/Experimental/Avalanche/Source/AvalancheInteractiveTools/Public/Tools/AvaInteractiveToolsToolPresetBase.h b/Engine/Plugins/Experimental/Avalanche/Source/AvalancheInteractiveTools/Public/Tools/AvaInteractiveToolsToolPresetBase.h
deleted file mode 100644
diff --git a/Engine/Plugins/Experimental/Avalanche/Source/AvalancheSVGEditor/Private/AvaSVGEditorModule.cpp b/Engine/Plugins/Experimental/Avalanche/Source/AvalancheSVGEditor/Private/AvaSVGEditorModule.cpp
--- a/Engine/Plugins/Experimental/Avalanche/Source/AvalancheSVGEditor/Private/AvaSVGEditorModule.cpp
+++ b/Engine/Plugins/Experimental/Avalanche/Source/AvalancheSVGEditor/Private/AvaSVGEditorModule.cpp
@@ -12,42 +12,21 @@
 #include "SVGImporter.h"
 #include "SVGShapesParentActor.h"
 #include "Subsystems/ActorModifierCoreSubsystem.h"
-#include "Tool/AvaSVGActorTool.h"
 
 #define LOCTEXT_NAMESPACE "AvalancheSVGEditorModule"
 
 void FAvaSVGEditorModule::StartupModule()
 {
-	FAvaInteractiveToolsDelegates::GetRegisterToolsDelegate().AddRaw(this, &FAvaSVGEditorModule::RegisterTools);
-
 	USVGEngineSubsystem::OnSVGActorSplit().BindRaw(this, &FAvaSVGEditorModule::OnSVGActorSplit);
 	USVGEngineSubsystem::OnSVGShapesUpdated().BindRaw(this, &FAvaSVGEditorModule::OnSVGShapesUpdated);
 }
 
 void FAvaSVGEditorModule::ShutdownModule()
 {
-	FAvaInteractiveToolsDelegates::GetRegisterToolsDelegate().RemoveAll(this);
-
 	USVGEngineSubsystem::OnSVGActorSplit().Unbind();
 	USVGEngineSubsystem::OnSVGShapesUpdated().Unbind();
 }
 
-void FAvaSVGEditorModule::RegisterTools(IAvalancheInteractiveToolsModule* InModule)
-{
-	if (!InModule)
-	{
-		return;
-	}
-
-	// Load svg module
-	FSVGImporterModule::Get();
-
-	InModule->RegisterTool(
-		IAvalancheInteractiveToolsModule::Get().CategoryNameActor,
-		GetDefault<UAvaSVGActorTool>()->GetToolParameters()
-	);
-}
-
 void FAvaSVGEditorModule::OnSVGActorSplit(ASVGShapesParentActor* InSVGShapesParent)
 {
 	if (!InSVGShapesParent)
diff --git a/Engine/Plugins/Experimental/Avalanche/Source/AvalancheSVGEditor/Private/AvaSVGEditorModule.h b/Engine/Plugins/Experimental/Avalanche/Source/AvalancheSVGEditor/Private/AvaSVGEditorModule.h
--- a/Engine/Plugins/Experimental/Avalanche/Source/AvalancheSVGEditor/Private/AvaSVGEditorModule.h
+++ b/Engine/Plugins/Experimental/Avalanche/Source/AvalancheSVGEditor/Private/AvaSVGEditorModule.h
@@ -19,8 +19,6 @@ public:
     virtual void ShutdownModule() override;
 
 private:
-	void RegisterTools(IAvalancheInteractiveToolsModule* InModule);
-
 	void OnSVGActorSplit(ASVGShapesParentActor* InSVGShapesParent);
 	void OnSVGShapesUpdated(AActor* InActor) const;
 };
diff --git a/Engine/Plugins/Experimental/Avalanche/Source/AvalancheSVGEditor/Private/Tool/AvaSVGActorTool.cpp b/Engine/Plugins/Experimental/Avalanche/Source/AvalancheSVGEditor/Private/Tool/AvaSVGActorTool.cpp
--- a/Engine/Plugins/Experimental/Avalanche/Source/AvalancheSVGEditor/Private/Tool/AvaSVGActorTool.cpp
+++ b/Engine/Plugins/Experimental/Avalanche/Source/AvalancheSVGEditor/Private/Tool/AvaSVGActorTool.cpp
@@ -5,6 +5,7 @@
 #include "Builders/AvaInteractiveToolsToolBuilder.h"
 #include "Factories/SVGActorFactory.h"
 #include "SVGActor.h"
+#include "SVGImporter.h"
 #include "SVGImporterEditorCommands.h"
 
 UAvaSVGActorTool::UAvaSVGActorTool()
@@ -12,14 +13,15 @@ UAvaSVGActorTool::UAvaSVGActorTool()
 	ActorClass = ASVGActor::StaticClass();
 }
 
-FName UAvaSVGActorTool::GetCategoryName()
+void UAvaSVGActorTool::OnRegisterTool(IAvalancheInteractiveToolsModule* InAITModule)
 {
-	return IAvalancheInteractiveToolsModule::CategoryNameActor;
-}
+	Super::OnRegisterTool(InAITModule);
 
-FAvaInteractiveToolsToolParameters UAvaSVGActorTool::GetToolParameters() const
-{
-	return {
+	// Load svg module
+	FSVGImporterModule::Get();
+
+	FAvaInteractiveToolsToolParameters ToolParameters =
+	{
 		FSVGImporterEditorCommands::GetExternal().SpawnSVGActor,
 		TEXT("SVG Actor Tool"),
 		6000,
@@ -31,4 +33,6 @@ FAvaInteractiveToolsToolParameters UAvaSVGActorTool::GetToolParameters() const
 		ActorClass,
 		CreateActorFactory<USVGActorFactory>()
 	};
+
+	InAITModule->RegisterTool(IAvalancheInteractiveToolsModule::CategoryNameActor, MoveTemp(ToolParameters));
 }
diff --git a/Engine/Plugins/Experimental/Avalanche/Source/AvalancheSVGEditor/Private/Tool/AvaSVGActorTool.h b/Engine/Plugins/Experimental/Avalanche/Source/AvalancheSVGEditor/Private/Tool/AvaSVGActorTool.h
--- a/Engine/Plugins/Experimental/Avalanche/Source/AvalancheSVGEditor/Private/Tool/AvaSVGActorTool.h
+++ b/Engine/Plugins/Experimental/Avalanche/Source/AvalancheSVGEditor/Private/Tool/AvaSVGActorTool.h
@@ -14,7 +14,6 @@ public:
 	UAvaSVGActorTool();
 
 	//~ Begin UAvaInteractiveToolsToolBase
-	virtual FName GetCategoryName() override;
-	virtual FAvaInteractiveToolsToolParameters GetToolParameters() const override;
+	virtual void OnRegisterTool(IAvalancheInteractiveToolsModule* InAITModule) override;
 	//~ End UAvaInteractiveToolsToolBase
 };
diff --git a/Engine/Plugins/Experimental/Avalanche/Source/AvalancheShapesEditor/Private/AvaShapesEditorCommands.cpp b/Engine/Plugins/Experimental/Avalanche/Source/AvalancheShapesEditor/Private/AvaShapesEditorCommands.cpp
--- a/Engine/Plugins/Experimental/Avalanche/Source/AvalancheShapesEditor/Private/AvaShapesEditorCommands.cpp
+++ b/Engine/Plugins/Experimental/Avalanche/Source/AvalancheShapesEditor/Private/AvaShapesEditorCommands.cpp
@@ -65,6 +65,12 @@ void FAvaShapesEditorCommands::Register2DCommands()
 		, EUserInterfaceActionType::ToggleButton
 		, FInputChord());
 
+	UI_COMMAND(Tool_Shape_Square
+		, "Square"
+		, "Create a 2D Square Shape in the viewport."
+		, EUserInterfaceActionType::ToggleButton
+		, FInputChord());
+
 	UI_COMMAND(Tool_Shape_Ring
 		, "Ring"
 		, "Create a 2D Ring Shape in the viewport."
@@ -86,6 +92,12 @@ void FAvaShapesEditorCommands::Register3DCommands()
 		, EUserInterfaceActionType::ToggleButton
 		, FInputChord());
 
+	UI_COMMAND(Tool_Shape_Cylinder
+		, "Cylinder"
+		, "Create a Cylinder Shape in the viewport."
+		, EUserInterfaceActionType::ToggleButton
+		, FInputChord());
+
 	UI_COMMAND(Tool_Shape_Cube
 		, "Cube"
 		, "Create a Cube Shape in the viewport."
diff --git a/Engine/Plugins/Experimental/Avalanche/Source/AvalancheShapesEditor/Private/AvaShapesEditorCommands.h b/Engine/Plugins/Experimental/Avalanche/Source/AvalancheShapesEditor/Private/AvaShapesEditorCommands.h
--- a/Engine/Plugins/Experimental/Avalanche/Source/AvalancheShapesEditor/Private/AvaShapesEditorCommands.h
+++ b/Engine/Plugins/Experimental/Avalanche/Source/AvalancheShapesEditor/Private/AvaShapesEditorCommands.h
@@ -24,11 +24,13 @@ public:
 	TSharedPtr<FUICommandInfo> Tool_Shape_Line;
 	TSharedPtr<FUICommandInfo> Tool_Shape_NGon;
 	TSharedPtr<FUICommandInfo> Tool_Shape_Rectangle;
+	TSharedPtr<FUICommandInfo> Tool_Shape_Square;
 	TSharedPtr<FUICommandInfo> Tool_Shape_Ring;
 	TSharedPtr<FUICommandInfo> Tool_Shape_Star;
 
 	// 3D Tools
 	TSharedPtr<FUICommandInfo> Tool_Shape_Cone;
+	TSharedPtr<FUICommandInfo> Tool_Shape_Cylinder;
 	TSharedPtr<FUICommandInfo> Tool_Shape_Cube;
 	TSharedPtr<FUICommandInfo> Tool_Shape_Sphere;
 	TSharedPtr<FUICommandInfo> Tool_Shape_Torus;
diff --git a/Engine/Plugins/Experimental/Avalanche/Source/AvalancheShapesEditor/Private/AvaShapesEditorStyle.cpp b/Engine/Plugins/Experimental/Avalanche/Source/AvalancheShapesEditor/Private/AvaShapesEditorStyle.cpp
--- a/Engine/Plugins/Experimental/Avalanche/Source/AvalancheShapesEditor/Private/AvaShapesEditorStyle.cpp
+++ b/Engine/Plugins/Experimental/Avalanche/Source/AvalancheShapesEditor/Private/AvaShapesEditorStyle.cpp
@@ -12,7 +12,6 @@ FAvaShapesEditorStyle::FAvaShapesEditorStyle()
 	: FSlateStyleSet(TEXT("AvaShapesEditor"))
 {
 	const FVector2f Icon16x16(16.0f, 16.0f);
-	const FVector2f Icon20x20(20.0f, 20.0f);
 
 	const TSharedPtr<IPlugin> Plugin = IPluginManager::Get().FindPlugin(UE_PLUGIN_NAME);
 	check(Plugin.IsValid());
@@ -20,53 +19,40 @@ FAvaShapesEditorStyle::FAvaShapesEditorStyle()
 	SetContentRoot(Plugin->GetBaseDir() / TEXT("Resources"));
 
 	// Class Icons
-	Set("ClassIcon.AvaShape2DArrowDynamicMesh", new IMAGE_BRUSH("Icons/ToolboxIcons/arrow", Icon16x16));
-	Set("ClassIcon.AvaShapeChevronDynamicMesh", new IMAGE_BRUSH("Icons/ToolboxIcons/chevron", Icon16x16));
-	Set("ClassIcon.AvaShapeConeDynamicMesh", new IMAGE_BRUSH("Icons/ToolboxIcons/cone", Icon16x16));
-	Set("ClassIcon.AvaShapeCubeDynamicMesh", new IMAGE_BRUSH("Icons/ToolboxIcons/cube", Icon16x16));
-	Set("ClassIcon.AvaShapeRoundedPolygonDynamicMesh", new IMAGE_BRUSH_SVG("Icons/ToolboxIcons/Toolbar_Primitive", Icon16x16));
-	Set("ClassIcon.AvaShapeSphereDynamicMesh", new IMAGE_BRUSH("Icons/ToolboxIcons/sphere", Icon16x16));
-	Set("ClassIcon.AvaShapeTorusDynamicMesh", new IMAGE_BRUSH("Icons/ToolboxIcons/torus", Icon16x16));
-	Set("ClassIcon.AvaShapeStarDynamicMesh", new IMAGE_BRUSH_SVG("Icons/ToolboxIcons/Favorite", Icon16x16));
-	Set("ClassIcon.AvaShapeLineDynamicMesh", new IMAGE_BRUSH_SVG("Icons/ToolboxIcons/minus", Icon16x16));
-	Set("ClassIcon.AvaShapeRectangleDynamicMesh", new IMAGE_BRUSH("Icons/ToolboxIcons/rectangle", Icon16x16));
-	Set("ClassIcon.AvaShapeNGonDynamicMesh", new IMAGE_BRUSH("Icons/ToolboxIcons/regularpolygon", Icon16x16));
+	Set("ClassIcon.AvaShape2DArrowDynamicMesh",          new IMAGE_BRUSH("Icons/ToolboxIcons/arrow", Icon16x16));
+	Set("ClassIcon.AvaShapeChevronDynamicMesh",          new IMAGE_BRUSH("Icons/ToolboxIcons/chevron", Icon16x16));
+	Set("ClassIcon.AvaShapeConeDynamicMesh",             new IMAGE_BRUSH("Icons/ToolboxIcons/cone", Icon16x16));
+	Set("ClassIcon.AvaShapeCubeDynamicMesh",             new IMAGE_BRUSH("Icons/ToolboxIcons/cube", Icon16x16));
+	Set("ClassIcon.AvaShapeRoundedPolygonDynamicMesh",   new IMAGE_BRUSH_SVG("Icons/ToolboxIcons/Toolbar_Primitive", Icon16x16));
+	Set("ClassIcon.AvaShapeSphereDynamicMesh",           new IMAGE_BRUSH("Icons/ToolboxIcons/sphere", Icon16x16));
+	Set("ClassIcon.AvaShapeTorusDynamicMesh",            new IMAGE_BRUSH("Icons/ToolboxIcons/torus", Icon16x16));
+	Set("ClassIcon.AvaShapeStarDynamicMesh",             new IMAGE_BRUSH_SVG("Icons/ToolboxIcons/Favorite", Icon16x16));
+	Set("ClassIcon.AvaShapeLineDynamicMesh",             new IMAGE_BRUSH_SVG("Icons/ToolboxIcons/minus", Icon16x16));
+	Set("ClassIcon.AvaShapeRectangleDynamicMesh",        new IMAGE_BRUSH("Icons/ToolboxIcons/rectangle", Icon16x16));
+	Set("ClassIcon.AvaShapeNGonDynamicMesh",             new IMAGE_BRUSH("Icons/ToolboxIcons/regularpolygon", Icon16x16));
 	Set("ClassIcon.AvaShapeIrregularPolygonDynamicMesh", new IMAGE_BRUSH("Icons/ToolboxIcons/irregularpolygon", Icon16x16));
-	Set("ClassIcon.AvaShapeEllipseDynamicMesh", new IMAGE_BRUSH("Icons/ToolboxIcons/circle", Icon16x16));
-	Set("ClassIcon.AvaShapeRingDynamicMesh", new IMAGE_BRUSH("Icons/ToolboxIcons/ring", Icon16x16));
-	Set("ClassIcon.AvaShapeShapeActor", new IMAGE_BRUSH("Icons/ToolboxIcons/regularpolygon", Icon16x16));
-
-	// 2D Tools
-	Set("AvaShapesEditor.Tool_Shape_2DArrow",       new IMAGE_BRUSH("Icons/ToolboxIcons/arrow", Icon20x20));
-	Set("AvaShapesEditor.Tool_Shape_Chevron",       new IMAGE_BRUSH("Icons/ToolboxIcons/chevron", Icon20x20));
-	Set("AvaShapesEditor.Tool_Shape_Ellipse",       new IMAGE_BRUSH("Icons/ToolboxIcons/circle", Icon20x20));
-	Set("AvaShapesEditor.Tool_Shape_IrregularPoly", new IMAGE_BRUSH("Icons/ToolboxIcons/irregularpolygon", Icon20x20));
-	Set("AvaShapesEditor.Tool_Shape_Line",          new IMAGE_BRUSH_SVG("Icons/ToolboxIcons/minus", Icon20x20));
-	Set("AvaShapesEditor.Tool_Shape_NGon",          new IMAGE_BRUSH("Icons/ToolboxIcons/regularpolygon", Icon20x20));
-	Set("AvaShapesEditor.Tool_Shape_Rectangle",     new IMAGE_BRUSH("Icons/ToolboxIcons/rectangle", Icon20x20));
-	Set("AvaShapesEditor.Tool_Shape_Ring",          new IMAGE_BRUSH("Icons/ToolboxIcons/ring", Icon20x20));
-	Set("AvaShapesEditor.Tool_Shape_Star",          new IMAGE_BRUSH_SVG("Icons/ToolboxIcons/Favorite", Icon20x20));
-
-	Set("Tool_Shape_2DArrow",       new IMAGE_BRUSH("Icons/ToolboxIcons/arrow", Icon20x20));
-	Set("Tool_Shape_Chevron",       new IMAGE_BRUSH("Icons/ToolboxIcons/chevron", Icon20x20));
-	Set("Tool_Shape_Ellipse",       new IMAGE_BRUSH("Icons/ToolboxIcons/circle", Icon20x20));
-	Set("Tool_Shape_IrregularPoly", new IMAGE_BRUSH("Icons/ToolboxIcons/irregularpolygon", Icon20x20));
-	Set("Tool_Shape_Line",          new IMAGE_BRUSH_SVG("Icons/ToolboxIcons/minus", Icon20x20));
-	Set("Tool_Shape_NGon",          new IMAGE_BRUSH("Icons/ToolboxIcons/regularpolygon", Icon20x20));
-	Set("Tool_Shape_Rectangle",     new IMAGE_BRUSH("Icons/ToolboxIcons/rectangle", Icon20x20));
-	Set("Tool_Shape_Ring",          new IMAGE_BRUSH("Icons/ToolboxIcons/ring", Icon20x20));
-	Set("Tool_Shape_Star",          new IMAGE_BRUSH_SVG("Icons/ToolboxIcons/Favorite", Icon20x20));
-
-	// 3D Tools
-	Set("AvaShapesEditor.Tool_Shape_Cone",   new IMAGE_BRUSH("Icons/ToolboxIcons/cone", Icon20x20));
-	Set("AvaShapesEditor.Tool_Shape_Cube",   new IMAGE_BRUSH("Icons/ToolboxIcons/cube", Icon20x20));
-	Set("AvaShapesEditor.Tool_Shape_Sphere", new IMAGE_BRUSH("Icons/ToolboxIcons/sphere", Icon20x20));
-	Set("AvaShapesEditor.Tool_Shape_Torus",  new IMAGE_BRUSH("Icons/ToolboxIcons/torus", Icon20x20));
-
-	Set("Tool_Shape_Cone",   new IMAGE_BRUSH("Icons/ToolboxIcons/cone", Icon20x20));
-	Set("Tool_Shape_Cube",   new IMAGE_BRUSH("Icons/ToolboxIcons/cube", Icon20x20));
-	Set("Tool_Shape_Sphere", new IMAGE_BRUSH("Icons/ToolboxIcons/sphere", Icon20x20));
-	Set("Tool_Shape_Torus",  new IMAGE_BRUSH("Icons/ToolboxIcons/torus", Icon20x20));
+	Set("ClassIcon.AvaShapeEllipseDynamicMesh",          new IMAGE_BRUSH("Icons/ToolboxIcons/circle", Icon16x16));
+	Set("ClassIcon.AvaShapeRingDynamicMesh",             new IMAGE_BRUSH("Icons/ToolboxIcons/ring", Icon16x16));
+	Set("ClassIcon.AvaShapeShapeActor",                  new IMAGE_BRUSH("Icons/ToolboxIcons/regularpolygon", Icon16x16));
+
+	// 2D Tools Commands
+	Set("AvaShapesEditor.Tool_Shape_2DArrow",       new IMAGE_BRUSH("Icons/ToolboxIcons/arrow", Icon16x16));
+	Set("AvaShapesEditor.Tool_Shape_Chevron",       new IMAGE_BRUSH("Icons/ToolboxIcons/chevron", Icon16x16));
+	Set("AvaShapesEditor.Tool_Shape_Ellipse",       new IMAGE_BRUSH("Icons/ToolboxIcons/circle", Icon16x16));
+	Set("AvaShapesEditor.Tool_Shape_IrregularPoly", new IMAGE_BRUSH("Icons/ToolboxIcons/irregularpolygon", Icon16x16));
+	Set("AvaShapesEditor.Tool_Shape_Line",          new IMAGE_BRUSH_SVG("Icons/ToolboxIcons/minus", Icon16x16));
+	Set("AvaShapesEditor.Tool_Shape_NGon",          new IMAGE_BRUSH("Icons/ToolboxIcons/regularpolygon", Icon16x16));
+	Set("AvaShapesEditor.Tool_Shape_Rectangle",     new IMAGE_BRUSH("Icons/ToolboxIcons/rectangle", Icon16x16));
+	Set("AvaShapesEditor.Tool_Shape_Square",        new IMAGE_BRUSH("Icons/ToolboxIcons/roundedsquare", Icon16x16));
+	Set("AvaShapesEditor.Tool_Shape_Ring",          new IMAGE_BRUSH("Icons/ToolboxIcons/ring", Icon16x16));
+	Set("AvaShapesEditor.Tool_Shape_Star",          new IMAGE_BRUSH_SVG("Icons/ToolboxIcons/Favorite", Icon16x16));
+
+	// 3D Tools Commands
+	Set("AvaShapesEditor.Tool_Shape_Cone",       new IMAGE_BRUSH("Icons/ToolboxIcons/cone", Icon16x16));
+	Set("AvaShapesEditor.Tool_Shape_Cylinder",   new IMAGE_BRUSH("Icons/ToolboxIcons/cylinder", Icon16x16));
+	Set("AvaShapesEditor.Tool_Shape_Cube",       new IMAGE_BRUSH("Icons/ToolboxIcons/cube", Icon16x16));
+	Set("AvaShapesEditor.Tool_Shape_Sphere",     new IMAGE_BRUSH("Icons/ToolboxIcons/sphere", Icon16x16));
+	Set("AvaShapesEditor.Tool_Shape_Torus",      new IMAGE_BRUSH("Icons/ToolboxIcons/torus", Icon16x16));
 
 	FSlateStyleRegistry::RegisterSlateStyle(*this);
 }
diff --git a/Engine/Plugins/Experimental/Avalanche/Source/AvalancheShapesEditor/Private/AvalancheShapesEditorModule.cpp b/Engine/Plugins/Experimental/Avalanche/Source/AvalancheShapesEditor/Private/AvalancheShapesEditorModule.cpp
--- a/Engine/Plugins/Experimental/Avalanche/Source/AvalancheShapesEditor/Private/AvalancheShapesEditorModule.cpp
+++ b/Engine/Plugins/Experimental/Avalanche/Source/AvalancheShapesEditor/Private/AvalancheShapesEditorModule.cpp
@@ -64,7 +64,7 @@
 void FAvalancheShapesEditorModule::StartupModule()
 {
 	FAvaShapesEditorCommands::Register();
-	FAvaInteractiveToolsDelegates::GetRegisterToolsDelegate().AddRaw(this, &FAvalancheShapesEditorModule::RegisterShapeTools);
+
 	FCoreDelegates::OnPostEngineInit.AddRaw(this, &FAvalancheShapesEditorModule::RegisterVisualizers);
 
 	IAvaComponentVisualizersSettings* AvaVisSettings = IAvalancheComponentVisualizersModule::Get().GetSettings();
@@ -100,7 +100,7 @@ void FAvalancheShapesEditorModule::StartupModule()
 void FAvalancheShapesEditorModule::ShutdownModule()
 {
 	FAvaShapesEditorCommands::Unregister();
-	FAvaInteractiveToolsDelegates::GetRegisterToolsDelegate().RemoveAll(this);
+
 	FCoreDelegates::OnPostEngineInit.RemoveAll(this);
 
 	if (FModuleManager::Get().IsModuleLoaded("Sequencer"))
@@ -111,76 +111,6 @@ void FAvalancheShapesEditorModule::ShutdownModule()
 	}
 }
 
-void FAvalancheShapesEditorModule::RegisterShapeTools(IAvalancheInteractiveToolsModule* InModule)
-{
-	// 2D Tools
-	InModule->RegisterTool(
-		IAvalancheInteractiveToolsModule::Get().CategoryName2D,
-		GetDefault<UAvaShapesEditorShapeToolRectangle>()->GetToolParameters()
-	);
-
-	InModule->RegisterTool(
-		IAvalancheInteractiveToolsModule::Get().CategoryName2D,
-		GetDefault<UAvaShapesEditorShapeToolEllipse>()->GetToolParameters()
-	);
-
-	InModule->RegisterTool(
-		IAvalancheInteractiveToolsModule::Get().CategoryName2D,
-		GetDefault<UAvaShapesEditorShapeToolNGon>()->GetToolParameters()
-	);
-
-	InModule->RegisterTool(
-		IAvalancheInteractiveToolsModule::Get().CategoryName2D,
-		GetDefault<UAvaShapesEditorShapeToolIrregularPoly>()->GetToolParameters()
-	);
-
-	InModule->RegisterTool(
-		IAvalancheInteractiveToolsModule::Get().CategoryName2D,
-		GetDefault<UAvaShapesEditorShapeToolRing>()->GetToolParameters()
-	);
-
-	InModule->RegisterTool(
-		IAvalancheInteractiveToolsModule::Get().CategoryName2D,
-		GetDefault<UAvaShapesEditorShapeToolStar>()->GetToolParameters()
-	);
-
-	InModule->RegisterTool(
-		IAvalancheInteractiveToolsModule::Get().CategoryName2D,
-		GetDefault<UAvaShapesEditorShapeTool2DArrow>()->GetToolParameters()
-	);
-
-	InModule->RegisterTool(
-		IAvalancheInteractiveToolsModule::Get().CategoryName2D,
-		GetDefault<UAvaShapesEditorShapeToolChevron>()->GetToolParameters()
-	);
-
-	InModule->RegisterTool(
-		IAvalancheInteractiveToolsModule::Get().CategoryName2D,
-		GetDefault<UAvaShapesEditorShapeToolLine>()->GetToolParameters()
-	);
-
-	// 3D Tools
-	InModule->RegisterTool(
-		IAvalancheInteractiveToolsModule::Get().CategoryName3D,
-		GetDefault<UAvaShapesEditorShapeToolCube>()->GetToolParameters()
-	);
-
-	InModule->RegisterTool(
-		IAvalancheInteractiveToolsModule::Get().CategoryName3D,
-		GetDefault<UAvaShapesEditorShapeToolSphere>()->GetToolParameters()
-	);
-
-	InModule->RegisterTool(
-		IAvalancheInteractiveToolsModule::Get().CategoryName3D,
-		GetDefault<UAvaShapesEditorShapeToolCone>()->GetToolParameters()
-	);
-
-	InModule->RegisterTool(
-		IAvalancheInteractiveToolsModule::Get().CategoryName3D,
-		GetDefault<UAvaShapesEditorShapeToolTorus>()->GetToolParameters()
-	);
-}
-
 void FAvalancheShapesEditorModule::RegisterVisualizers()
 {
 	// 2D Visualizers
diff --git a/Engine/Plugins/Experimental/Avalanche/Source/AvalancheShapesEditor/Private/AvalancheShapesEditorModule.h b/Engine/Plugins/Experimental/Avalanche/Source/AvalancheShapesEditor/Private/AvalancheShapesEditorModule.h
--- a/Engine/Plugins/Experimental/Avalanche/Source/AvalancheShapesEditor/Private/AvalancheShapesEditorModule.h
+++ b/Engine/Plugins/Experimental/Avalanche/Source/AvalancheShapesEditor/Private/AvalancheShapesEditorModule.h
@@ -19,7 +19,6 @@ class FAvalancheShapesEditorModule : public IModuleInterface
 
 	TArray<TSharedPtr<FComponentVisualizer>> Visualizers;
 
-	void RegisterShapeTools(IAvalancheInteractiveToolsModule* InModule);
 	void RegisterVisualizers();
 
 	FDelegateHandle TrackEditorHandle;
diff --git a/Engine/Plugins/Experimental/Avalanche/Source/AvalancheShapesEditor/Private/Tools/AvaShapesEditorShapeAreaToolBase.cpp b/Engine/Plugins/Experimental/Avalanche/Source/AvalancheShapesEditor/Private/Tools/AvaShapesEditorShapeAreaToolBase.cpp
--- a/Engine/Plugins/Experimental/Avalanche/Source/AvalancheShapesEditor/Private/Tools/AvaShapesEditorShapeAreaToolBase.cpp
+++ b/Engine/Plugins/Experimental/Avalanche/Source/AvalancheShapesEditor/Private/Tools/AvaShapesEditorShapeAreaToolBase.cpp
@@ -17,7 +17,7 @@ UAvaShapesEditorShapeAreaToolBase::UAvaShapesEditorShapeAreaToolBase()
 }
 
 AActor* UAvaShapesEditorShapeAreaToolBase::SpawnActor(TSubclassOf<AActor> InActorClass, EAvaViewportStatus InViewportStatus,
-	const FVector2f& InViewportPosition, bool bInPreview, FString* InActorLabelOverride) const
+	const FVector2f& InViewportPosition, bool bInPreview, FString* InActorLabelOverride)
 {
 	AActor* NewActor = Super::SpawnActor(InActorClass, InViewportStatus, InViewportPosition, bInPreview, InActorLabelOverride);
 
diff --git a/Engine/Plugins/Experimental/Avalanche/Source/AvalancheShapesEditor/Private/Tools/AvaShapesEditorShapeAreaToolBase.h b/Engine/Plugins/Experimental/Avalanche/Source/AvalancheShapesEditor/Private/Tools/AvaShapesEditorShapeAreaToolBase.h
--- a/Engine/Plugins/Experimental/Avalanche/Source/AvalancheShapesEditor/Private/Tools/AvaShapesEditorShapeAreaToolBase.h
+++ b/Engine/Plugins/Experimental/Avalanche/Source/AvalancheShapesEditor/Private/Tools/AvaShapesEditorShapeAreaToolBase.h
@@ -29,7 +29,6 @@ protected:
 	void SelectActor(AActor* InActor) const;
 
 	//~ Begin UAvaInteractiveToolsToolBase
-	virtual AActor* SpawnActor(TSubclassOf<AActor> InActorClass, EAvaViewportStatus InViewportStatus, 
-		const FVector2f& InViewportPosition, bool bInPreview, FString* InActorLabelOverride = nullptr) const override;
+	virtual AActor* SpawnActor(TSubclassOf<AActor> InActorClass, EAvaViewportStatus InViewportStatus, const FVector2f& InViewportPosition, bool bInPreview, FString* InActorLabelOverride = nullptr) override;
 	//~ End UAvaInteractiveToolsToolBase
 };
diff --git a/Engine/Plugins/Experimental/Avalanche/Source/AvalancheShapesEditor/Private/Tools/AvaShapesEditorShapeTool2DArrow.cpp b/Engine/Plugins/Experimental/Avalanche/Source/AvalancheShapesEditor/Private/Tools/AvaShapesEditorShapeTool2DArrow.cpp
--- a/Engine/Plugins/Experimental/Avalanche/Source/AvalancheShapesEditor/Private/Tools/AvaShapesEditorShapeTool2DArrow.cpp
+++ b/Engine/Plugins/Experimental/Avalanche/Source/AvalancheShapesEditor/Private/Tools/AvaShapesEditorShapeTool2DArrow.cpp
@@ -11,14 +11,12 @@ UAvaShapesEditorShapeTool2DArrow::UAvaShapesEditorShapeTool2DArrow()
 	ShapeClass = UAvaShape2DArrowDynamicMesh::StaticClass();
 }
 
-FName UAvaShapesEditorShapeTool2DArrow::GetCategoryName()
+void UAvaShapesEditorShapeTool2DArrow::OnRegisterTool(IAvalancheInteractiveToolsModule* InAITModule)
 {
-	return IAvalancheInteractiveToolsModule::CategoryName2D;
-}
+	Super::OnRegisterTool(InAITModule);
 
-FAvaInteractiveToolsToolParameters UAvaShapesEditorShapeTool2DArrow::GetToolParameters() const
-{
-	return {
+	FAvaInteractiveToolsToolParameters ToolParameters = 
+	{
 		FAvaShapesEditorCommands::Get().Tool_Shape_2DArrow,
 		TEXT("Parametric 2D Arrow Tool"),
 		7000,
@@ -30,6 +28,8 @@ FAvaInteractiveToolsToolParameters UAvaShapesEditorShapeTool2DArrow::GetToolPara
 		nullptr,
 		CreateFactory<UAvaShape2DArrowDynamicMesh>()
 	};
+
+	InAITModule->RegisterTool(IAvalancheInteractiveToolsModule::CategoryName2D, MoveTemp(ToolParameters));
 }
 
 void UAvaShapesEditorShapeTool2DArrow::InitShape(UAvaShapeDynamicMeshBase* InShape) const
diff --git a/Engine/Plugins/Experimental/Avalanche/Source/AvalancheShapesEditor/Private/Tools/AvaShapesEditorShapeTool2DArrow.h b/Engine/Plugins/Experimental/Avalanche/Source/AvalancheShapesEditor/Private/Tools/AvaShapesEditorShapeTool2DArrow.h
--- a/Engine/Plugins/Experimental/Avalanche/Source/AvalancheShapesEditor/Private/Tools/AvaShapesEditorShapeTool2DArrow.h
+++ b/Engine/Plugins/Experimental/Avalanche/Source/AvalancheShapesEditor/Private/Tools/AvaShapesEditorShapeTool2DArrow.h
@@ -14,8 +14,7 @@ public:
 	UAvaShapesEditorShapeTool2DArrow();
 
 	//~ Begin UAvaInteractiveToolsToolBase
-	virtual FName GetCategoryName() override;
-	virtual FAvaInteractiveToolsToolParameters GetToolParameters() const override;
+	virtual void OnRegisterTool(IAvalancheInteractiveToolsModule* InAITModule) override;
 	//~ End UAvaInteractiveToolsToolBase
 
 protected:
diff --git a/Engine/Plugins/Experimental/Avalanche/Source/AvalancheShapesEditor/Private/Tools/AvaShapesEditorShapeToolBase.cpp b/Engine/Plugins/Experimental/Avalanche/Source/AvalancheShapesEditor/Private/Tools/AvaShapesEditorShapeToolBase.cpp
--- a/Engine/Plugins/Experimental/Avalanche/Source/AvalancheShapesEditor/Private/Tools/AvaShapesEditorShapeToolBase.cpp
+++ b/Engine/Plugins/Experimental/Avalanche/Source/AvalancheShapesEditor/Private/Tools/AvaShapesEditorShapeToolBase.cpp
@@ -27,7 +27,7 @@ bool UAvaShapesEditorShapeToolBase::OnBegin()
 }
 
 AActor* UAvaShapesEditorShapeToolBase::SpawnActor(TSubclassOf<AActor> InActorClass, EAvaViewportStatus InViewportStatus,
-	const FVector2f& InViewportPosition, bool bInPreview, FString* InActorLabelOverride) const
+	const FVector2f& InViewportPosition, bool bInPreview, FString* InActorLabelOverride)
 {
 	FString ActorLabelOverride = (InActorLabelOverride && !InActorLabelOverride->IsEmpty())
 		? *InActorLabelOverride
diff --git a/Engine/Plugins/Experimental/Avalanche/Source/AvalancheShapesEditor/Private/Tools/AvaShapesEditorShapeToolBase.h b/Engine/Plugins/Experimental/Avalanche/Source/AvalancheShapesEditor/Private/Tools/AvaShapesEditorShapeToolBase.h
--- a/Engine/Plugins/Experimental/Avalanche/Source/AvalancheShapesEditor/Private/Tools/AvaShapesEditorShapeToolBase.h
+++ b/Engine/Plugins/Experimental/Avalanche/Source/AvalancheShapesEditor/Private/Tools/AvaShapesEditorShapeToolBase.h
@@ -50,7 +50,7 @@ protected:
 	virtual bool UseIdentityLocation() const override;
 
 	virtual AActor* SpawnActor(TSubclassOf<AActor> InActorClass, EAvaViewportStatus InViewportStatus, 
-		const FVector2f& InViewportPosition, bool bInPreview, FString* InActorLabelOverride = nullptr) const override;
+		const FVector2f& InViewportPosition, bool bInPreview, FString* InActorLabelOverride = nullptr) override;
 	//~ End UAvaInteractiveToolsToolBase
 
 	virtual FString GetActorNameOverride() const;
diff --git a/Engine/Plugins/Experimental/Avalanche/Source/AvalancheShapesEditor/Private/Tools/AvaShapesEditorShapeToolChevron.cpp b/Engine/Plugins/Experimental/Avalanche/Source/AvalancheShapesEditor/Private/Tools/AvaShapesEditorShapeToolChevron.cpp
--- a/Engine/Plugins/Experimental/Avalanche/Source/AvalancheShapesEditor/Private/Tools/AvaShapesEditorShapeToolChevron.cpp
+++ b/Engine/Plugins/Experimental/Avalanche/Source/AvalancheShapesEditor/Private/Tools/AvaShapesEditorShapeToolChevron.cpp
@@ -11,14 +11,12 @@ UAvaShapesEditorShapeToolChevron::UAvaShapesEditorShapeToolChevron()
 	ShapeClass = UAvaShapeChevronDynamicMesh::StaticClass();
 }
 
-FName UAvaShapesEditorShapeToolChevron::GetCategoryName()
+void UAvaShapesEditorShapeToolChevron::OnRegisterTool(IAvalancheInteractiveToolsModule* InAITModule)
 {
-	return IAvalancheInteractiveToolsModule::CategoryName2D;
-}
+	Super::OnRegisterTool(InAITModule);
 
-FAvaInteractiveToolsToolParameters UAvaShapesEditorShapeToolChevron::GetToolParameters() const
-{
-	return {
+	FAvaInteractiveToolsToolParameters ToolParameters = 
+	{
 		FAvaShapesEditorCommands::Get().Tool_Shape_Chevron,
 		TEXT("Parametric Chevron Tool"),
 		8000,
@@ -30,6 +28,8 @@ FAvaInteractiveToolsToolParameters UAvaShapesEditorShapeToolChevron::GetToolPara
 		nullptr,
 		CreateFactory<UAvaShapeChevronDynamicMesh>()
 	};
+
+	InAITModule->RegisterTool(IAvalancheInteractiveToolsModule::CategoryName2D, MoveTemp(ToolParameters));
 }
 
 void UAvaShapesEditorShapeToolChevron::InitShape(UAvaShapeDynamicMeshBase* InShape) const
diff --git a/Engine/Plugins/Experimental/Avalanche/Source/AvalancheShapesEditor/Private/Tools/AvaShapesEditorShapeToolChevron.h b/Engine/Plugins/Experimental/Avalanche/Source/AvalancheShapesEditor/Private/Tools/AvaShapesEditorShapeToolChevron.h
--- a/Engine/Plugins/Experimental/Avalanche/Source/AvalancheShapesEditor/Private/Tools/AvaShapesEditorShapeToolChevron.h
+++ b/Engine/Plugins/Experimental/Avalanche/Source/AvalancheShapesEditor/Private/Tools/AvaShapesEditorShapeToolChevron.h
@@ -14,8 +14,7 @@ public:
 	UAvaShapesEditorShapeToolChevron();
 
 	//~ Begin UAvaInteractiveToolsToolBase
-	virtual FName GetCategoryName() override;
-	virtual FAvaInteractiveToolsToolParameters GetToolParameters() const override;
+	virtual void OnRegisterTool(IAvalancheInteractiveToolsModule* InAITModule) override;
 	//~ End UAvaInteractiveToolsToolBase
 
 protected:
diff --git a/Engine/Plugins/Experimental/Avalanche/Source/AvalancheShapesEditor/Private/Tools/AvaShapesEditorShapeToolCone.cpp b/Engine/Plugins/Experimental/Avalanche/Source/AvalancheShapesEditor/Private/Tools/AvaShapesEditorShapeToolCone.cpp
--- a/Engine/Plugins/Experimental/Avalanche/Source/AvalancheShapesEditor/Private/Tools/AvaShapesEditorShapeToolCone.cpp
+++ b/Engine/Plugins/Experimental/Avalanche/Source/AvalancheShapesEditor/Private/Tools/AvaShapesEditorShapeToolCone.cpp
@@ -6,21 +6,38 @@
 #include "DynamicMeshes/AvaShapeConeDynMesh.h"
 #include "UObject/ConstructorHelpers.h"
 
+const TCHAR* UAvaShapesEditorShapeToolCone::ConeToolIdentifier = TEXT("Parametric Cone Tool");
+const TCHAR* UAvaShapesEditorShapeToolCone::CylinderToolIdentifier = TEXT("Parametric Cylinder Tool");
+
 UAvaShapesEditorShapeToolCone::UAvaShapesEditorShapeToolCone()
 {
 	ShapeClass = UAvaShapeConeDynamicMesh::StaticClass();
 }
 
-FName UAvaShapesEditorShapeToolCone::GetCategoryName()
+void UAvaShapesEditorShapeToolCone::OnRegisterTool(IAvalancheInteractiveToolsModule* InAITModule)
 {
-	return IAvalancheInteractiveToolsModule::CategoryName3D;
-}
+	Super::OnRegisterTool(InAITModule);
 
-FAvaInteractiveToolsToolParameters UAvaShapesEditorShapeToolCone::GetToolParameters() const
-{
-	return {
+	FAvaInteractiveToolsToolParameters ConeToolParameters = 
+	{
 		FAvaShapesEditorCommands::Get().Tool_Shape_Cone,
-		TEXT("Parametric Cone Tool"),
+		ConeToolIdentifier,
+		3000,
+		FAvalancheInteractiveToolsCreateBuilder::CreateLambda(
+			[](UEdMode* InEdMode)
+			{
+				return UAvaInteractiveToolsToolBuilder::CreateToolBuilder<UAvaShapesEditorShapeToolCone>(InEdMode);
+			}),
+		nullptr,
+		CreateFactory<UAvaShapeConeDynamicMesh>()
+	};
+
+	InAITModule->RegisterTool(IAvalancheInteractiveToolsModule::CategoryName3D, MoveTemp(ConeToolParameters));
+	
+	FAvaInteractiveToolsToolParameters CylinderToolParameters = 
+	{
+		FAvaShapesEditorCommands::Get().Tool_Shape_Cylinder,
+		CylinderToolIdentifier,
 		3000,
 		FAvalancheInteractiveToolsCreateBuilder::CreateLambda(
 			[](UEdMode* InEdMode)
@@ -30,14 +47,22 @@ FAvaInteractiveToolsToolParameters UAvaShapesEditorShapeToolCone::GetToolParamet
 		nullptr,
 		CreateFactory<UAvaShapeConeDynamicMesh>()
 	};
+
+	InAITModule->RegisterTool(IAvalancheInteractiveToolsModule::CategoryName3D, MoveTemp(CylinderToolParameters));
 }
 
 void UAvaShapesEditorShapeToolCone::InitShape(UAvaShapeDynamicMeshBase* InShape) const
 {
 	UAvaShapeConeDynamicMesh* Cone = Cast<UAvaShapeConeDynamicMesh>(InShape);
+
 	check(Cone);
 
 	Cone->SetSize3D({DefaultDim, DefaultDim, DefaultDim});
 
+	if (GetActiveToolIdentifier().Equals(CylinderToolIdentifier))
+	{
+		Cone->SetTopRadius(1.f);
+	}
+
 	Super::InitShape(InShape);
 }
diff --git a/Engine/Plugins/Experimental/Avalanche/Source/AvalancheShapesEditor/Private/Tools/AvaShapesEditorShapeToolCone.h b/Engine/Plugins/Experimental/Avalanche/Source/AvalancheShapesEditor/Private/Tools/AvaShapesEditorShapeToolCone.h
--- a/Engine/Plugins/Experimental/Avalanche/Source/AvalancheShapesEditor/Private/Tools/AvaShapesEditorShapeToolCone.h
+++ b/Engine/Plugins/Experimental/Avalanche/Source/AvalancheShapesEditor/Private/Tools/AvaShapesEditorShapeToolCone.h
@@ -11,11 +11,13 @@ class UAvaShapesEditorShapeToolCone : public UAvaShapesEditorShapeAreaToolBase
 	GENERATED_BODY()
 
 public:
+	static const TCHAR* ConeToolIdentifier;
+	static const TCHAR* CylinderToolIdentifier;
+
 	UAvaShapesEditorShapeToolCone();
 
 	//~ Begin UAvaInteractiveToolsToolBase
-	virtual FName GetCategoryName() override;
-	virtual FAvaInteractiveToolsToolParameters GetToolParameters() const override;
+	virtual void OnRegisterTool(IAvalancheInteractiveToolsModule* InAITModule) override;
 	//~ End UAvaInteractiveToolsToolBase
 
 protected:
diff --git a/Engine/Plugins/Experimental/Avalanche/Source/AvalancheShapesEditor/Private/Tools/AvaShapesEditorShapeToolCube.cpp b/Engine/Plugins/Experimental/Avalanche/Source/AvalancheShapesEditor/Private/Tools/AvaShapesEditorShapeToolCube.cpp
--- a/Engine/Plugins/Experimental/Avalanche/Source/AvalancheShapesEditor/Private/Tools/AvaShapesEditorShapeToolCube.cpp
+++ b/Engine/Plugins/Experimental/Avalanche/Source/AvalancheShapesEditor/Private/Tools/AvaShapesEditorShapeToolCube.cpp
@@ -11,14 +11,12 @@ UAvaShapesEditorShapeToolCube::UAvaShapesEditorShapeToolCube()
 	ShapeClass = UAvaShapeCubeDynamicMesh::StaticClass();
 }
 
-FName UAvaShapesEditorShapeToolCube::GetCategoryName()
+void UAvaShapesEditorShapeToolCube::OnRegisterTool(IAvalancheInteractiveToolsModule* InAITModule)
 {
-	return IAvalancheInteractiveToolsModule::CategoryName3D;
-}
+	Super::OnRegisterTool(InAITModule);
 
-FAvaInteractiveToolsToolParameters UAvaShapesEditorShapeToolCube::GetToolParameters() const
-{
-	return {
+	FAvaInteractiveToolsToolParameters ToolParameters = 
+	{
 		FAvaShapesEditorCommands::Get().Tool_Shape_Cube,
 		TEXT("Parametric Cube Tool"),
 		1000,
@@ -30,6 +28,8 @@ FAvaInteractiveToolsToolParameters UAvaShapesEditorShapeToolCube::GetToolParamet
 		nullptr,
 		CreateFactory<UAvaShapeCubeDynamicMesh>()
 	};
+
+	InAITModule->RegisterTool(IAvalancheInteractiveToolsModule::CategoryName3D, MoveTemp(ToolParameters));
 }
 
 void UAvaShapesEditorShapeToolCube::InitShape(UAvaShapeDynamicMeshBase* InShape) const
diff --git a/Engine/Plugins/Experimental/Avalanche/Source/AvalancheShapesEditor/Private/Tools/AvaShapesEditorShapeToolCube.h b/Engine/Plugins/Experimental/Avalanche/Source/AvalancheShapesEditor/Private/Tools/AvaShapesEditorShapeToolCube.h
--- a/Engine/Plugins/Experimental/Avalanche/Source/AvalancheShapesEditor/Private/Tools/AvaShapesEditorShapeToolCube.h
+++ b/Engine/Plugins/Experimental/Avalanche/Source/AvalancheShapesEditor/Private/Tools/AvaShapesEditorShapeToolCube.h
@@ -14,8 +14,7 @@ public:
 	UAvaShapesEditorShapeToolCube();
 
 	//~ Begin UAvaInteractiveToolsToolBase
-	virtual FName GetCategoryName() override;
-	virtual FAvaInteractiveToolsToolParameters GetToolParameters() const override;
+	virtual void OnRegisterTool(IAvalancheInteractiveToolsModule* InAITModule) override;
 	//~ End UAvaInteractiveToolsToolBase
 
 protected:
diff --git a/Engine/Plugins/Experimental/Avalanche/Source/AvalancheShapesEditor/Private/Tools/AvaShapesEditorShapeToolEllipse.cpp b/Engine/Plugins/Experimental/Avalanche/Source/AvalancheShapesEditor/Private/Tools/AvaShapesEditorShapeToolEllipse.cpp
--- a/Engine/Plugins/Experimental/Avalanche/Source/AvalancheShapesEditor/Private/Tools/AvaShapesEditorShapeToolEllipse.cpp
+++ b/Engine/Plugins/Experimental/Avalanche/Source/AvalancheShapesEditor/Private/Tools/AvaShapesEditorShapeToolEllipse.cpp
@@ -11,14 +11,12 @@ UAvaShapesEditorShapeToolEllipse::UAvaShapesEditorShapeToolEllipse()
 	ShapeClass = UAvaShapeEllipseDynamicMesh::StaticClass();
 }
 
-FName UAvaShapesEditorShapeToolEllipse::GetCategoryName()
+void UAvaShapesEditorShapeToolEllipse::OnRegisterTool(IAvalancheInteractiveToolsModule* InAITModule)
 {
-	return IAvalancheInteractiveToolsModule::CategoryName2D;
-}
+	Super::OnRegisterTool(InAITModule);
 
-FAvaInteractiveToolsToolParameters UAvaShapesEditorShapeToolEllipse::GetToolParameters() const
-{
-	return {
+	FAvaInteractiveToolsToolParameters ToolParameters = 
+	{
 		FAvaShapesEditorCommands::Get().Tool_Shape_Ellipse,
 		TEXT("Parametric Ellipse Tool"),
 		2000,
@@ -30,6 +28,8 @@ FAvaInteractiveToolsToolParameters UAvaShapesEditorShapeToolEllipse::GetToolPara
 		nullptr,
 		CreateFactory<UAvaShapeEllipseDynamicMesh>()
 	};
+	
+	InAITModule->RegisterTool(IAvalancheInteractiveToolsModule::CategoryName2D, MoveTemp(ToolParameters));
 }
 
 void UAvaShapesEditorShapeToolEllipse::InitShape(UAvaShapeDynamicMeshBase* InShape) const
diff --git a/Engine/Plugins/Experimental/Avalanche/Source/AvalancheShapesEditor/Private/Tools/AvaShapesEditorShapeToolEllipse.h b/Engine/Plugins/Experimental/Avalanche/Source/AvalancheShapesEditor/Private/Tools/AvaShapesEditorShapeToolEllipse.h
--- a/Engine/Plugins/Experimental/Avalanche/Source/AvalancheShapesEditor/Private/Tools/AvaShapesEditorShapeToolEllipse.h
+++ b/Engine/Plugins/Experimental/Avalanche/Source/AvalancheShapesEditor/Private/Tools/AvaShapesEditorShapeToolEllipse.h
@@ -14,8 +14,7 @@ public:
 	UAvaShapesEditorShapeToolEllipse();
 
 	//~ Begin UAvaInteractiveToolsToolBase
-	virtual FName GetCategoryName() override;
-	virtual FAvaInteractiveToolsToolParameters GetToolParameters() const override;
+	virtual void OnRegisterTool(IAvalancheInteractiveToolsModule* InAITModule) override;
 	//~ End UAvaInteractiveToolsToolBase
 
 protected:
diff --git a/Engine/Plugins/Experimental/Avalanche/Source/AvalancheShapesEditor/Private/Tools/AvaShapesEditorShapeToolIrregularPoly.cpp b/Engine/Plugins/Experimental/Avalanche/Source/AvalancheShapesEditor/Private/Tools/AvaShapesEditorShapeToolIrregularPoly.cpp
--- a/Engine/Plugins/Experimental/Avalanche/Source/AvalancheShapesEditor/Private/Tools/AvaShapesEditorShapeToolIrregularPoly.cpp
+++ b/Engine/Plugins/Experimental/Avalanche/Source/AvalancheShapesEditor/Private/Tools/AvaShapesEditorShapeToolIrregularPoly.cpp
@@ -21,14 +21,12 @@ UAvaShapesEditorShapeToolIrregularPoly::UAvaShapesEditorShapeToolIrregularPoly()
 	ShapeClass = UAvaShapeIrregularPolygonDynamicMesh::StaticClass();
 }
 
-FName UAvaShapesEditorShapeToolIrregularPoly::GetCategoryName()
+void UAvaShapesEditorShapeToolIrregularPoly::OnRegisterTool(IAvalancheInteractiveToolsModule* InAITModule)
 {
-	return IAvalancheInteractiveToolsModule::CategoryName2D;
-}
+	Super::OnRegisterTool(InAITModule);
 
-FAvaInteractiveToolsToolParameters UAvaShapesEditorShapeToolIrregularPoly::GetToolParameters() const
-{
-	return {
+	FAvaInteractiveToolsToolParameters ToolParameters =
+	{
 		FAvaShapesEditorCommands::Get().Tool_Shape_IrregularPoly,
 		TEXT("Parametric Freehand Polygon Tool"),
 		4000,
@@ -38,6 +36,8 @@ FAvaInteractiveToolsToolParameters UAvaShapesEditorShapeToolIrregularPoly::GetTo
 				return UAvaInteractiveToolsToolBuilder::CreateToolBuilder<UAvaShapesEditorShapeToolIrregularPoly>(InEdMode);
 			})
 	};
+
+	InAITModule->RegisterTool(IAvalancheInteractiveToolsModule::CategoryName2D, MoveTemp(ToolParameters));
 }
 
 void UAvaShapesEditorShapeToolIrregularPoly::DrawHUD(FCanvas* InCanvas, IToolsContextRenderAPI* InRenderAPI)
diff --git a/Engine/Plugins/Experimental/Avalanche/Source/AvalancheShapesEditor/Private/Tools/AvaShapesEditorShapeToolIrregularPoly.h b/Engine/Plugins/Experimental/Avalanche/Source/AvalancheShapesEditor/Private/Tools/AvaShapesEditorShapeToolIrregularPoly.h
--- a/Engine/Plugins/Experimental/Avalanche/Source/AvalancheShapesEditor/Private/Tools/AvaShapesEditorShapeToolIrregularPoly.h
+++ b/Engine/Plugins/Experimental/Avalanche/Source/AvalancheShapesEditor/Private/Tools/AvaShapesEditorShapeToolIrregularPoly.h
@@ -14,8 +14,7 @@ public:
 	UAvaShapesEditorShapeToolIrregularPoly();
 
 	//~ Begin UAvaInteractiveToolsToolBase
-	virtual FName GetCategoryName() override;
-	virtual FAvaInteractiveToolsToolParameters GetToolParameters() const override;
+	virtual void OnRegisterTool(IAvalancheInteractiveToolsModule* InAITModule) override;
 	virtual void DrawHUD(FCanvas* InCanvas, IToolsContextRenderAPI* InRenderAPI) override;
 	virtual void OnViewportPlannerUpdate() override;
 	virtual void OnViewportPlannerComplete() override {} // Does nothing
diff --git a/Engine/Plugins/Experimental/Avalanche/Source/AvalancheShapesEditor/Private/Tools/AvaShapesEditorShapeToolLine.cpp b/Engine/Plugins/Experimental/Avalanche/Source/AvalancheShapesEditor/Private/Tools/AvaShapesEditorShapeToolLine.cpp
--- a/Engine/Plugins/Experimental/Avalanche/Source/AvalancheShapesEditor/Private/Tools/AvaShapesEditorShapeToolLine.cpp
+++ b/Engine/Plugins/Experimental/Avalanche/Source/AvalancheShapesEditor/Private/Tools/AvaShapesEditorShapeToolLine.cpp
@@ -59,14 +59,12 @@ UAvaShapesEditorShapeToolLine::UAvaShapesEditorShapeToolLine()
 	ShapeClass = UAvaShapeLineDynamicMesh::StaticClass();
 }
 
-FName UAvaShapesEditorShapeToolLine::GetCategoryName()
+void UAvaShapesEditorShapeToolLine::OnRegisterTool(IAvalancheInteractiveToolsModule* InAITModule)
 {
-	return IAvalancheInteractiveToolsModule::CategoryName2D;
-}
+	Super::OnRegisterTool(InAITModule);
 
-FAvaInteractiveToolsToolParameters UAvaShapesEditorShapeToolLine::GetToolParameters() const
-{
-	return {
+	FAvaInteractiveToolsToolParameters ToolParameters =
+	{
 		FAvaShapesEditorCommands::Get().Tool_Shape_Line,
 		TEXT("Parametric Line Tool"),
 		9000,
@@ -78,6 +76,8 @@ FAvaInteractiveToolsToolParameters UAvaShapesEditorShapeToolLine::GetToolParamet
 		nullptr,
 		CreateFactory<UAvaShapeLineDynamicMesh>()
 	};
+
+	InAITModule->RegisterTool(IAvalancheInteractiveToolsModule::CategoryName2D, MoveTemp(ToolParameters));
 }
 
 void UAvaShapesEditorShapeToolLine::OnViewportPlannerUpdate()
diff --git a/Engine/Plugins/Experimental/Avalanche/Source/AvalancheShapesEditor/Private/Tools/AvaShapesEditorShapeToolLine.h b/Engine/Plugins/Experimental/Avalanche/Source/AvalancheShapesEditor/Private/Tools/AvaShapesEditorShapeToolLine.h
--- a/Engine/Plugins/Experimental/Avalanche/Source/AvalancheShapesEditor/Private/Tools/AvaShapesEditorShapeToolLine.h
+++ b/Engine/Plugins/Experimental/Avalanche/Source/AvalancheShapesEditor/Private/Tools/AvaShapesEditorShapeToolLine.h
@@ -14,8 +14,7 @@ public:
 	UAvaShapesEditorShapeToolLine();
 
 	//~ Begin UAvaInteractiveToolsToolBase
-	virtual FName GetCategoryName() override;
-	virtual FAvaInteractiveToolsToolParameters GetToolParameters() const override;
+	virtual void OnRegisterTool(IAvalancheInteractiveToolsModule* InAITModule) override;
 	virtual void OnViewportPlannerUpdate() override;
 	virtual void OnViewportPlannerComplete() override {} // Does nothing
 	//~ End UAvaInteractiveToolsToolBase
@@ -24,7 +23,7 @@ protected:
 	// The minimum dimension
 	static constexpr float MinDim = 5;
 
-	FVector2f LineEndLocation;
+	FVector2f LineEndLocation = FVector2f::ZeroVector;
 
 	//~ Begin UAvaShapesEditorShapeToolBase
 	virtual void InitShape(UAvaShapeDynamicMeshBase* InShape) const override;
diff --git a/Engine/Plugins/Experimental/Avalanche/Source/AvalancheShapesEditor/Private/Tools/AvaShapesEditorShapeToolNGon.cpp b/Engine/Plugins/Experimental/Avalanche/Source/AvalancheShapesEditor/Private/Tools/AvaShapesEditorShapeToolNGon.cpp
--- a/Engine/Plugins/Experimental/Avalanche/Source/AvalancheShapesEditor/Private/Tools/AvaShapesEditorShapeToolNGon.cpp
+++ b/Engine/Plugins/Experimental/Avalanche/Source/AvalancheShapesEditor/Private/Tools/AvaShapesEditorShapeToolNGon.cpp
@@ -11,14 +11,12 @@ UAvaShapesEditorShapeToolNGon::UAvaShapesEditorShapeToolNGon()
 	ShapeClass = UAvaShapeNGonDynamicMesh::StaticClass();
 }
 
-FName UAvaShapesEditorShapeToolNGon::GetCategoryName()
+void UAvaShapesEditorShapeToolNGon::OnRegisterTool(IAvalancheInteractiveToolsModule* InAITModule)
 {
-	return IAvalancheInteractiveToolsModule::CategoryName2D;
-}
+	Super::OnRegisterTool(InAITModule);
 
-FAvaInteractiveToolsToolParameters UAvaShapesEditorShapeToolNGon::GetToolParameters() const
-{
-	return {
+	FAvaInteractiveToolsToolParameters ToolParameters =
+	{
 		FAvaShapesEditorCommands::Get().Tool_Shape_NGon,
 		TEXT("Parametric Regular Polygon Tool"),
 		3000,
@@ -30,6 +28,8 @@ FAvaInteractiveToolsToolParameters UAvaShapesEditorShapeToolNGon::GetToolParamet
 		nullptr,
 		CreateFactory<UAvaShapeNGonDynamicMesh>()
 	};
+	
+	InAITModule->RegisterTool(IAvalancheInteractiveToolsModule::CategoryName2D, MoveTemp(ToolParameters));
 }
 
 void UAvaShapesEditorShapeToolNGon::InitShape(UAvaShapeDynamicMeshBase* InShape) const
diff --git a/Engine/Plugins/Experimental/Avalanche/Source/AvalancheShapesEditor/Private/Tools/AvaShapesEditorShapeToolNGon.h b/Engine/Plugins/Experimental/Avalanche/Source/AvalancheShapesEditor/Private/Tools/AvaShapesEditorShapeToolNGon.h
--- a/Engine/Plugins/Experimental/Avalanche/Source/AvalancheShapesEditor/Private/Tools/AvaShapesEditorShapeToolNGon.h
+++ b/Engine/Plugins/Experimental/Avalanche/Source/AvalancheShapesEditor/Private/Tools/AvaShapesEditorShapeToolNGon.h
@@ -14,8 +14,7 @@ public:
 	UAvaShapesEditorShapeToolNGon();
 
 	//~ Begin UAvaInteractiveToolsToolBase
-	virtual FName GetCategoryName() override;
-	virtual FAvaInteractiveToolsToolParameters GetToolParameters() const override;
+	virtual void OnRegisterTool(IAvalancheInteractiveToolsModule* InAITModule) override;
 	//~ End UAvaInteractiveToolsToolBase
 
 protected:
diff --git a/Engine/Plugins/Experimental/Avalanche/Source/AvalancheShapesEditor/Private/Tools/AvaShapesEditorShapeToolRectangle.cpp b/Engine/Plugins/Experimental/Avalanche/Source/AvalancheShapesEditor/Private/Tools/AvaShapesEditorShapeToolRectangle.cpp
--- a/Engine/Plugins/Experimental/Avalanche/Source/AvalancheShapesEditor/Private/Tools/AvaShapesEditorShapeToolRectangle.cpp
+++ b/Engine/Plugins/Experimental/Avalanche/Source/AvalancheShapesEditor/Private/Tools/AvaShapesEditorShapeToolRectangle.cpp
@@ -4,80 +4,53 @@
 #include "AvaShapesEditorCommands.h"
 #include "Builders/AvaInteractiveToolsToolBuilder.h"
 #include "DynamicMeshes/AvaShapeRectangleDynMesh.h"
-#include "Tools/AvaInteractiveToolsToolPresetBase.h"
 #include "UObject/ConstructorHelpers.h"
 
 #define LOCTEXT_NAMESPACE "AvaShapesEditorShapeToolRectangle"
 
-class FAvaShapesEditorShapeToolRectanglePreset : public FAvaInteractiveToolsToolPresetBase
-{
-public:
-	FAvaShapesEditorShapeToolRectanglePreset(const FText& InName, const FText& InDescription, const FVector2D& InSize)
-	{
-		Name = InName;
-		Description = InDescription;
-		Size = InSize;
-	}
-
-	virtual void ApplyPreset(AActor* InActor) const override
-	{
-		if (AAvaShapeActor* ShapeActor = Cast<AAvaShapeActor>(InActor))
-		{
-			if (UAvaShapeRectangleDynamicMesh* Rectangle = Cast<UAvaShapeRectangleDynamicMesh>(ShapeActor->GetDynamicMesh()))
-			{
-				Rectangle->SetSize2D(Size);
-			}
-		}
-	}
-
-protected:
-	FVector2D Size;
-};
+const TCHAR* UAvaShapesEditorShapeToolRectangle::SquareToolIdentifier = TEXT("Parametric Square Tool");
+const TCHAR* UAvaShapesEditorShapeToolRectangle::RectangleToolIdentifier = TEXT("Parametric Rectangle Tool");
 
 UAvaShapesEditorShapeToolRectangle::UAvaShapesEditorShapeToolRectangle()
 {
 	ShapeClass = UAvaShapeRectangleDynamicMesh::StaticClass();
 }
 
-FName UAvaShapesEditorShapeToolRectangle::GetCategoryName()
+void UAvaShapesEditorShapeToolRectangle::OnRegisterTool(IAvalancheInteractiveToolsModule* InAITModule)
 {
-	return IAvalancheInteractiveToolsModule::CategoryName2D;
-}
+	Super::OnRegisterTool(InAITModule);
 
-FAvaInteractiveToolsToolParameters UAvaShapesEditorShapeToolRectangle::GetToolParameters() const
-{
-	static const TMap<FName, TSharedRef<FAvaInteractiveToolsToolPresetBase>> Presets = {
-			{
-				"Square100",
-				MakeShared<FAvaShapesEditorShapeToolRectanglePreset>(
-					LOCTEXT("100x100", "Square [100 x 100]"),
-					LOCTEXT("100x100.ToolTip", "Create a 100 x 100 square shape."),
-					FVector2D(100, 100)
-				)
-			},
+	FAvaInteractiveToolsToolParameters SquareToolParameters = 
+	{
+		FAvaShapesEditorCommands::Get().Tool_Shape_Square,
+		SquareToolIdentifier,
+		1000,
+		FAvalancheInteractiveToolsCreateBuilder::CreateLambda(
+			[](UEdMode* InEdMode)
 			{
-				"Rectangle16x9",
-				MakeShared<FAvaShapesEditorShapeToolRectanglePreset>(
-					LOCTEXT("160x90", "Rectangle [160 x 90]"),
-					LOCTEXT("160x90.ToolTip", "Create a 160 x 90 rectangle shape."),
-					FVector2D(160, 90)
-				)
-			}
+				return UAvaInteractiveToolsToolBuilder::CreateToolBuilder<UAvaShapesEditorShapeToolRectangle>(InEdMode);
+			}),
+		nullptr,
+		CreateFactory<UAvaShapeRectangleDynamicMesh>()
 	};
 
-	return {
+	InAITModule->RegisterTool(IAvalancheInteractiveToolsModule::CategoryName2D, MoveTemp(SquareToolParameters));
+
+	FAvaInteractiveToolsToolParameters RectangleToolParameters = 
+	{
 		FAvaShapesEditorCommands::Get().Tool_Shape_Rectangle,
-		TEXT("Parametric Rectangle Tool"),
-		1000,
+		RectangleToolIdentifier,
+		1001,
 		FAvalancheInteractiveToolsCreateBuilder::CreateLambda(
 			[](UEdMode* InEdMode)
 			{
 				return UAvaInteractiveToolsToolBuilder::CreateToolBuilder<UAvaShapesEditorShapeToolRectangle>(InEdMode);
 			}),
 		nullptr,
-		CreateFactory<UAvaShapeRectangleDynamicMesh>(),
-		Presets
+		CreateFactory<UAvaShapeRectangleDynamicMesh>()
 	};
+
+	InAITModule->RegisterTool(IAvalancheInteractiveToolsModule::CategoryName2D, MoveTemp(RectangleToolParameters));
 }
 
 void UAvaShapesEditorShapeToolRectangle::InitShape(UAvaShapeDynamicMeshBase* InShape) const
@@ -85,7 +58,14 @@ void UAvaShapesEditorShapeToolRectangle::InitShape(UAvaShapeDynamicMeshBase* InS
 	UAvaShapeRectangleDynamicMesh* Rectangle = Cast<UAvaShapeRectangleDynamicMesh>(InShape);
 	check(Rectangle);
 
-	Rectangle->SetSize2D({DefaultDim, DefaultDim});
+	if (GetActiveToolIdentifier().Equals(SquareToolIdentifier))
+	{
+		Rectangle->SetSize2D({DefaultDim, DefaultDim});
+	}
+	else
+	{
+		Rectangle->SetSize2D({160, 90});
+	}
 
 	Super::InitShape(InShape);
 }
diff --git a/Engine/Plugins/Experimental/Avalanche/Source/AvalancheShapesEditor/Private/Tools/AvaShapesEditorShapeToolRectangle.h b/Engine/Plugins/Experimental/Avalanche/Source/AvalancheShapesEditor/Private/Tools/AvaShapesEditorShapeToolRectangle.h
--- a/Engine/Plugins/Experimental/Avalanche/Source/AvalancheShapesEditor/Private/Tools/AvaShapesEditorShapeToolRectangle.h
+++ b/Engine/Plugins/Experimental/Avalanche/Source/AvalancheShapesEditor/Private/Tools/AvaShapesEditorShapeToolRectangle.h
@@ -11,11 +11,13 @@ class UAvaShapesEditorShapeToolRectangle : public UAvaShapesEditorShapeAreaToolB
 	GENERATED_BODY()
 
 public:
+	static const TCHAR* SquareToolIdentifier;
+	static const TCHAR* RectangleToolIdentifier;
+	
 	UAvaShapesEditorShapeToolRectangle();
 
 	//~ Begin UAvaInteractiveToolsToolBase
-	virtual FName GetCategoryName() override;
-	virtual FAvaInteractiveToolsToolParameters GetToolParameters() const override;
+	virtual void OnRegisterTool(IAvalancheInteractiveToolsModule* InAITModule) override;
 	//~ End UAvaInteractiveToolsToolBase
 
 protected:
diff --git a/Engine/Plugins/Experimental/Avalanche/Source/AvalancheShapesEditor/Private/Tools/AvaShapesEditorShapeToolRing.cpp b/Engine/Plugins/Experimental/Avalanche/Source/AvalancheShapesEditor/Private/Tools/AvaShapesEditorShapeToolRing.cpp
--- a/Engine/Plugins/Experimental/Avalanche/Source/AvalancheShapesEditor/Private/Tools/AvaShapesEditorShapeToolRing.cpp
+++ b/Engine/Plugins/Experimental/Avalanche/Source/AvalancheShapesEditor/Private/Tools/AvaShapesEditorShapeToolRing.cpp
@@ -11,14 +11,12 @@ UAvaShapesEditorShapeToolRing::UAvaShapesEditorShapeToolRing()
 	ShapeClass = UAvaShapeRingDynamicMesh::StaticClass();
 }
 
-FName UAvaShapesEditorShapeToolRing::GetCategoryName()
+void UAvaShapesEditorShapeToolRing::OnRegisterTool(IAvalancheInteractiveToolsModule* InAITModule)
 {
-	return IAvalancheInteractiveToolsModule::CategoryName2D;
-}
+	Super::OnRegisterTool(InAITModule);
 
-FAvaInteractiveToolsToolParameters UAvaShapesEditorShapeToolRing::GetToolParameters() const
-{
-	return {
+	FAvaInteractiveToolsToolParameters ToolParameters = 
+	{
 		FAvaShapesEditorCommands::Get().Tool_Shape_Ring,
 		TEXT("Parametric Ring Tool"),
 		5000,
@@ -30,6 +28,8 @@ FAvaInteractiveToolsToolParameters UAvaShapesEditorShapeToolRing::GetToolParamet
 		nullptr,
 		CreateFactory<UAvaShapeRingDynamicMesh>()
 	};
+
+	InAITModule->RegisterTool(IAvalancheInteractiveToolsModule::CategoryName2D, MoveTemp(ToolParameters));
 }
 
 void UAvaShapesEditorShapeToolRing::InitShape(UAvaShapeDynamicMeshBase* InShape) const
diff --git a/Engine/Plugins/Experimental/Avalanche/Source/AvalancheShapesEditor/Private/Tools/AvaShapesEditorShapeToolRing.h b/Engine/Plugins/Experimental/Avalanche/Source/AvalancheShapesEditor/Private/Tools/AvaShapesEditorShapeToolRing.h
--- a/Engine/Plugins/Experimental/Avalanche/Source/AvalancheShapesEditor/Private/Tools/AvaShapesEditorShapeToolRing.h
+++ b/Engine/Plugins/Experimental/Avalanche/Source/AvalancheShapesEditor/Private/Tools/AvaShapesEditorShapeToolRing.h
@@ -14,8 +14,7 @@ public:
 	UAvaShapesEditorShapeToolRing();
 
 	//~ Begin UAvaInteractiveToolsToolBase
-	virtual FName GetCategoryName() override;
-	virtual FAvaInteractiveToolsToolParameters GetToolParameters() const override;
+	virtual void OnRegisterTool(IAvalancheInteractiveToolsModule* InAITModule) override;
 	//~ End UAvaInteractiveToolsToolBase
 
 protected:
diff --git a/Engine/Plugins/Experimental/Avalanche/Source/AvalancheShapesEditor/Private/Tools/AvaShapesEditorShapeToolSphere.cpp b/Engine/Plugins/Experimental/Avalanche/Source/AvalancheShapesEditor/Private/Tools/AvaShapesEditorShapeToolSphere.cpp
--- a/Engine/Plugins/Experimental/Avalanche/Source/AvalancheShapesEditor/Private/Tools/AvaShapesEditorShapeToolSphere.cpp
+++ b/Engine/Plugins/Experimental/Avalanche/Source/AvalancheShapesEditor/Private/Tools/AvaShapesEditorShapeToolSphere.cpp
@@ -11,14 +11,12 @@ UAvaShapesEditorShapeToolSphere::UAvaShapesEditorShapeToolSphere()
 	ShapeClass = UAvaShapeSphereDynamicMesh::StaticClass();
 }
 
-FName UAvaShapesEditorShapeToolSphere::GetCategoryName()
+void UAvaShapesEditorShapeToolSphere::OnRegisterTool(IAvalancheInteractiveToolsModule* InAITModule)
 {
-	return IAvalancheInteractiveToolsModule::CategoryName3D;
-}
+	Super::OnRegisterTool(InAITModule);
 
-FAvaInteractiveToolsToolParameters UAvaShapesEditorShapeToolSphere::GetToolParameters() const
-{
-	return {
+	FAvaInteractiveToolsToolParameters ToolParameters = 
+	{
 		FAvaShapesEditorCommands::Get().Tool_Shape_Sphere,
 		TEXT("Parametric Sphere Tool"),
 		2000,
@@ -30,6 +28,8 @@ FAvaInteractiveToolsToolParameters UAvaShapesEditorShapeToolSphere::GetToolParam
 		nullptr,
 		CreateFactory<UAvaShapeSphereDynamicMesh>()
 	};
+	
+	InAITModule->RegisterTool(IAvalancheInteractiveToolsModule::CategoryName3D, MoveTemp(ToolParameters));
 }
 
 void UAvaShapesEditorShapeToolSphere::InitShape(UAvaShapeDynamicMeshBase* InShape) const
diff --git a/Engine/Plugins/Experimental/Avalanche/Source/AvalancheShapesEditor/Private/Tools/AvaShapesEditorShapeToolSphere.h b/Engine/Plugins/Experimental/Avalanche/Source/AvalancheShapesEditor/Private/Tools/AvaShapesEditorShapeToolSphere.h
--- a/Engine/Plugins/Experimental/Avalanche/Source/AvalancheShapesEditor/Private/Tools/AvaShapesEditorShapeToolSphere.h
+++ b/Engine/Plugins/Experimental/Avalanche/Source/AvalancheShapesEditor/Private/Tools/AvaShapesEditorShapeToolSphere.h
@@ -14,8 +14,7 @@ public:
 	UAvaShapesEditorShapeToolSphere();
 
 	//~ Begin UAvaInteractiveToolsToolBase
-	virtual FName GetCategoryName() override;
-	virtual FAvaInteractiveToolsToolParameters GetToolParameters() const override;
+	virtual void OnRegisterTool(IAvalancheInteractiveToolsModule* InAITModule) override;
 	//~ End UAvaInteractiveToolsToolBase
 
 protected:
diff --git a/Engine/Plugins/Experimental/Avalanche/Source/AvalancheShapesEditor/Private/Tools/AvaShapesEditorShapeToolStar.cpp b/Engine/Plugins/Experimental/Avalanche/Source/AvalancheShapesEditor/Private/Tools/AvaShapesEditorShapeToolStar.cpp
--- a/Engine/Plugins/Experimental/Avalanche/Source/AvalancheShapesEditor/Private/Tools/AvaShapesEditorShapeToolStar.cpp
+++ b/Engine/Plugins/Experimental/Avalanche/Source/AvalancheShapesEditor/Private/Tools/AvaShapesEditorShapeToolStar.cpp
@@ -11,14 +11,12 @@ UAvaShapesEditorShapeToolStar::UAvaShapesEditorShapeToolStar()
 	ShapeClass = UAvaShapeStarDynamicMesh::StaticClass();
 }
 
-FName UAvaShapesEditorShapeToolStar::GetCategoryName()
+void UAvaShapesEditorShapeToolStar::OnRegisterTool(IAvalancheInteractiveToolsModule* InAITModule)
 {
-	return IAvalancheInteractiveToolsModule::CategoryName2D;
-}
+	Super::OnRegisterTool(InAITModule);
 
-FAvaInteractiveToolsToolParameters UAvaShapesEditorShapeToolStar::GetToolParameters() const
-{
-	return {
+	FAvaInteractiveToolsToolParameters ToolParameters = 
+	{
 		FAvaShapesEditorCommands::Get().Tool_Shape_Star,
 		TEXT("Parametric Star Tool"),
 		6000,
@@ -30,6 +28,8 @@ FAvaInteractiveToolsToolParameters UAvaShapesEditorShapeToolStar::GetToolParamet
 		nullptr,
 		CreateFactory<UAvaShapeStarDynamicMesh>()
 	};
+	
+	InAITModule->RegisterTool(IAvalancheInteractiveToolsModule::CategoryName2D, MoveTemp(ToolParameters));
 }
 
 void UAvaShapesEditorShapeToolStar::InitShape(UAvaShapeDynamicMeshBase* InShape) const
diff --git a/Engine/Plugins/Experimental/Avalanche/Source/AvalancheShapesEditor/Private/Tools/AvaShapesEditorShapeToolStar.h b/Engine/Plugins/Experimental/Avalanche/Source/AvalancheShapesEditor/Private/Tools/AvaShapesEditorShapeToolStar.h
--- a/Engine/Plugins/Experimental/Avalanche/Source/AvalancheShapesEditor/Private/Tools/AvaShapesEditorShapeToolStar.h
+++ b/Engine/Plugins/Experimental/Avalanche/Source/AvalancheShapesEditor/Private/Tools/AvaShapesEditorShapeToolStar.h
@@ -14,8 +14,7 @@ public:
 	UAvaShapesEditorShapeToolStar();
 
 	//~ Begin UAvaInteractiveToolsToolBase
-	virtual FName GetCategoryName() override;
-	virtual FAvaInteractiveToolsToolParameters GetToolParameters() const override;
+	virtual void OnRegisterTool(IAvalancheInteractiveToolsModule* InAITModule) override;
 	//~ End UAvaInteractiveToolsToolBase
 
 protected:
diff --git a/Engine/Plugins/Experimental/Avalanche/Source/AvalancheShapesEditor/Private/Tools/AvaShapesEditorShapeToolTorus.cpp b/Engine/Plugins/Experimental/Avalanche/Source/AvalancheShapesEditor/Private/Tools/AvaShapesEditorShapeToolTorus.cpp
--- a/Engine/Plugins/Experimental/Avalanche/Source/AvalancheShapesEditor/Private/Tools/AvaShapesEditorShapeToolTorus.cpp
+++ b/Engine/Plugins/Experimental/Avalanche/Source/AvalancheShapesEditor/Private/Tools/AvaShapesEditorShapeToolTorus.cpp
@@ -11,14 +11,12 @@ UAvaShapesEditorShapeToolTorus::UAvaShapesEditorShapeToolTorus()
 	ShapeClass = UAvaShapeTorusDynamicMesh::StaticClass();
 }
 
-FName UAvaShapesEditorShapeToolTorus::GetCategoryName()
+void UAvaShapesEditorShapeToolTorus::OnRegisterTool(IAvalancheInteractiveToolsModule* InAITModule)
 {
-	return IAvalancheInteractiveToolsModule::CategoryName3D;
-}
+	Super::OnRegisterTool(InAITModule);
 
-FAvaInteractiveToolsToolParameters UAvaShapesEditorShapeToolTorus::GetToolParameters() const
-{
-	return {
+	FAvaInteractiveToolsToolParameters ToolParameters = 
+	{
 		FAvaShapesEditorCommands::Get().Tool_Shape_Torus,
 		TEXT("Parametric Torus Tool"),
 		4000,
@@ -30,6 +28,8 @@ FAvaInteractiveToolsToolParameters UAvaShapesEditorShapeToolTorus::GetToolParame
 		nullptr,
 		CreateFactory<UAvaShapeTorusDynamicMesh>()
 	};
+	
+	InAITModule->RegisterTool(IAvalancheInteractiveToolsModule::CategoryName3D, MoveTemp(ToolParameters));
 }
 
 void UAvaShapesEditorShapeToolTorus::InitShape(UAvaShapeDynamicMeshBase* InShape) const
diff --git a/Engine/Plugins/Experimental/Avalanche/Source/AvalancheShapesEditor/Private/Tools/AvaShapesEditorShapeToolTorus.h b/Engine/Plugins/Experimental/Avalanche/Source/AvalancheShapesEditor/Private/Tools/AvaShapesEditorShapeToolTorus.h
--- a/Engine/Plugins/Experimental/Avalanche/Source/AvalancheShapesEditor/Private/Tools/AvaShapesEditorShapeToolTorus.h
+++ b/Engine/Plugins/Experimental/Avalanche/Source/AvalancheShapesEditor/Private/Tools/AvaShapesEditorShapeToolTorus.h
@@ -14,8 +14,7 @@ public:
 	UAvaShapesEditorShapeToolTorus();
 
 	//~ Begin UAvaInteractiveToolsToolBase
-	virtual FName GetCategoryName() override;
-	virtual FAvaInteractiveToolsToolParameters GetToolParameters() const override;
+	virtual void OnRegisterTool(IAvalancheInteractiveToolsModule* InAITModule) override;
 	//~ End UAvaInteractiveToolsToolBase
 
 protected:
diff --git a/Engine/Plugins/Experimental/Avalanche/Source/AvalancheTextEditor/Private/AvaTextEditorModule.cpp b/Engine/Plugins/Experimental/Avalanche/Source/AvalancheTextEditor/Private/AvaTextEditorModule.cpp
--- a/Engine/Plugins/Experimental/Avalanche/Source/AvalancheTextEditor/Private/AvaTextEditorModule.cpp
+++ b/Engine/Plugins/Experimental/Avalanche/Source/AvalancheTextEditor/Private/AvaTextEditorModule.cpp
@@ -21,7 +21,6 @@
 #include "Misc/CoreDelegates.h"
 #include "Modules/ModuleManager.h"
 #include "PropertyEditorModule.h"
-#include "Tool/AvaTextActorTool.h"
 #include "Visualizer/AvaTextVisualizer.h"
 
 void FAvaTextEditorModule::StartupModule()
@@ -30,7 +29,6 @@ void FAvaTextEditorModule::StartupModule()
 
 	FAvaTextEditorCommands::Register();
 	RegisterCustomLayouts();
-	FAvaInteractiveToolsDelegates::GetRegisterToolsDelegate().AddRaw(this, &FAvaTextEditorModule::RegisterTools);
 
 	RegisterDynamicMaterialPropertyGenerator();
 
@@ -47,8 +45,6 @@ void FAvaTextEditorModule::ShutdownModule()
 	{
 		UnregisterCustomLayouts();
 	}
-
-	FAvaInteractiveToolsDelegates::GetRegisterToolsDelegate().RemoveAll(this);
 }
 
 void FAvaTextEditorModule::PostEngineInit()
@@ -59,15 +55,6 @@ void FAvaTextEditorModule::PostEngineInit()
 	}
 }
 
-void FAvaTextEditorModule::RegisterTools(IAvalancheInteractiveToolsModule* InModule)
-{
-	// 2D Tools
-	InModule->RegisterTool(
-		IAvalancheInteractiveToolsModule::Get().CategoryNameActor,
-		GetDefault<UAvaTextActorTool>()->GetToolParameters()
-	);
-}
-
 void FAvaTextEditorModule::RegisterComponentVisualizers()
 {
 	IAvalancheComponentVisualizersModule& AvaComponentVisualizerModule = IAvalancheComponentVisualizersModule::Get();
diff --git a/Engine/Plugins/Experimental/Avalanche/Source/AvalancheTextEditor/Private/AvaTextEditorModule.h b/Engine/Plugins/Experimental/Avalanche/Source/AvalancheTextEditor/Private/AvaTextEditorModule.h
--- a/Engine/Plugins/Experimental/Avalanche/Source/AvalancheTextEditor/Private/AvaTextEditorModule.h
+++ b/Engine/Plugins/Experimental/Avalanche/Source/AvalancheTextEditor/Private/AvaTextEditorModule.h
@@ -27,8 +27,6 @@ private:
 
 	void PostEngineInit();
 
-	void RegisterTools(IAvalancheInteractiveToolsModule* InModule);
-
 	void RegisterComponentVisualizers();
 
 	void RegisterCustomLayouts();
diff --git a/Engine/Plugins/Experimental/Avalanche/Source/AvalancheTextEditor/Private/AvaTextEditorStyle.cpp b/Engine/Plugins/Experimental/Avalanche/Source/AvalancheTextEditor/Private/AvaTextEditorStyle.cpp
--- a/Engine/Plugins/Experimental/Avalanche/Source/AvalancheTextEditor/Private/AvaTextEditorStyle.cpp
+++ b/Engine/Plugins/Experimental/Avalanche/Source/AvalancheTextEditor/Private/AvaTextEditorStyle.cpp
@@ -11,7 +11,6 @@ FAvaTextEditorStyle::FAvaTextEditorStyle()
 	: FSlateStyleSet(TEXT("AvaTextEditor"))
 {
 	const FVector2D Icon16x16(16.0f, 16.0f);
-	const FVector2D Icon20x20(20.0f, 20.0f);
 
 	const TSharedPtr<IPlugin> Plugin = IPluginManager::Get().FindPlugin(UE_PLUGIN_NAME);
 	check(Plugin.IsValid());
@@ -19,8 +18,8 @@ FAvaTextEditorStyle::FAvaTextEditorStyle()
 	SetContentRoot(Plugin->GetBaseDir() / TEXT("Resources"));
 	
 	Set("ClassIcon.AvaTextActor", new IMAGE_BRUSH("Icons/ToolboxIcons/3d-text", Icon16x16));
-	Set("AvaTextEditor.Tool_Actor_Text3D", new IMAGE_BRUSH_SVG("Icons/ToolboxIcons/Toolbar_Text", Icon20x20));
-	Set("Tool_Actor_Text3D", new IMAGE_BRUSH_SVG("Icons/ToolboxIcons/Toolbar_Text", Icon20x20));
+	Set("AvaTextEditor.Tool_Actor_Text3D", new IMAGE_BRUSH_SVG("Icons/ToolboxIcons/Toolbar_Text", Icon16x16));
+	Set("Tool_Actor_Text3D", new IMAGE_BRUSH_SVG("Icons/ToolboxIcons/Toolbar_Text", Icon16x16));
 
 	FSlateStyleRegistry::RegisterSlateStyle(*this);
 }
diff --git a/Engine/Plugins/Experimental/Avalanche/Source/AvalancheTextEditor/Private/Tool/AvaTextActorTool.cpp b/Engine/Plugins/Experimental/Avalanche/Source/AvalancheTextEditor/Private/Tool/AvaTextActorTool.cpp
--- a/Engine/Plugins/Experimental/Avalanche/Source/AvalancheTextEditor/Private/Tool/AvaTextActorTool.cpp
+++ b/Engine/Plugins/Experimental/Avalanche/Source/AvalancheTextEditor/Private/Tool/AvaTextActorTool.cpp
@@ -13,14 +13,12 @@ UAvaTextActorTool::UAvaTextActorTool()
 	ActorClass = AAvaTextActor::StaticClass();
 }
 
-FName UAvaTextActorTool::GetCategoryName()
+void UAvaTextActorTool::OnRegisterTool(IAvalancheInteractiveToolsModule* InAITModule)
 {
-	return IAvalancheInteractiveToolsModule::CategoryNameActor;
-}
+	Super::OnRegisterTool(InAITModule);
 
-FAvaInteractiveToolsToolParameters UAvaTextActorTool::GetToolParameters() const
-{
-	return {
+	FAvaInteractiveToolsToolParameters ToolParameters =
+	{
 		FAvaTextEditorCommands::Get().Tool_Actor_Text3D,
 		TEXT("Text Actor Tool"),
 		2000,
@@ -32,4 +30,6 @@ FAvaInteractiveToolsToolParameters UAvaTextActorTool::GetToolParameters() const
 		ActorClass,
 		CreateActorFactory<UAvaTextActorFactory>()
 	};
+
+	InAITModule->RegisterTool(IAvalancheInteractiveToolsModule::CategoryNameActor, MoveTemp(ToolParameters));
 }
diff --git a/Engine/Plugins/Experimental/Avalanche/Source/AvalancheTextEditor/Private/Tool/AvaTextActorTool.h b/Engine/Plugins/Experimental/Avalanche/Source/AvalancheTextEditor/Private/Tool/AvaTextActorTool.h
--- a/Engine/Plugins/Experimental/Avalanche/Source/AvalancheTextEditor/Private/Tool/AvaTextActorTool.h
+++ b/Engine/Plugins/Experimental/Avalanche/Source/AvalancheTextEditor/Private/Tool/AvaTextActorTool.h
@@ -14,7 +14,6 @@ public:
 	UAvaTextActorTool();
 
 	//~ Begin UAvaInteractiveToolsToolBase
-	virtual FName GetCategoryName() override;
-	virtual FAvaInteractiveToolsToolParameters GetToolParameters() const override;
+	virtual void OnRegisterTool(IAvalancheInteractiveToolsModule* InAITModule) override;
 	//~ End UAvaInteractiveToolsToolBase
 };
diff --git a/Engine/Plugins/Experimental/ClonerEffector/Source/ClonerEffector/Private/Utilities/CEClonerLibrary.cpp b/Engine/Plugins/Experimental/ClonerEffector/Source/ClonerEffector/Private/Utilities/CEClonerLibrary.cpp
--- a/Engine/Plugins/Experimental/ClonerEffector/Source/ClonerEffector/Private/Utilities/CEClonerLibrary.cpp
+++ b/Engine/Plugins/Experimental/ClonerEffector/Source/ClonerEffector/Private/Utilities/CEClonerLibrary.cpp
@@ -40,6 +40,16 @@ bool UCEClonerLibrary::GetClonerLayoutName(TSubclassOf<UCEClonerLayoutBase> InLa
 	return !OutLayoutName.IsNone();
 }
 
+bool UCEClonerLibrary::GetClonerLayoutNames(TSet<FName>& OutLayoutNames)
+{
+	if (const UCEClonerSubsystem* ClonerSubsystem = UCEClonerSubsystem::Get())
+	{
+		OutLayoutNames = ClonerSubsystem->GetLayoutNames();
+	}
+
+	return !OutLayoutNames.IsEmpty();
+}
+
 bool UCEClonerLibrary::GetClonerLayoutClass(FName InLayoutName, TSubclassOf<UCEClonerLayoutBase>& OutLayoutClass)
 {
 	if (const UCEClonerSubsystem* ClonerSubsystem = UCEClonerSubsystem::Get())
diff --git a/Engine/Plugins/Experimental/ClonerEffector/Source/ClonerEffector/Private/Utilities/CEEffectorLibrary.cpp b/Engine/Plugins/Experimental/ClonerEffector/Source/ClonerEffector/Private/Utilities/CEEffectorLibrary.cpp
--- a/Engine/Plugins/Experimental/ClonerEffector/Source/ClonerEffector/Private/Utilities/CEEffectorLibrary.cpp
+++ b/Engine/Plugins/Experimental/ClonerEffector/Source/ClonerEffector/Private/Utilities/CEEffectorLibrary.cpp
@@ -34,6 +34,14 @@ void UCEEffectorLibrary::GetEffectorTypeClasses(TSet<TSubclassOf<UCEEffectorType
 	}
 }
 
+void UCEEffectorLibrary::GetEffectorTypeNames(TSet<FName>& OutTypeNames)
+{
+	if (const UCEEffectorSubsystem* EffectorSubsystem = UCEEffectorSubsystem::Get())
+	{
+		OutTypeNames = EffectorSubsystem->GetExtensionNames<UCEEffectorTypeBase>();
+	}
+}
+
 void UCEEffectorLibrary::GetEffectorEffectClasses(TSet<TSubclassOf<UCEEffectorEffectBase>>& OutEffectClasses)
 {
 	OutEffectClasses.Empty();
diff --git a/Engine/Plugins/Experimental/ClonerEffector/Source/ClonerEffector/Public/Utilities/CEClonerLibrary.h b/Engine/Plugins/Experimental/ClonerEffector/Source/ClonerEffector/Public/Utilities/CEClonerLibrary.h
--- a/Engine/Plugins/Experimental/ClonerEffector/Source/ClonerEffector/Public/Utilities/CEClonerLibrary.h
+++ b/Engine/Plugins/Experimental/ClonerEffector/Source/ClonerEffector/Public/Utilities/CEClonerLibrary.h
@@ -42,6 +42,14 @@ public:
 	UFUNCTION(BlueprintCallable, Category="Motion Design|Cloner|Utility")
 	static CLONEREFFECTOR_API bool GetClonerLayoutName(TSubclassOf<UCEClonerLayoutBase> InLayoutClass, FName& OutLayoutName);
 
+	/**
+	 * Retrieves all the layout names
+	 * @param OutLayoutNames [Out] Layout names
+	 * @return true when the layout names were retrieved
+	 */
+	UFUNCTION(BlueprintCallable, Category="Motion Design|Cloner|Utility")
+	static CLONEREFFECTOR_API bool GetClonerLayoutNames(TSet<FName>& OutLayoutNames);
+
 	/**
 	 * Retrieves the layout class from a layout name
 	 * @param InLayoutName Layout name to get the class from
diff --git a/Engine/Plugins/Experimental/ClonerEffector/Source/ClonerEffector/Public/Utilities/CEEffectorLibrary.h b/Engine/Plugins/Experimental/ClonerEffector/Source/ClonerEffector/Public/Utilities/CEEffectorLibrary.h
--- a/Engine/Plugins/Experimental/ClonerEffector/Source/ClonerEffector/Public/Utilities/CEEffectorLibrary.h
+++ b/Engine/Plugins/Experimental/ClonerEffector/Source/ClonerEffector/Public/Utilities/CEEffectorLibrary.h
@@ -33,6 +33,13 @@ public:
 	UFUNCTION(BlueprintCallable, Category="Motion Design|Effector|Utility")
 	static CLONEREFFECTOR_API void GetEffectorTypeClasses(TSet<TSubclassOf<UCEEffectorTypeBase>>& OutTypeClasses);
 
+	/**
+	 * Retrieves all type classes available for an effector
+	 * @param OutTypeNames [Out] Available type names
+	 */
+	UFUNCTION(BlueprintCallable, Category="Motion Design|Effector|Utility")
+	static CLONEREFFECTOR_API void GetEffectorTypeNames(TSet<FName>& OutTypeNames);
+
 	/**
 	 * Retrieves all effect classes available for an effector
 	 * @param OutEffectClasses [Out] Available effect classes
diff --git a/Engine/Plugins/Experimental/SVGImporter/Source/SVGImporterEditor/Private/SVGImporterEditorStyle.cpp b/Engine/Plugins/Experimental/SVGImporter/Source/SVGImporterEditor/Private/SVGImporterEditorStyle.cpp
--- a/Engine/Plugins/Experimental/SVGImporter/Source/SVGImporterEditor/Private/SVGImporterEditorStyle.cpp
+++ b/Engine/Plugins/Experimental/SVGImporter/Source/SVGImporterEditor/Private/SVGImporterEditorStyle.cpp
@@ -10,7 +10,6 @@ FSVGImporterEditorStyle::FSVGImporterEditorStyle()
 	: FSlateStyleSet(TEXT("SVGImporterEditor"))
 {
 	const FVector2f Icon16x16(16.f, 16.f);
-	const FVector2f Icon20x20(20.f, 20.f);
 
 	const TSharedPtr<IPlugin> Plugin = IPluginManager::Get().FindPlugin(UE_PLUGIN_NAME);
 	check(Plugin.IsValid());
@@ -21,7 +20,7 @@ FSVGImporterEditorStyle::FSVGImporterEditorStyle()
 	Set("ClassIcon.SVGShapesParentActor",    new IMAGE_BRUSH_SVG("Icons/DefaultSVG", Icon16x16));
 	Set("ClassIcon.SVGShapeActor",           new IMAGE_BRUSH("Icons/irregularpolygon", Icon16x16));
 	Set("ClassIcon.SVGDynamicMeshComponent", new IMAGE_BRUSH("Icons/irregularpolygon", Icon16x16));
-	Set("SVGImporterEditor.SpawnSVGActor",   new IMAGE_BRUSH_SVG("Icons/DefaultSVG", Icon20x20));
+	Set("SVGImporterEditor.SpawnSVGActor",   new IMAGE_BRUSH_SVG("Icons/DefaultSVG", Icon16x16));
 
 	FSlateStyleRegistry::RegisterSlateStyle(*this);
 }
