StateTree: Fix compilation issue with shadow variable. Introduced with CL 39893788. Change the EStateTreeNodeFormatting for a class enum to add a scope to the variable Text.

diff --git a/Engine/Plugins/Experimental/Avalanche/Source/AvalancheTransition/Public/AvaTransitionLayerUtils.h b/Engine/Plugins/Experimental/Avalanche/Source/AvalancheTransition/Public/AvaTransitionLayerUtils.h
--- a/Engine/Plugins/Experimental/Avalanche/Source/AvalancheTransition/Public/AvaTransitionLayerUtils.h
+++ b/Engine/Plugins/Experimental/Avalanche/Source/AvalancheTransition/Public/AvaTransitionLayerUtils.h
@@ -7,7 +7,7 @@
 #include "UObject/NameTypes.h"
 
 class UAvaTransitionSubsystem;
-enum EStateTreeNodeFormatting : uint8;
+enum class EStateTreeNodeFormatting : uint8;
 enum class EAvaTransitionLayerCompareType : uint8;
 struct FAvaTagHandle;
 struct FAvaTagHandleContainer;
