Motion Design: remove the category "Selection Utility" in the default TL mode.

diff --git a/Engine/Plugins/Experimental/Avalanche/Source/AvalancheTransitionEditor/Private/Customizations/AvaStateTreeStateCustomization.cpp b/Engine/Plugins/Experimental/Avalanche/Source/AvalancheTransitionEditor/Private/Customizations/AvaStateTreeStateCustomization.cpp
--- a/Engine/Plugins/Experimental/Avalanche/Source/AvalancheTransitionEditor/Private/Customizations/AvaStateTreeStateCustomization.cpp
+++ b/Engine/Plugins/Experimental/Avalanche/Source/AvalancheTransitionEditor/Private/Customizations/AvaStateTreeStateCustomization.cpp
@@ -35,6 +35,7 @@ void FAvaStateTreeStateCustomization::CustomizeDetails(IDetailLayoutBuilder& InD
 	InDetailBuilder.HideCategory(TEXT("State"));
 	InDetailBuilder.HideCategory(TEXT("Transitions"));
 	InDetailBuilder.HideCategory(TEXT("Parameters"));
+	InDetailBuilder.HideCategory(TEXT("Selection Utility"));
 
 	InDetailBuilder.EditCategory(TEXT("Enter Conditions")).InitiallyCollapsed(false);
 	InDetailBuilder.EditCategory(TEXT("Task")).InitiallyCollapsed(false);
