Motion Design: Fix issue with State Tree Ref being read-only with its parameters.

diff --git a/Engine/Plugins/Experimental/Avalanche/Source/AvalancheTransition/Private/Behavior/AvaTransitionBehaviorActor.cpp b/Engine/Plugins/Experimental/Avalanche/Source/AvalancheTransition/Private/Behavior/AvaTransitionBehaviorActor.cpp
--- a/Engine/Plugins/Experimental/Avalanche/Source/AvalancheTransition/Private/Behavior/AvaTransitionBehaviorActor.cpp
+++ b/Engine/Plugins/Experimental/Avalanche/Source/AvalancheTransition/Private/Behavior/AvaTransitionBehaviorActor.cpp
@@ -44,6 +44,38 @@ void AAvaTransitionBehaviorActor::PostLoad()
 	ValidateTransitionTree();
 }
 
+#if WITH_EDITOR
+bool AAvaTransitionBehaviorActor::CanEditChange(const FProperty* InProperty) const
+{
+	return Super::CanEditChange(InProperty);
+}
+
+bool AAvaTransitionBehaviorActor::CanEditChange(const FEditPropertyChain& InPropertyChain) const
+{
+	if (!UObject::CanEditChange(InPropertyChain))
+	{
+		return false;
+	}
+
+	auto GetPropertyName =
+		[](TDoubleLinkedList<FProperty*>::TDoubleLinkedListNode* InChainNode)->FName
+		{
+			if (InChainNode && InChainNode->GetValue())
+			{
+				return InChainNode->GetValue()->GetFName();
+			}
+			return NAME_None;	
+		};
+
+	const FName MemberPropertyName = GetPropertyName(InPropertyChain.GetActiveMemberNode());
+	const FName ActivePropertyName = GetPropertyName(InPropertyChain.GetActiveNode());
+
+	// Disable change for StateTree in FStateTreeReference
+	return MemberPropertyName != GET_MEMBER_NAME_CHECKED(AAvaTransitionBehaviorActor, StateTreeReference)
+		|| ActivePropertyName != TEXT("StateTree");
+}
+#endif
+
 UAvaTransitionSubsystem* AAvaTransitionBehaviorActor::GetTransitionSubsystem() const
 {
 	UWorld* const World = GetWorld();
diff --git a/Engine/Plugins/Experimental/Avalanche/Source/AvalancheTransition/Private/Behavior/AvaTransitionBehaviorActor.h b/Engine/Plugins/Experimental/Avalanche/Source/AvalancheTransition/Private/Behavior/AvaTransitionBehaviorActor.h
--- a/Engine/Plugins/Experimental/Avalanche/Source/AvalancheTransition/Private/Behavior/AvaTransitionBehaviorActor.h
+++ b/Engine/Plugins/Experimental/Avalanche/Source/AvalancheTransition/Private/Behavior/AvaTransitionBehaviorActor.h
@@ -39,6 +39,10 @@ protected:
 
 	//~ Begin UObject
 	virtual void PostLoad() override;
+#if WITH_EDITOR
+	virtual bool CanEditChange(const FProperty* InProperty) const override;
+	virtual bool CanEditChange(const FEditPropertyChain& InPropertyChain) const;
+#endif
 	//~ End UObject
 
 private:
@@ -49,6 +53,6 @@ private:
 	UPROPERTY()
 	TObjectPtr<UAvaTransitionTree> TransitionTree;
 
-	UPROPERTY(VisibleAnywhere, Category="Transition Logic")
+	UPROPERTY(EditAnywhere, Category="Transition Logic")
 	FStateTreeReference StateTreeReference;
 };
