MotionDesign : InteractiveTools
- Simplifying interactive tools by using actor factory to setup and spawn actor to have the same code path
- Fix place actors tab motion design actors variant not working when dropped eg: Square, Cylinder, ...
- Improving viewport toolbar style and dropdown context menu
- Re-adding motion design toolbox tab that adapts when resized

diff --git a/Engine/Plugins/Experimental/Avalanche/Source/AvalancheEffectorsEditor/Private/Cloner/AvaClonerActorTool.cpp b/Engine/Plugins/Experimental/Avalanche/Source/AvalancheEffectorsEditor/Private/Cloner/AvaClonerActorTool.cpp
--- a/Engine/Plugins/Experimental/Avalanche/Source/AvalancheEffectorsEditor/Private/Cloner/AvaClonerActorTool.cpp
+++ b/Engine/Plugins/Experimental/Avalanche/Source/AvalancheEffectorsEditor/Private/Cloner/AvaClonerActorTool.cpp
@@ -5,7 +5,6 @@
 #include "Builders/AvaInteractiveToolsToolBuilder.h"
 #include "Cloner/CEClonerActor.h"
 #include "Cloner/CEClonerActorFactory.h"
-#include "Cloner/CEClonerComponent.h"
 
 UAvaClonerActorTool::UAvaClonerActorTool()
 {
@@ -22,6 +21,9 @@ void UAvaClonerActorTool::OnRegisterTool(IAvalancheInteractiveToolsModule* InAIT
 	const FString ToolIdentifier = TEXT("Clone Actor Tool ");
 	for (const TPair<FName, TSharedPtr<FUICommandInfo>>& ClonerCommandPair : FAvaEffectorsEditorCommands::Get().Tool_Actor_Cloners)
 	{
+		UCEClonerActorFactory* ClonerActorFactory = CreateActorFactory<UCEClonerActorFactory>();
+		ClonerActorFactory->SetClonerLayout(ClonerCommandPair.Key);
+		
 		FAvaInteractiveToolsToolParameters ToolParameters =
 		{
 			ClonerCommandPair.Value,
@@ -33,25 +35,10 @@ void UAvaClonerActorTool::OnRegisterTool(IAvalancheInteractiveToolsModule* InAIT
 					return UAvaInteractiveToolsToolBuilder::CreateToolBuilder<UAvaClonerActorTool>(InEdMode);
 				}),
 			ActorClass,
-			CreateActorFactory<UCEClonerActorFactory>(),
+			ClonerActorFactory,
 			FText::FromString(ClonerCommandPair.Key.ToString()),
 		};
 
 		InAITModule->RegisterTool(IAvalancheInteractiveToolsModule::CategoryNameCloner, MoveTemp(ToolParameters));
 	}
 }
-
-void UAvaClonerActorTool::OnActorSpawned(AActor* InActor)
-{
-	Super::OnActorSpawned(InActor);
-
-	if (const ACEClonerActor* ClonerActor = Cast<ACEClonerActor>(InActor))
-	{
-		const FAvaInteractiveToolsToolParameters* ActiveTool = GetActiveToolParameters();
-
-		if (ActiveTool && ActiveTool->Label.IsSet())
-		{
-			ClonerActor->GetClonerComponent()->SetLayoutName(FName(ActiveTool->Label.GetValue().ToString()));
-		}
-	}
-}
diff --git a/Engine/Plugins/Experimental/Avalanche/Source/AvalancheEffectorsEditor/Private/Cloner/AvaClonerActorTool.h b/Engine/Plugins/Experimental/Avalanche/Source/AvalancheEffectorsEditor/Private/Cloner/AvaClonerActorTool.h
--- a/Engine/Plugins/Experimental/Avalanche/Source/AvalancheEffectorsEditor/Private/Cloner/AvaClonerActorTool.h
+++ b/Engine/Plugins/Experimental/Avalanche/Source/AvalancheEffectorsEditor/Private/Cloner/AvaClonerActorTool.h
@@ -15,6 +15,5 @@ public:
 
 	//~ Begin UAvaInteractiveToolsToolBase
 	virtual void OnRegisterTool(IAvalancheInteractiveToolsModule* InAITModule) override;
-	virtual void OnActorSpawned(AActor* InActor) override;
 	//~ End UAvaInteractiveToolsToolBase
 };
diff --git a/Engine/Plugins/Experimental/Avalanche/Source/AvalancheEffectorsEditor/Private/Effector/AvaEffectorActorTool.cpp b/Engine/Plugins/Experimental/Avalanche/Source/AvalancheEffectorsEditor/Private/Effector/AvaEffectorActorTool.cpp
--- a/Engine/Plugins/Experimental/Avalanche/Source/AvalancheEffectorsEditor/Private/Effector/AvaEffectorActorTool.cpp
+++ b/Engine/Plugins/Experimental/Avalanche/Source/AvalancheEffectorsEditor/Private/Effector/AvaEffectorActorTool.cpp
@@ -22,6 +22,9 @@ void UAvaEffectorActorTool::OnRegisterTool(IAvalancheInteractiveToolsModule* InA
 	const FString ToolIdentifier = TEXT("Effector Actor Tool ");
 	for (const TPair<FName, TSharedPtr<FUICommandInfo>>& EffectorCommandPair : FAvaEffectorsEditorCommands::Get().Tool_Actor_Effectors)
 	{
+		UCEEffectorActorFactory* EffectorActorFactory = CreateActorFactory<UCEEffectorActorFactory>();
+		EffectorActorFactory->SetEffectorTypeName(EffectorCommandPair.Key);
+
 		FAvaInteractiveToolsToolParameters ToolParameters =
 		{
 			EffectorCommandPair.Value,
@@ -33,25 +36,10 @@ void UAvaEffectorActorTool::OnRegisterTool(IAvalancheInteractiveToolsModule* InA
 					return UAvaInteractiveToolsToolBuilder::CreateToolBuilder<UAvaEffectorActorTool>(InEdMode);
 				}),
 			ActorClass,
-			CreateActorFactory<UCEEffectorActorFactory>(),
+			EffectorActorFactory,
 			FText::FromString(EffectorCommandPair.Key.ToString()),
 		};
 
 		InAITModule->RegisterTool(IAvalancheInteractiveToolsModule::CategoryNameEffector, MoveTemp(ToolParameters));
 	}
-}
-
-void UAvaEffectorActorTool::OnActorSpawned(AActor* InActor)
-{
-	Super::OnActorSpawned(InActor);
-
-	if (const ACEEffectorActor* EffectorActor = Cast<ACEEffectorActor>(InActor))
-	{
-		const FAvaInteractiveToolsToolParameters* ActiveTool = GetActiveToolParameters();
-
-		if (ActiveTool && ActiveTool->Label.IsSet())
-		{
-			EffectorActor->GetEffectorComponent()->SetTypeName(FName(ActiveTool->Label.GetValue().ToString()));
-		}
-	}
-}
+}
\ No newline at end of file
diff --git a/Engine/Plugins/Experimental/Avalanche/Source/AvalancheEffectorsEditor/Private/Effector/AvaEffectorActorTool.h b/Engine/Plugins/Experimental/Avalanche/Source/AvalancheEffectorsEditor/Private/Effector/AvaEffectorActorTool.h
--- a/Engine/Plugins/Experimental/Avalanche/Source/AvalancheEffectorsEditor/Private/Effector/AvaEffectorActorTool.h
+++ b/Engine/Plugins/Experimental/Avalanche/Source/AvalancheEffectorsEditor/Private/Effector/AvaEffectorActorTool.h
@@ -15,6 +15,5 @@ public:
 
 	//~ Begin UAvaInteractiveToolsToolBase
 	virtual void OnRegisterTool(IAvalancheInteractiveToolsModule* InAITModule) override;
-	virtual void OnActorSpawned(AActor* InActor) override;
 	//~ End UAvaInteractiveToolsToolBase
 };
diff --git a/Engine/Plugins/Experimental/Avalanche/Source/AvalancheInteractiveTools/Private/AvaInteractiveToolsStyle.cpp b/Engine/Plugins/Experimental/Avalanche/Source/AvalancheInteractiveTools/Private/AvaInteractiveToolsStyle.cpp
--- a/Engine/Plugins/Experimental/Avalanche/Source/AvalancheInteractiveTools/Private/AvaInteractiveToolsStyle.cpp
+++ b/Engine/Plugins/Experimental/Avalanche/Source/AvalancheInteractiveTools/Private/AvaInteractiveToolsStyle.cpp
@@ -7,6 +7,7 @@
 #include "Math/MathFwd.h"
 #include "Misc/Paths.h"
 #include "Styling/AppStyle.h"
+#include "Styling/CoreStyle.h"
 #include "Styling/StyleColors.h"
 #include "Styling/SlateStyleMacros.h"
 #include "Styling/SlateStyleRegistry.h"
@@ -23,21 +24,33 @@ FAvaInteractiveToolsStyle::FAvaInteractiveToolsStyle()
 	SetContentRoot(Plugin->GetBaseDir() / TEXT("Resources"));
 	SetCoreContentRoot(FPaths::EngineContentDir() / TEXT("Editor/Slate"));
 
-	FToolBarStyle ToolbarStyle = FAppStyle::Get().GetWidgetStyle<FToolBarStyle>("EditorViewportToolBar");
-	FTextBlockStyle TextStyle = FAppStyle::Get().GetWidgetStyle<FTextBlockStyle>("SmallText");
+	const FTextBlockStyle TextStyle = FTextBlockStyle()
+		.SetFont(FCoreStyle::GetDefaultFontStyle("Regular", 6))
+		.SetColorAndOpacity(FLinearColor(0.4f, 0.4f, 0.4f))
+		.SetOverflowPolicy(ETextOverflowPolicy::Clip);
 
+	const FMargin Margin(0.5f);
+	FToolBarStyle ToolbarStyle = FAppStyle::Get().GetWidgetStyle<FToolBarStyle>("EditorViewportToolBar");
 	ToolbarStyle
-		.SetButtonPadding(FMargin(2.f))
+	    .SetComboLabelMaxWidth(38.f)
+		.SetComboLabelMinWidth(38.f)
+		.SetButtonContentMaxWidth(38.f)
 		.SetAllowWrapButton(true)
 		.SetIconSize(Icon16x16)
-		.SetIconPadding(FMargin(2.f))
-		.SetLabelPadding(FMargin(2.f))
+		.SetButtonPadding(Margin)
+		.SetIconPadding(Margin)
+		.SetLabelPadding(Margin)
 		.SetLabelStyle(TextStyle)
-		.SetBackground(FSlateColorBrush(FStyleColors::Panel));
+		.SetBlockPadding(Margin)
+		.SetBackgroundPadding(Margin)
+		.SetCheckBoxPadding(Margin)
+		.SetBackground(FSlateColorBrush(FStyleColors::Transparent));
 
 	Set("ViewportToolbar", ToolbarStyle);
+
 	Set("AvaInteractiveTools.ToggleViewportToolbar", new IMAGE_BRUSH_SVG("Icons/EditorIcons/visible", Icon16x16));
 	Set("AvaInteractiveTools.OpenViewportToolbarSettings", new IMAGE_BRUSH_SVG("Icons/EditorIcons/settings", Icon16x16));
+	Set("AvaInteractiveTools.Dropdown", new IMAGE_BRUSH_SVG("Icons/EditorIcons/chevron-down", Icon16x16));
 
 	Set("Icons.Toolbox", new IMAGE_BRUSH(TEXT("Icons/ToolboxIcons/toolbox"), Icon16x16));
 
diff --git a/Engine/Plugins/Experimental/Avalanche/Source/AvalancheInteractiveTools/Private/EdMode/AvaInteractiveToolsEdModeToolkit.cpp b/Engine/Plugins/Experimental/Avalanche/Source/AvalancheInteractiveTools/Private/EdMode/AvaInteractiveToolsEdModeToolkit.cpp
--- a/Engine/Plugins/Experimental/Avalanche/Source/AvalancheInteractiveTools/Private/EdMode/AvaInteractiveToolsEdModeToolkit.cpp
+++ b/Engine/Plugins/Experimental/Avalanche/Source/AvalancheInteractiveTools/Private/EdMode/AvaInteractiveToolsEdModeToolkit.cpp
@@ -7,6 +7,7 @@
 #include "AvaInteractiveToolsSettings.h"
 #include "AvaInteractiveToolsStyle.h"
 #include "IAvalancheInteractiveToolsModule.h"
+#include "Toolkits/AssetEditorModeUILayer.h"
 #include "Tools/UEdMode.h"
 #include "Widgets/Layout/SBorder.h"
 #include "Widgets/Layout/SBox.h"
@@ -43,6 +44,32 @@ FText FAvaInteractiveToolsEdModeToolkit::GetBaseToolkitName() const
 	return LOCTEXT("DisplayName", "AvaInteractiveToolsEdMode Tool");
 }
 
+void FAvaInteractiveToolsEdModeToolkit::GetToolPaletteNames(TArray<FName>& InPaletteName) const
+{
+	const TMap<FName, TSharedPtr<FUICommandInfo>>& Categories = IAvalancheInteractiveToolsModule::Get().GetCategories();
+	Categories.GetKeys(InPaletteName);
+}
+
+FText FAvaInteractiveToolsEdModeToolkit::GetToolPaletteDisplayName(FName InPaletteName) const
+{
+	const TMap<FName, TSharedPtr<FUICommandInfo>>& Categories = IAvalancheInteractiveToolsModule::Get().GetCategories();
+
+	if (const TSharedPtr<FUICommandInfo>* CommandInfo = Categories.Find(InPaletteName))
+	{
+		return CommandInfo->Get()->GetLabel();
+	}
+
+	return FText::FromName(InPaletteName);
+}
+
+void FAvaInteractiveToolsEdModeToolkit::OnToolPaletteChanged(FName InPaletteName)
+{
+	if (UAvaInteractiveToolsEdMode* AvaInteractiveToolsEdMode = Cast<UAvaInteractiveToolsEdMode>(GetScriptableEditorMode().Get()))
+	{
+		AvaInteractiveToolsEdMode->OnToolPaletteChanged(InPaletteName);
+	}
+}
+
 void FAvaInteractiveToolsEdModeToolkit::Init(const TSharedPtr<IToolkitHost>& InitToolkitHost, TWeakObjectPtr<UEdMode> InOwningMode)
 {
 	FModeToolkit::Init(InitToolkitHost, InOwningMode);
@@ -60,14 +87,16 @@ void FAvaInteractiveToolsEdModeToolkit::OnToolEnded(UInteractiveToolManager* InM
 	// Nothing
 }
 
-void FAvaInteractiveToolsEdModeToolkit::InvokeUI()
-{
-	// Nothing
-}
-
 void FAvaInteractiveToolsEdModeToolkit::RequestModeUITabs()
 {
-	// Nothing
+	if (!HasIntegratedToolPalettes() && !HasToolkitBuilder())
+	{
+		const TSharedPtr<FAssetEditorModeUILayer> ModeUILayerPtr = ModeUILayer.Pin();
+		ToolbarInfo.OnSpawnTab = FOnSpawnTab::CreateSP(SharedThis(this), &FAvaInteractiveToolsEdModeToolkit::MakeModeToolbarTab);
+		ToolbarInfo.TabLabel = LOCTEXT("MotionDesignToolboxTab", "Motion Design Toolbox");
+		ToolbarInfo.TabTooltip = LOCTEXT("MotionDesignToolboxTabTooltipText", "Opens the Motion Design Toolbox tab.");
+		ModeUILayerPtr->SetModePanelInfo(UAssetEditorUISubsystem::VerticalToolbarID, ToolbarInfo);
+	}
 }
 
 void FAvaInteractiveToolsEdModeToolkit::MakeViewportOverlayToolbar()
@@ -115,13 +144,25 @@ void FAvaInteractiveToolsEdModeToolkit::MakeViewportOverlayToolbar()
 				SNew(SAvaInteractiveToolsToolbarCategoryButton)
 				.CommandList(ToolkitCommandList)
 				.ShowLabel(Settings->bViewportToolbarLabelEnabled)
-				.Category(Category.Key)
+				.ToolCategory(Category.Key)
 			);
 		}
 
 		InToolbarBuilder.AddSeparator();
-		InToolbarBuilder.AddToolBarButton(FAvaInteractiveToolsCommands::Get().ToggleViewportToolbar);
-		InToolbarBuilder.AddToolBarButton(FAvaInteractiveToolsCommands::Get().OpenViewportToolbarSettings);
+
+		InToolbarBuilder.AddWidget(
+			SNew(SAvaInteractiveToolsToolbarCategoryButton)
+			.CommandList(ToolkitCommandList)
+			.ShowLabel(Settings->bViewportToolbarLabelEnabled)
+			.Command(FAvaInteractiveToolsCommands::Get().ToggleViewportToolbar)
+		);
+
+		InToolbarBuilder.AddWidget(
+			SNew(SAvaInteractiveToolsToolbarCategoryButton)
+			.CommandList(ToolkitCommandList)
+			.ShowLabel(Settings->bViewportToolbarLabelEnabled)
+			.Command(FAvaInteractiveToolsCommands::Get().OpenViewportToolbarSettings)
+		);
 	};
 
 	TSharedPtr<SWidget> ToolbarWidget;
@@ -157,16 +198,15 @@ void FAvaInteractiveToolsEdModeToolkit::MakeViewportOverlayToolbar()
 		VAlign = VAlign_Center;
 	}
 
-	SAssignNew(ViewportToolbarWidget, SBorder)
+	SAssignNew(ViewportToolbarWidget, SVerticalBox)
+		+ SVerticalBox::Slot()
 		.HAlign(HAlign)
 		.VAlign(VAlign)
-		.BorderBackgroundColor(FSlateColor(EStyleColor::Panel))
 		.Padding(0.f)
 		[
-			SNew(SBox)
-			.HAlign(HAlign_Center)
-			.VAlign(VAlign_Center)
-			.Padding(5.f)
+			SNew(SBorder)
+			.BorderImage(FAppStyle::Get().GetBrush("EditorViewport.OverlayBrush"))
+			.Padding(FMargin(3.0f, 6.0f, 3.f, 6.f))
 			[
 				ToolbarWidget.ToSharedRef()
 			]
diff --git a/Engine/Plugins/Experimental/Avalanche/Source/AvalancheInteractiveTools/Private/EdMode/AvaInteractiveToolsEdModeToolkit.h b/Engine/Plugins/Experimental/Avalanche/Source/AvalancheInteractiveTools/Private/EdMode/AvaInteractiveToolsEdModeToolkit.h
--- a/Engine/Plugins/Experimental/Avalanche/Source/AvalancheInteractiveTools/Private/EdMode/AvaInteractiveToolsEdModeToolkit.h
+++ b/Engine/Plugins/Experimental/Avalanche/Source/AvalancheInteractiveTools/Private/EdMode/AvaInteractiveToolsEdModeToolkit.h
@@ -23,6 +23,9 @@ protected:
 	//~ Begin IToolkit
 	virtual FName GetToolkitFName() const override;
 	virtual FText GetBaseToolkitName() const override;
+	virtual void GetToolPaletteNames(TArray<FName>& PaletteNames) const override;
+	virtual FText GetToolPaletteDisplayName(FName Palette) const override;
+	virtual void OnToolPaletteChanged(FName PaletteName) override;
 	//~ End IToolkit
 
 	//~ Begin FModeToolkit
@@ -31,7 +34,6 @@ protected:
 	virtual bool HasExclusiveToolPalettes() const override { return false; }
 	virtual void OnToolStarted(UInteractiveToolManager* InManager, UInteractiveTool* InTool) override;
 	virtual void OnToolEnded(UInteractiveToolManager* InManager, UInteractiveTool* InTool) override;
-	virtual void InvokeUI() override;
 	//~ End FModeToolkit
 
 	//~ Begin FModeToolkit
diff --git a/Engine/Plugins/Experimental/Avalanche/Source/AvalancheInteractiveTools/Private/Tools/AvaInteractiveToolsActorToolNull.cpp b/Engine/Plugins/Experimental/Avalanche/Source/AvalancheInteractiveTools/Private/Tools/AvaInteractiveToolsActorToolNull.cpp
--- a/Engine/Plugins/Experimental/Avalanche/Source/AvalancheInteractiveTools/Private/Tools/AvaInteractiveToolsActorToolNull.cpp
+++ b/Engine/Plugins/Experimental/Avalanche/Source/AvalancheInteractiveTools/Private/Tools/AvaInteractiveToolsActorToolNull.cpp
@@ -19,7 +19,7 @@ void UAvaInteractiveToolsActorToolNull::OnRegisterTool(IAvalancheInteractiveTool
 	{
 		FAvaInteractiveToolsCommands::Get().Tool_Actor_Null,
 		TEXT("Null Actor Tool"),
-		1000,
+		2000,
 		FAvalancheInteractiveToolsCreateBuilder::CreateLambda(
 			[](UEdMode* InEdMode)
 			{
diff --git a/Engine/Plugins/Experimental/Avalanche/Source/AvalancheInteractiveTools/Private/Tools/AvaInteractiveToolsToolBase.cpp b/Engine/Plugins/Experimental/Avalanche/Source/AvalancheInteractiveTools/Private/Tools/AvaInteractiveToolsToolBase.cpp
--- a/Engine/Plugins/Experimental/Avalanche/Source/AvalancheInteractiveTools/Private/Tools/AvaInteractiveToolsToolBase.cpp
+++ b/Engine/Plugins/Experimental/Avalanche/Source/AvalancheInteractiveTools/Private/Tools/AvaInteractiveToolsToolBase.cpp
@@ -1,6 +1,8 @@
 // Copyright Epic Games, Inc. All Rights Reserved.
 
 #include "Tools/AvaInteractiveToolsToolBase.h"
+
+#include "ActorFactories/ActorFactory.h"
 #include "AvaInteractiveToolsSettings.h"
 #include "AvalancheInteractiveToolsModule.h"
 #include "AvaViewportUtils.h"
@@ -228,54 +230,66 @@ AActor* UAvaInteractiveToolsToolBase::SpawnActor(TSubclassOf<AActor> InActorClas
 AActor* UAvaInteractiveToolsToolBase::SpawnActor(TSubclassOf<AActor> InActorClass, EAvaViewportStatus InViewportStatus, 
 	const FVector2f& InViewportPosition, bool bInPreview, FString* InActorLabelOverride)
 {
-	UWorld* World = nullptr;;
-	FVector SpawnLocation = FVector::ZeroVector;
-	FRotator SpawnRotation = FRotator::ZeroRotator;
-	if (IToolsContextQueriesAPI* ContextAPI = GetToolManager()->GetContextQueriesAPI())
+	auto GetSpawnContext = [this, &InViewportPosition, &InViewportStatus](UWorld*& OutWorld, FVector& OutSpawnLocation, FRotator& OutSpawnRotation)->bool
 	{
-		const bool bUseIdentityLocation = UseIdentityLocation();
-		const bool bUseIdentityRotation = UseIdentityRotation();
+		bool bSuccess = false;
 
-		if (FViewport* Viewport = ContextAPI->GetFocusedViewport())
+		if (const IToolsContextQueriesAPI* ContextAPI = GetToolManager()->GetContextQueriesAPI())
 		{
-			FVector CameraForward = FVector(1, 0, 0);
+			const bool bUseIdentityLocation = UseIdentityLocation();
+			const bool bUseIdentityRotation = UseIdentityRotation();
 
-			if (TSharedPtr<IAvaViewportClient> AvaViewportClient = FAvaViewportUtils::GetAvaViewportClient(Viewport))
+			if (FViewport* Viewport = ContextAPI->GetFocusedViewport())
 			{
-				CameraForward = AvaViewportClient->GetViewportViewTransform().TransformVectorNoScale(FVector::ForwardVector);
-				CameraForward.Z = 0;
-			}
+				FVector CameraForward = FVector(1, 0, 0);
 
-			const bool bSuccess = ViewportPositionToWorldPositionAndOrientation(
-				InViewportStatus,
-				InViewportPosition,
-				World,
-				SpawnLocation,
-				SpawnRotation
-			);
+				if (const TSharedPtr<IAvaViewportClient> AvaViewportClient = FAvaViewportUtils::GetAvaViewportClient(Viewport))
+				{
+					CameraForward = AvaViewportClient->GetViewportViewTransform().TransformVectorNoScale(FVector::ForwardVector);
+					CameraForward.Z = 0;
+				}
 
-			if (!bSuccess)
-			{
-				return nullptr;
-			}
+				bSuccess = ViewportPositionToWorldPositionAndOrientation(
+					InViewportStatus,
+					InViewportPosition,
+					OutWorld,
+					OutSpawnLocation,
+					OutSpawnRotation
+				);
 
-			if (bUseIdentityLocation)
-			{
-				SpawnLocation = FVector::ZeroVector;
-			}
+				if (!bSuccess || !IsValid(OutWorld))
+				{
+					return false;
+				}
 
-			if (bUseIdentityRotation)
-			{
-				if (CameraForward.Dot(FVector::ForwardVector) >= 0)
+				if (bUseIdentityLocation)
 				{
-					SpawnRotation = FRotator::ZeroRotator;
+					OutSpawnLocation = FVector::ZeroVector;
 				}
-				else
+
+				if (bUseIdentityRotation)
 				{
-					SpawnRotation = FRotator(0, 180, 0);
+					if (CameraForward.Dot(FVector::ForwardVector) >= 0)
+					{
+						OutSpawnRotation = FRotator::ZeroRotator;
+					}
+					else
+					{
+						OutSpawnRotation = FRotator(0, 180, 0);
+					}
 				}
 			}
 		}
+
+		return bSuccess;
+	};
+	
+	UWorld* World = nullptr;
+	FVector SpawnLocation = FVector::ZeroVector;
+	FRotator SpawnRotation = FRotator::ZeroRotator;
+	if (!GetSpawnContext(World, SpawnLocation, SpawnRotation))
+	{
+		return nullptr;
 	}
 
 	FActorSpawnParameters SpawnParams;
@@ -298,7 +312,38 @@ AActor* UAvaInteractiveToolsToolBase::SpawnActor(TSubclassOf<AActor> InActorClas
 		SpawnParams.ObjectFlags |= RF_Transactional;
 	}
 
-	AActor* NewActor = World->SpawnActor<AActor>(InActorClass, SpawnLocation, SpawnRotation, SpawnParams);
+	UActorFactory* ActorFactory = nullptr;
+	const FAvaInteractiveToolsToolParameters* ActiveToolParameters = GetActiveToolParameters();
+
+	if (ActiveToolParameters)
+	{
+		if (ActiveToolParameters->Factory)
+		{
+			ActorFactory = ActiveToolParameters->Factory;
+		}
+		else if (ActiveToolParameters->FactoryClass)
+		{
+			ActorFactory = GEditor->FindActorFactoryByClass(ActiveToolParameters->FactoryClass);
+		}
+	}
+
+	if (!ActorFactory)
+	{
+		ActorFactory = GEditor->FindActorFactoryForActorClass(InActorClass);
+	}
+
+	AActor* NewActor = nullptr;
+
+	if (ActorFactory)
+	{
+		const FTransform SpawnTransform(SpawnRotation, SpawnLocation);
+		NewActor = ActorFactory->CreateActor(InActorClass.Get(), World->PersistentLevel, SpawnTransform, SpawnParams);
+	}
+
+	if (!NewActor)
+	{
+		NewActor = World->SpawnActor<AActor>(InActorClass, SpawnLocation, SpawnRotation, SpawnParams);
+	}
 
 	if (bUseOverride)
 	{
diff --git a/Engine/Plugins/Experimental/Avalanche/Source/AvalancheInteractiveTools/Private/Widgets/SAvaInteractiveToolsToolbarCategoryButton.cpp b/Engine/Plugins/Experimental/Avalanche/Source/AvalancheInteractiveTools/Private/Widgets/SAvaInteractiveToolsToolbarCategoryButton.cpp
--- a/Engine/Plugins/Experimental/Avalanche/Source/AvalancheInteractiveTools/Private/Widgets/SAvaInteractiveToolsToolbarCategoryButton.cpp
+++ b/Engine/Plugins/Experimental/Avalanche/Source/AvalancheInteractiveTools/Private/Widgets/SAvaInteractiveToolsToolbarCategoryButton.cpp
@@ -2,6 +2,7 @@
 
 #include "Widgets/SAvaInteractiveToolsToolbarCategoryButton.h"
 
+#include "AvaInteractiveToolsSettings.h"
 #include "Application/ThrottleManager.h"
 #include "AvaInteractiveToolsStyle.h"
 #include "Framework/Application/IMenu.h"
@@ -9,18 +10,27 @@
 #include "Framework/Commands/UICommandList.h"
 #include "Framework/MultiBox/MultiBoxBuilder.h"
 #include "IAvalancheInteractiveToolsModule.h"
+#include "Widgets/Images/SImage.h"
 
 void SAvaInteractiveToolsToolbarCategoryButton::Construct(const FArguments& InArgs)
 {
 	CommandList = InArgs._CommandList;
 	bShowLabel = InArgs._ShowLabel;
-	Category = InArgs._Category;
+	ToolCategory = InArgs._ToolCategory;
+	Command = InArgs._Command;
 
-	IAvalancheInteractiveToolsModule& AITModule = IAvalancheInteractiveToolsModule::Get();
-
-	check(CommandList.IsValid() && AITModule.GetCategories().Contains(Category) && !AITModule.GetTools(Category)->IsEmpty())
+	check(CommandList.IsValid())
 
-	ActiveToolIdentifier = (*AITModule.GetTools(Category))[0].ToolIdentifier;
+	if (!ToolCategory.IsNone())
+	{
+		IAvalancheInteractiveToolsModule& AITModule = IAvalancheInteractiveToolsModule::Get();
+		check(AITModule.GetCategories().Contains(ToolCategory) && !AITModule.GetTools(ToolCategory)->IsEmpty())
+		ActiveToolIdentifier = (*AITModule.GetTools(ToolCategory))[0].ToolIdentifier;
+	}
+	else
+	{
+		check(Command.IsValid())
+	}
 
 	CreateActiveCommandWidget();
 
@@ -36,7 +46,11 @@ FReply SAvaInteractiveToolsToolbarCategoryButton::OnMouseButtonDown(const FGeome
 {
 	if (InMouseEvent.GetEffectingButton() == EKeys::RightMouseButton)
 	{
-		ShowCommandsContextMenu();
+		if (!ToolCategory.IsNone())
+		{
+			ShowCommandsContextMenu();
+		}
+
 		return FReply::Handled();
 	}
 
@@ -45,30 +59,60 @@ FReply SAvaInteractiveToolsToolbarCategoryButton::OnMouseButtonDown(const FGeome
 
 void SAvaInteractiveToolsToolbarCategoryButton::CreateActiveCommandWidget()
 {
-	IAvalancheInteractiveToolsModule& AITModule = IAvalancheInteractiveToolsModule::Get();
-	const FAvaInteractiveToolsToolParameters* Tool = AITModule.GetTool(ActiveToolIdentifier);
+	FVerticalToolBarBuilder CategoryToolbarBuilder(
+			CommandList,
+			FMultiBoxCustomization::None,
+			TSharedPtr<FExtender>()
+			, /* ForceSmallIcons */ true);
 
-	if (Tool && Tool->UICommand.IsValid())
-	{
-		FVerticalToolBarBuilder CategoryToolbarBuilder(
-		CommandList,
-		FMultiBoxCustomization::None,
-		TSharedPtr<FExtender>()
-		, /* ForceSmallIcons */ true);
+	CategoryToolbarBuilder.SetStyle(&FAvaInteractiveToolsStyle::Get(), "ViewportToolbar");
+	CategoryToolbarBuilder.SetLabelVisibility(bShowLabel ? EVisibility::Visible : EVisibility::Collapsed);
 
-		CategoryToolbarBuilder.SetStyle(&FAvaInteractiveToolsStyle::Get(), "ViewportToolbar");
-		CategoryToolbarBuilder.SetLabelVisibility(bShowLabel ? EVisibility::Visible : EVisibility::Collapsed);
+	if (!ToolCategory.IsNone())
+	{
+		IAvalancheInteractiveToolsModule& AITModule = IAvalancheInteractiveToolsModule::Get();
+		const FAvaInteractiveToolsToolParameters* Tool = AITModule.GetTool(ActiveToolIdentifier);
 
-		TAttribute<FText> LabelAttribute = Tool->Label.IsSet() ? Tool->Label.GetValue() : TAttribute<FText>();
-		TAttribute<FText> TooltipAttribute = Tool->Tooltip.IsSet() ? Tool->Tooltip.GetValue() : TAttribute<FText>();
-		TAttribute<FSlateIcon> IconAttribute = Tool->Icon.IsSet() ? Tool->Icon.GetValue() : TAttribute<FSlateIcon>();
-		CategoryToolbarBuilder.AddToolBarButton(Tool->UICommand, NAME_None, LabelAttribute, TooltipAttribute, IconAttribute);
+		if (Tool && Tool->UICommand.IsValid())
+		{
+			TAttribute<FText> LabelAttribute = Tool->Label.IsSet() ? Tool->Label.GetValue() : TAttribute<FText>();
+			TAttribute<FText> TooltipAttribute = Tool->Tooltip.IsSet() ? Tool->Tooltip.GetValue() : TAttribute<FText>();
+			TAttribute<FSlateIcon> IconAttribute = Tool->Icon.IsSet() ? Tool->Icon.GetValue() : TAttribute<FSlateIcon>();
+			CategoryToolbarBuilder.AddToolBarButton(Tool->UICommand, NAME_None, LabelAttribute, TooltipAttribute, IconAttribute);
+		}
+	}
+	else
+	{
+		CategoryToolbarBuilder.AddToolBarButton(Command);
+	}
 
-		ChildSlot
+	ChildSlot
+	[
+		SNew(SOverlay)
+		+ SOverlay::Slot()
+		.ZOrder(1)
+		.HAlign(HAlign_Right)
+		.VAlign(bShowLabel ? VAlign_Top : VAlign_Bottom)
+		.Padding(1.f)
 		[
-			CategoryToolbarBuilder.MakeWidget()
-		];
-	}
+			SNew(SImage)
+			.DesiredSizeOverride(FVector2D(8, 8))
+			// Show dropdown chevron
+			.Visibility(!ToolCategory.IsNone() ? EVisibility::HitTestInvisible : EVisibility::Collapsed)
+			.Image(FAvaInteractiveToolsStyle::Get().GetBrush("AvaInteractiveTools.Dropdown"))
+			.ColorAndOpacity(FSlateColor::UseForeground())
+		]
+		+ SOverlay::Slot()
+		.ZOrder(0)
+		[
+			SNew(SBox)
+			.WidthOverride(bShowLabel ? 38.f : FOptionalSize())
+			.HeightOverride(bShowLabel ? 38.f : FOptionalSize())
+			[
+				CategoryToolbarBuilder.MakeWidget()
+			]
+		]
+	];
 }
 
 void SAvaInteractiveToolsToolbarCategoryButton::ShowCommandsContextMenu()
@@ -78,18 +122,33 @@ void SAvaInteractiveToolsToolbarCategoryButton::ShowCommandsContextMenu()
 	// Needed otherwise cannot select entries in context toolbar below
 	FSlateThrottleManager::Get().DisableThrottle(true);
 
-	const FVector2D MousePosition = FSlateApplication::Get().GetLastCursorPos();
+	FVector2D MenuPosition = FSlateApplication::Get().GetLastCursorPos();
 
-	FVerticalToolBarBuilder CategoryToolbarBuilder(
-		CommandList,
-		FMultiBoxCustomization::None,
-		TSharedPtr<FExtender>()
-		, /* ForceSmallIcons */ false);
+	if (UAvaInteractiveToolsSettings* ToolsSettings = UAvaInteractiveToolsSettings::Get())
+	{
+		const FGeometry& WidgetGeometry = GetCachedGeometry();
 
-	CategoryToolbarBuilder.SetStyle(&FAvaInteractiveToolsStyle::Get(), "ViewportToolbar");
+		FVector2D LocalPosition;
+		if (ToolsSettings->ViewportToolbarPosition == EAvaInteractiveToolsViewportToolbarPosition::Bottom
+			|| ToolsSettings->ViewportToolbarPosition == EAvaInteractiveToolsViewportToolbarPosition::Top)
+		{
+			LocalPosition = FVector2D(0, WidgetGeometry.GetLocalSize().Y);
+		}
+		else if (ToolsSettings->ViewportToolbarPosition == EAvaInteractiveToolsViewportToolbarPosition::Left
+			|| ToolsSettings->ViewportToolbarPosition == EAvaInteractiveToolsViewportToolbarPosition::Right)
+		{
+			LocalPosition = FVector2D(WidgetGeometry.GetLocalSize().X, 0);
+		}
+
+        MenuPosition = WidgetGeometry.LocalToAbsolute(LocalPosition);
+	}
+
+	FMenuBuilder CategoryToolbarBuilder(/** CloseAfterSelection */true, CommandList);
+	CategoryToolbarBuilder.SetSearchable(false);
+	CategoryToolbarBuilder.SetCheckBoxStyle("TransparentCheckBox");
 
 	IAvalancheInteractiveToolsModule& AITModule = IAvalancheInteractiveToolsModule::Get();
-	for (const FAvaInteractiveToolsToolParameters& Tool : (*AITModule.GetTools(Category)))
+	for (const FAvaInteractiveToolsToolParameters& Tool : (*AITModule.GetTools(ToolCategory)))
 	{
 		if (Tool.UICommand.IsValid())
 		{
@@ -97,7 +156,7 @@ void SAvaInteractiveToolsToolbarCategoryButton::ShowCommandsContextMenu()
 			TAttribute<FText> TooltipAttribute = Tool.Tooltip.IsSet() ? Tool.Tooltip.GetValue() : TAttribute<FText>();
 			TAttribute<FSlateIcon> IconAttribute = Tool.Icon.IsSet() ? Tool.Icon.GetValue() : TAttribute<FSlateIcon>();
 
-			CategoryToolbarBuilder.AddToolBarButton(Tool.UICommand, NAME_None, LabelAttribute, TooltipAttribute, IconAttribute);
+			CategoryToolbarBuilder.AddMenuEntry(Tool.UICommand, NAME_None, LabelAttribute, TooltipAttribute, IconAttribute.Get());
 		}
 	}
 
@@ -105,10 +164,10 @@ void SAvaInteractiveToolsToolbarCategoryButton::ShowCommandsContextMenu()
 		FSlateApplication::Get().GetActiveTopLevelWindow().ToSharedRef(),
 		FWidgetPath(),
 		CategoryToolbarBuilder.MakeWidget(),
-		MousePosition,
+		MenuPosition,
 		FPopupTransitionEffect(FPopupTransitionEffect::None)
 	);
-	
+
 	ContextMenu->GetOnMenuDismissed().AddLambda([this](TSharedRef<IMenu>)
 	{
 		FSlateThrottleManager::Get().DisableThrottle(false);
@@ -127,8 +186,8 @@ void SAvaInteractiveToolsToolbarCategoryButton::HideCommandsContextMenu()
 
 void SAvaInteractiveToolsToolbarCategoryButton::OnToolActivated(const FString& InToolIdentifier)
 {
-	const FName ToolCategory = IAvalancheInteractiveToolsModule::Get().GetToolCategory(InToolIdentifier);
-	if (ToolCategory.IsEqual(Category))
+	const FName ActiveToolCategory = IAvalancheInteractiveToolsModule::Get().GetToolCategory(InToolIdentifier);
+	if (ActiveToolCategory.IsEqual(ToolCategory))
 	{
 		ActiveToolIdentifier = InToolIdentifier;
 		CreateActiveCommandWidget();
diff --git a/Engine/Plugins/Experimental/Avalanche/Source/AvalancheInteractiveTools/Private/Widgets/SAvaInteractiveToolsToolbarCategoryButton.h b/Engine/Plugins/Experimental/Avalanche/Source/AvalancheInteractiveTools/Private/Widgets/SAvaInteractiveToolsToolbarCategoryButton.h
--- a/Engine/Plugins/Experimental/Avalanche/Source/AvalancheInteractiveTools/Private/Widgets/SAvaInteractiveToolsToolbarCategoryButton.h
+++ b/Engine/Plugins/Experimental/Avalanche/Source/AvalancheInteractiveTools/Private/Widgets/SAvaInteractiveToolsToolbarCategoryButton.h
@@ -5,6 +5,7 @@
 #include "Templates/SharedPointer.h"
 #include "Widgets/SCompoundWidget.h"
 
+class FUICommandInfo;
 class FUICommandList;
 class IMenu;
 
@@ -13,8 +14,10 @@ class SAvaInteractiveToolsToolbarCategoryButton : public SCompoundWidget
 public:
 	SLATE_BEGIN_ARGS(SAvaInteractiveToolsToolbarCategoryButton)
 	{}
-		/** Category to retrieve tools from */
-		SLATE_ARGUMENT(FName, Category)
+		/** Category to retrieve ITF commands from */
+		SLATE_ARGUMENT(FName, ToolCategory)
+		/** If not category is provided then a command must be provided */
+		SLATE_ARGUMENT(TSharedPtr<FUICommandInfo>, Command)
 		/** Command list to execute selected command */
 		SLATE_ARGUMENT(TSharedPtr<FUICommandList>, CommandList)
 		/** Show label under category button */
@@ -36,9 +39,10 @@ protected:
 	void OnToolActivated(const FString& InToolIdentifier);
 
 private:
+	TSharedPtr<FUICommandInfo> Command;
 	TSharedPtr<FUICommandList> CommandList;
 	TSharedPtr<IMenu> ContextMenu;
 	bool bShowLabel = false;
-	FName Category = NAME_None;
+	FName ToolCategory = NAME_None;
 	FString ActiveToolIdentifier;
 };
\ No newline at end of file
diff --git a/Engine/Plugins/Experimental/Avalanche/Source/AvalancheShapesEditor/Private/AvaShapeFactory.cpp b/Engine/Plugins/Experimental/Avalanche/Source/AvalancheShapesEditor/Private/AvaShapeFactory.cpp
--- a/Engine/Plugins/Experimental/Avalanche/Source/AvalancheShapesEditor/Private/AvaShapeFactory.cpp
+++ b/Engine/Plugins/Experimental/Avalanche/Source/AvalancheShapesEditor/Private/AvaShapeFactory.cpp
@@ -14,6 +14,7 @@
 UAvaShapeFactory::UAvaShapeFactory()
 {
 	NewActorClass = AAvaShapeActor::StaticClass();
+	MeshSize = FVector(UAvaShapesEditorShapeToolBase::DefaultDim);
 }
 
 void UAvaShapeFactory::SetMeshClass(TSubclassOf<UAvaShapeDynamicMeshBase> InMeshClass)
@@ -21,6 +22,16 @@ void UAvaShapeFactory::SetMeshClass(TSubclassOf<UAvaShapeDynamicMeshBase> InMesh
 	MeshClass = InMeshClass;
 }
 
+void UAvaShapeFactory::SetMeshSize(const FVector& InMeshSize)
+{
+	MeshSize = InMeshSize;
+}
+
+void UAvaShapeFactory::SetMeshFunction(TFunction<void(UAvaShapeDynamicMeshBase*)> InFunction)
+{
+	MeshFunction = InFunction;
+}
+
 bool UAvaShapeFactory::CanCreateActorFrom(const FAssetData& AssetData, FText& OutErrorMsg)
 {
 	if (!MeshClass)
@@ -59,26 +70,16 @@ AActor* UAvaShapeFactory::SpawnActor(UObject* InAsset, ULevel* InLevel, const FT
 	AAvaShapeActor* ShapeActor = World->SpawnActor<AAvaShapeActor>(AAvaShapeActor::StaticClass(), InTransform, Params);
 
 	UAvaShapeDynamicMeshBase* MeshBase = NewObject<UAvaShapeDynamicMeshBase>(ShapeActor, MeshClass.Get());
-
-	if (UAvaShape2DDynMeshBase* Mesh2D = Cast<UAvaShape2DDynMeshBase>(MeshBase))
-	{
-		Mesh2D->SetSize2D({
-			UAvaShapesEditorShapeToolBase::DefaultDim,
-			UAvaShapesEditorShapeToolBase::DefaultDim
-		});
-	}
-	else if (UAvaShape3DDynMeshBase* Mesh3D = Cast<UAvaShape3DDynMeshBase>(MeshBase))
-	{
-		Mesh3D->SetSize3D({
-			UAvaShapesEditorShapeToolBase::DefaultDim,
-			UAvaShapesEditorShapeToolBase::DefaultDim,
-			UAvaShapesEditorShapeToolBase::DefaultDim
-		});
-	}
+	MeshBase->SetSize3D(MeshSize);
 
 	const FAvaShapeParametricMaterial DefaultParametricMaterial;
 	MeshBase->SetParametricMaterial(UAvaShapeDynamicMeshBase::MESH_INDEX_PRIMARY, DefaultParametricMaterial);
 
+	if (MeshFunction.IsSet())
+	{
+		MeshFunction(MeshBase);
+	}
+
 	ShapeActor->SetDynamicMesh(MeshBase);
 
 	return ShapeActor;
diff --git a/Engine/Plugins/Experimental/Avalanche/Source/AvalancheShapesEditor/Private/AvaShapeFactory.h b/Engine/Plugins/Experimental/Avalanche/Source/AvalancheShapesEditor/Private/AvaShapeFactory.h
--- a/Engine/Plugins/Experimental/Avalanche/Source/AvalancheShapesEditor/Private/AvaShapeFactory.h
+++ b/Engine/Plugins/Experimental/Avalanche/Source/AvalancheShapesEditor/Private/AvaShapeFactory.h
@@ -17,10 +17,10 @@ public:
 	UAvaShapeFactory();
 
 	void SetMeshClass(TSubclassOf<UAvaShapeDynamicMeshBase> InMeshClass);
+	void SetMeshSize(const FVector& InMeshSize);
+	void SetMeshFunction(TFunction<void(UAvaShapeDynamicMeshBase*)> InFunction);
 
 protected:
-	TSubclassOf<UAvaShapeDynamicMeshBase> MeshClass;
-
 	//~ Begin UActorFactory
 	virtual bool CanCreateActorFrom(const FAssetData& AssetData, FText& OutErrorMsg) override;
 	virtual AActor* GetDefaultActor(const FAssetData& AssetData) override;
@@ -28,4 +28,8 @@ protected:
 	virtual void PostSpawnActor(UObject* InAsset, AActor* InNewActor) override;
 	virtual void PostPlaceAsset(TArrayView<const FTypedElementHandle> InHandle, const FAssetPlacementInfo& InPlacementInfo, const FPlacementOptions& InPlacementOptions) override;
 	//~ End UActorFactory
+
+	TSubclassOf<UAvaShapeDynamicMeshBase> MeshClass;
+	FVector MeshSize = FVector::ZeroVector;
+	TFunction<void(UAvaShapeDynamicMeshBase*)> MeshFunction;
 };
diff --git a/Engine/Plugins/Experimental/Avalanche/Source/AvalancheShapesEditor/Private/Tools/AvaShapesEditorShapeTool2DArrow.cpp b/Engine/Plugins/Experimental/Avalanche/Source/AvalancheShapesEditor/Private/Tools/AvaShapesEditorShapeTool2DArrow.cpp
--- a/Engine/Plugins/Experimental/Avalanche/Source/AvalancheShapesEditor/Private/Tools/AvaShapesEditorShapeTool2DArrow.cpp
+++ b/Engine/Plugins/Experimental/Avalanche/Source/AvalancheShapesEditor/Private/Tools/AvaShapesEditorShapeTool2DArrow.cpp
@@ -31,13 +31,3 @@ void UAvaShapesEditorShapeTool2DArrow::OnRegisterTool(IAvalancheInteractiveTools
 
 	InAITModule->RegisterTool(IAvalancheInteractiveToolsModule::CategoryName2D, MoveTemp(ToolParameters));
 }
-
-void UAvaShapesEditorShapeTool2DArrow::InitShape(UAvaShapeDynamicMeshBase* InShape) const
-{
-	UAvaShape2DArrowDynamicMesh* Arrow2D = Cast<UAvaShape2DArrowDynamicMesh>(InShape);
-	check(Arrow2D);
-
-	Arrow2D->SetSize2D({DefaultDim, DefaultDim});
-
-	Super::InitShape(InShape);
-}
diff --git a/Engine/Plugins/Experimental/Avalanche/Source/AvalancheShapesEditor/Private/Tools/AvaShapesEditorShapeTool2DArrow.h b/Engine/Plugins/Experimental/Avalanche/Source/AvalancheShapesEditor/Private/Tools/AvaShapesEditorShapeTool2DArrow.h
--- a/Engine/Plugins/Experimental/Avalanche/Source/AvalancheShapesEditor/Private/Tools/AvaShapesEditorShapeTool2DArrow.h
+++ b/Engine/Plugins/Experimental/Avalanche/Source/AvalancheShapesEditor/Private/Tools/AvaShapesEditorShapeTool2DArrow.h
@@ -13,12 +13,8 @@ class UAvaShapesEditorShapeTool2DArrow : public UAvaShapesEditorShapeAreaToolBas
 public:
 	UAvaShapesEditorShapeTool2DArrow();
 
+protected:
 	//~ Begin UAvaInteractiveToolsToolBase
 	virtual void OnRegisterTool(IAvalancheInteractiveToolsModule* InAITModule) override;
 	//~ End UAvaInteractiveToolsToolBase
-
-protected:
-	//~ Begin UAvaShapesEditorShapeToolBase
-	virtual void InitShape(UAvaShapeDynamicMeshBase* InShape) const override;
-	//~ End UAvaShapesEditorShapeToolBase
 };
diff --git a/Engine/Plugins/Experimental/Avalanche/Source/AvalancheShapesEditor/Private/Tools/AvaShapesEditorShapeToolBase.cpp b/Engine/Plugins/Experimental/Avalanche/Source/AvalancheShapesEditor/Private/Tools/AvaShapesEditorShapeToolBase.cpp
--- a/Engine/Plugins/Experimental/Avalanche/Source/AvalancheShapesEditor/Private/Tools/AvaShapesEditorShapeToolBase.cpp
+++ b/Engine/Plugins/Experimental/Avalanche/Source/AvalancheShapesEditor/Private/Tools/AvaShapesEditorShapeToolBase.cpp
@@ -11,6 +11,9 @@
 #include "InteractiveToolManager.h"
 #include "ToolContextInterfaces.h"
 
+const FVector UAvaShapesEditorShapeToolBase::DefaultDim = FVector(100);
+TFunction<void(UAvaShapeDynamicMeshBase*)> UAvaShapesEditorShapeToolBase::DefaultFunction = [](UAvaShapeDynamicMeshBase*){};
+
 UAvaShapesEditorShapeToolBase::UAvaShapesEditorShapeToolBase()
 {
 	ActorClass = AAvaShapeActor::StaticClass();
@@ -26,27 +29,14 @@ bool UAvaShapesEditorShapeToolBase::OnBegin()
 	return Super::OnBegin();
 }
 
-AActor* UAvaShapesEditorShapeToolBase::SpawnActor(TSubclassOf<AActor> InActorClass, EAvaViewportStatus InViewportStatus,
-	const FVector2f& InViewportPosition, bool bInPreview, FString* InActorLabelOverride)
+void UAvaShapesEditorShapeToolBase::OnActorSpawned(AActor* InActor)
 {
-	FString ActorLabelOverride = (InActorLabelOverride && !InActorLabelOverride->IsEmpty())
-		? *InActorLabelOverride
-		: GetActorNameOverride();
-
-	AActor* NewActor = Super::SpawnActor(InActorClass, InViewportStatus, InViewportPosition, bInPreview, &ActorLabelOverride);
+	Super::OnActorSpawned(InActor);
 
-	if (AAvaShapeActor* ShapeActor = Cast<AAvaShapeActor>(NewActor))
+	if (const AAvaShapeActor* ShapeActor = Cast<AAvaShapeActor>(InActor))
 	{
-		UAvaShapeDynamicMeshBase* NewShape = NewObject<UAvaShapeDynamicMeshBase>(ShapeActor, ShapeClass);
-		SetShapeSize(ShapeActor, FVector2D(DefaultDim, DefaultDim));
-		InitShape(NewShape);
-
-		ShapeActor->SetDynamicMesh(NewShape);
-
-		SetToolkitSettingsObject(NewShape);
+		SetToolkitSettingsObject(ShapeActor->GetDynamicMesh());
 	}
-
-	return NewActor;
 }
 
 FString UAvaShapesEditorShapeToolBase::GetActorNameOverride() const
@@ -59,11 +49,6 @@ FString UAvaShapesEditorShapeToolBase::GetActorNameOverride() const
 	return "";
 }
 
-void UAvaShapesEditorShapeToolBase::InitShape(UAvaShapeDynamicMeshBase* InShape) const
-{
-	InShape->SetParametricMaterial(UAvaShapeDynamicMeshBase::MESH_INDEX_PRIMARY, DefaultMaterialParams);
-}
-
 void UAvaShapesEditorShapeToolBase::SetShapeSize(AAvaShapeActor* InShapeActor, const FVector2D& InShapeSize) const
 {
 	if (!InShapeActor)
@@ -84,7 +69,8 @@ void UAvaShapesEditorShapeToolBase::SetShapeSize(AAvaShapeActor* InShapeActor, c
 	}
 	else if (UAvaShape3DDynMeshBase* Mesh3D = Cast<UAvaShape3DDynMeshBase>(MeshBase))
 	{
-		Mesh3D->SetSize3D({DefaultDim, InShapeSize.X, InShapeSize.Y});
+		const FVector Size3D = Mesh3D->GetSize3D();
+		Mesh3D->SetSize3D({Size3D.X, InShapeSize.X, InShapeSize.Y});
 	}
 	else
 	{
diff --git a/Engine/Plugins/Experimental/Avalanche/Source/AvalancheShapesEditor/Private/Tools/AvaShapesEditorShapeToolBase.h b/Engine/Plugins/Experimental/Avalanche/Source/AvalancheShapesEditor/Private/Tools/AvaShapesEditorShapeToolBase.h
--- a/Engine/Plugins/Experimental/Avalanche/Source/AvalancheShapesEditor/Private/Tools/AvaShapesEditorShapeToolBase.h
+++ b/Engine/Plugins/Experimental/Avalanche/Source/AvalancheShapesEditor/Private/Tools/AvaShapesEditorShapeToolBase.h
@@ -4,7 +4,6 @@
 
 #include "Tools/AvaInteractiveToolsActorToolBase.h"
 #include "AvaShapeFactory.h"
-#include "AvaShapeParametricMaterial.h"
 #include "Templates/SubclassOf.h"
 #include "AvaShapesEditorShapeToolBase.generated.h"
 
@@ -19,39 +18,35 @@ class UAvaShapesEditorShapeToolBase : public UAvaInteractiveToolsActorToolBase
 public:
 	using UAvaInteractiveToolsActorToolBase::SpawnActor;
 
-	static constexpr float DefaultDim = 100;
+	static const FVector DefaultDim;
+	static TFunction<void(UAvaShapeDynamicMeshBase*)> DefaultFunction;
 
 	UAvaShapesEditorShapeToolBase();
 
-	//~ Begin UAvaInteractiveToolsToolBase
-	virtual bool OnBegin() override;
-	//~ End UAvaInteractiveToolsToolBase
-
 protected:
 	template<typename InMeshClass
 		UE_REQUIRES(std::derived_from<InMeshClass, UAvaShapeDynamicMeshBase>)>
-	static UAvaShapeFactory* CreateFactory()
+	static UAvaShapeFactory* CreateFactory(const FVector& InMeshSize = DefaultDim, TFunction<void(UAvaShapeDynamicMeshBase*)> InFunction = DefaultFunction)
 	{
 		UAvaShapeFactory* Factory = CreateActorFactory<UAvaShapeFactory>();
 		Factory->SetMeshClass(InMeshClass::StaticClass());
+		Factory->SetMeshSize(InMeshSize);
+		Factory->SetMeshFunction(InFunction);
 		return Factory;
 	}
 
-	UPROPERTY()
-	TSubclassOf<UAvaShapeDynamicMeshBase> ShapeClass = nullptr;
-
-	FAvaShapeParametricMaterial DefaultMaterialParams;
-
-	virtual void InitShape(UAvaShapeDynamicMeshBase* InShape) const;
-
-	virtual void SetShapeSize(AAvaShapeActor* InShapeActor, const FVector2D& InShapeSize) const;
-
+	//~ Begin UAvaInteractiveToolsToolBase
+	virtual bool OnBegin() override;
+	//~ End UAvaInteractiveToolsToolBase
+	
 	//~ Begin UAvaInteractiveToolsToolBase
 	virtual bool UseIdentityLocation() const override;
-
-	virtual AActor* SpawnActor(TSubclassOf<AActor> InActorClass, EAvaViewportStatus InViewportStatus, 
-		const FVector2f& InViewportPosition, bool bInPreview, FString* InActorLabelOverride = nullptr) override;
+	virtual void OnActorSpawned(AActor* InActor) override;
 	//~ End UAvaInteractiveToolsToolBase
 
+	virtual void SetShapeSize(AAvaShapeActor* InShapeActor, const FVector2D& InShapeSize) const;
 	virtual FString GetActorNameOverride() const;
+
+	UPROPERTY()
+	TSubclassOf<UAvaShapeDynamicMeshBase> ShapeClass = nullptr;
 };
diff --git a/Engine/Plugins/Experimental/Avalanche/Source/AvalancheShapesEditor/Private/Tools/AvaShapesEditorShapeToolChevron.cpp b/Engine/Plugins/Experimental/Avalanche/Source/AvalancheShapesEditor/Private/Tools/AvaShapesEditorShapeToolChevron.cpp
--- a/Engine/Plugins/Experimental/Avalanche/Source/AvalancheShapesEditor/Private/Tools/AvaShapesEditorShapeToolChevron.cpp
+++ b/Engine/Plugins/Experimental/Avalanche/Source/AvalancheShapesEditor/Private/Tools/AvaShapesEditorShapeToolChevron.cpp
@@ -31,13 +31,3 @@ void UAvaShapesEditorShapeToolChevron::OnRegisterTool(IAvalancheInteractiveTools
 
 	InAITModule->RegisterTool(IAvalancheInteractiveToolsModule::CategoryName2D, MoveTemp(ToolParameters));
 }
-
-void UAvaShapesEditorShapeToolChevron::InitShape(UAvaShapeDynamicMeshBase* InShape) const
-{
-	UAvaShapeChevronDynamicMesh* Chevron = Cast<UAvaShapeChevronDynamicMesh>(InShape);
-	check(Chevron);
-
-	Chevron->SetSize2D({DefaultDim, DefaultDim});
-
-	Super::InitShape(InShape);
-}
diff --git a/Engine/Plugins/Experimental/Avalanche/Source/AvalancheShapesEditor/Private/Tools/AvaShapesEditorShapeToolChevron.h b/Engine/Plugins/Experimental/Avalanche/Source/AvalancheShapesEditor/Private/Tools/AvaShapesEditorShapeToolChevron.h
--- a/Engine/Plugins/Experimental/Avalanche/Source/AvalancheShapesEditor/Private/Tools/AvaShapesEditorShapeToolChevron.h
+++ b/Engine/Plugins/Experimental/Avalanche/Source/AvalancheShapesEditor/Private/Tools/AvaShapesEditorShapeToolChevron.h
@@ -13,12 +13,8 @@ class UAvaShapesEditorShapeToolChevron : public UAvaShapesEditorShapeAreaToolBas
 public:
 	UAvaShapesEditorShapeToolChevron();
 
+protected:
 	//~ Begin UAvaInteractiveToolsToolBase
 	virtual void OnRegisterTool(IAvalancheInteractiveToolsModule* InAITModule) override;
 	//~ End UAvaInteractiveToolsToolBase
-
-protected:
-	//~ Begin UAvaShapesEditorShapeToolBase
-	virtual void InitShape(UAvaShapeDynamicMeshBase* InShape) const override;
-	//~ End UAvaShapesEditorShapeToolBase
 };
diff --git a/Engine/Plugins/Experimental/Avalanche/Source/AvalancheShapesEditor/Private/Tools/AvaShapesEditorShapeToolCone.cpp b/Engine/Plugins/Experimental/Avalanche/Source/AvalancheShapesEditor/Private/Tools/AvaShapesEditorShapeToolCone.cpp
--- a/Engine/Plugins/Experimental/Avalanche/Source/AvalancheShapesEditor/Private/Tools/AvaShapesEditorShapeToolCone.cpp
+++ b/Engine/Plugins/Experimental/Avalanche/Source/AvalancheShapesEditor/Private/Tools/AvaShapesEditorShapeToolCone.cpp
@@ -6,9 +6,6 @@
 #include "DynamicMeshes/AvaShapeConeDynMesh.h"
 #include "UObject/ConstructorHelpers.h"
 
-const TCHAR* UAvaShapesEditorShapeToolCone::ConeToolIdentifier = TEXT("Parametric Cone Tool");
-const TCHAR* UAvaShapesEditorShapeToolCone::CylinderToolIdentifier = TEXT("Parametric Cylinder Tool");
-
 UAvaShapesEditorShapeToolCone::UAvaShapesEditorShapeToolCone()
 {
 	ShapeClass = UAvaShapeConeDynamicMesh::StaticClass();
@@ -21,7 +18,7 @@ void UAvaShapesEditorShapeToolCone::OnRegisterTool(IAvalancheInteractiveToolsMod
 	FAvaInteractiveToolsToolParameters ConeToolParameters = 
 	{
 		FAvaShapesEditorCommands::Get().Tool_Shape_Cone,
-		ConeToolIdentifier,
+		TEXT("Parametric Cone Tool"),
 		3000,
 		FAvalancheInteractiveToolsCreateBuilder::CreateLambda(
 			[](UEdMode* InEdMode)
@@ -37,7 +34,7 @@ void UAvaShapesEditorShapeToolCone::OnRegisterTool(IAvalancheInteractiveToolsMod
 	FAvaInteractiveToolsToolParameters CylinderToolParameters = 
 	{
 		FAvaShapesEditorCommands::Get().Tool_Shape_Cylinder,
-		CylinderToolIdentifier,
+		TEXT("Parametric Cylinder Tool"),
 		3000,
 		FAvalancheInteractiveToolsCreateBuilder::CreateLambda(
 			[](UEdMode* InEdMode)
@@ -45,24 +42,15 @@ void UAvaShapesEditorShapeToolCone::OnRegisterTool(IAvalancheInteractiveToolsMod
 				return UAvaInteractiveToolsToolBuilder::CreateToolBuilder<UAvaShapesEditorShapeToolCone>(InEdMode);
 			}),
 		nullptr,
-		CreateFactory<UAvaShapeConeDynamicMesh>()
+		CreateFactory<UAvaShapeConeDynamicMesh>(
+		DefaultDim,
+			[](UAvaShapeDynamicMeshBase* InMesh)
+			{
+				UAvaShapeConeDynamicMesh* Cone = Cast<UAvaShapeConeDynamicMesh>(InMesh);
+				Cone->SetTopRadius(1.f);
+			}
+		)
 	};
 
 	InAITModule->RegisterTool(IAvalancheInteractiveToolsModule::CategoryName3D, MoveTemp(CylinderToolParameters));
 }
-
-void UAvaShapesEditorShapeToolCone::InitShape(UAvaShapeDynamicMeshBase* InShape) const
-{
-	UAvaShapeConeDynamicMesh* Cone = Cast<UAvaShapeConeDynamicMesh>(InShape);
-
-	check(Cone);
-
-	Cone->SetSize3D({DefaultDim, DefaultDim, DefaultDim});
-
-	if (GetActiveToolIdentifier().Equals(CylinderToolIdentifier))
-	{
-		Cone->SetTopRadius(1.f);
-	}
-
-	Super::InitShape(InShape);
-}
diff --git a/Engine/Plugins/Experimental/Avalanche/Source/AvalancheShapesEditor/Private/Tools/AvaShapesEditorShapeToolCone.h b/Engine/Plugins/Experimental/Avalanche/Source/AvalancheShapesEditor/Private/Tools/AvaShapesEditorShapeToolCone.h
--- a/Engine/Plugins/Experimental/Avalanche/Source/AvalancheShapesEditor/Private/Tools/AvaShapesEditorShapeToolCone.h
+++ b/Engine/Plugins/Experimental/Avalanche/Source/AvalancheShapesEditor/Private/Tools/AvaShapesEditorShapeToolCone.h
@@ -11,17 +11,10 @@ class UAvaShapesEditorShapeToolCone : public UAvaShapesEditorShapeAreaToolBase
 	GENERATED_BODY()
 
 public:
-	static const TCHAR* ConeToolIdentifier;
-	static const TCHAR* CylinderToolIdentifier;
-
 	UAvaShapesEditorShapeToolCone();
 
+protected:
 	//~ Begin UAvaInteractiveToolsToolBase
 	virtual void OnRegisterTool(IAvalancheInteractiveToolsModule* InAITModule) override;
 	//~ End UAvaInteractiveToolsToolBase
-
-protected:
-	//~ Begin UAvaShapesEditorShapeToolBase
-	virtual void InitShape(UAvaShapeDynamicMeshBase* InShape) const override;
-	//~ End UAvaShapesEditorShapeToolBase
 };
diff --git a/Engine/Plugins/Experimental/Avalanche/Source/AvalancheShapesEditor/Private/Tools/AvaShapesEditorShapeToolCube.cpp b/Engine/Plugins/Experimental/Avalanche/Source/AvalancheShapesEditor/Private/Tools/AvaShapesEditorShapeToolCube.cpp
--- a/Engine/Plugins/Experimental/Avalanche/Source/AvalancheShapesEditor/Private/Tools/AvaShapesEditorShapeToolCube.cpp
+++ b/Engine/Plugins/Experimental/Avalanche/Source/AvalancheShapesEditor/Private/Tools/AvaShapesEditorShapeToolCube.cpp
@@ -31,13 +31,3 @@ void UAvaShapesEditorShapeToolCube::OnRegisterTool(IAvalancheInteractiveToolsMod
 
 	InAITModule->RegisterTool(IAvalancheInteractiveToolsModule::CategoryName3D, MoveTemp(ToolParameters));
 }
-
-void UAvaShapesEditorShapeToolCube::InitShape(UAvaShapeDynamicMeshBase* InShape) const
-{
-	UAvaShapeCubeDynamicMesh* Cube = Cast<UAvaShapeCubeDynamicMesh>(InShape);
-	check(Cube);
-
-	Cube->SetSize3D({DefaultDim, DefaultDim, DefaultDim});
-
-	Super::InitShape(InShape);
-}
diff --git a/Engine/Plugins/Experimental/Avalanche/Source/AvalancheShapesEditor/Private/Tools/AvaShapesEditorShapeToolCube.h b/Engine/Plugins/Experimental/Avalanche/Source/AvalancheShapesEditor/Private/Tools/AvaShapesEditorShapeToolCube.h
--- a/Engine/Plugins/Experimental/Avalanche/Source/AvalancheShapesEditor/Private/Tools/AvaShapesEditorShapeToolCube.h
+++ b/Engine/Plugins/Experimental/Avalanche/Source/AvalancheShapesEditor/Private/Tools/AvaShapesEditorShapeToolCube.h
@@ -13,12 +13,8 @@ class UAvaShapesEditorShapeToolCube : public UAvaShapesEditorShapeAreaToolBase
 public:
 	UAvaShapesEditorShapeToolCube();
 
+protected:
 	//~ Begin UAvaInteractiveToolsToolBase
 	virtual void OnRegisterTool(IAvalancheInteractiveToolsModule* InAITModule) override;
 	//~ End UAvaInteractiveToolsToolBase
-
-protected:
-	//~ Begin UAvaShapesEditorShapeToolBase
-	virtual void InitShape(UAvaShapeDynamicMeshBase* InShape) const override;
-	//~ End UAvaShapesEditorShapeToolBase
 };
diff --git a/Engine/Plugins/Experimental/Avalanche/Source/AvalancheShapesEditor/Private/Tools/AvaShapesEditorShapeToolEllipse.cpp b/Engine/Plugins/Experimental/Avalanche/Source/AvalancheShapesEditor/Private/Tools/AvaShapesEditorShapeToolEllipse.cpp
--- a/Engine/Plugins/Experimental/Avalanche/Source/AvalancheShapesEditor/Private/Tools/AvaShapesEditorShapeToolEllipse.cpp
+++ b/Engine/Plugins/Experimental/Avalanche/Source/AvalancheShapesEditor/Private/Tools/AvaShapesEditorShapeToolEllipse.cpp
@@ -31,13 +31,3 @@ void UAvaShapesEditorShapeToolEllipse::OnRegisterTool(IAvalancheInteractiveTools
 	
 	InAITModule->RegisterTool(IAvalancheInteractiveToolsModule::CategoryName2D, MoveTemp(ToolParameters));
 }
-
-void UAvaShapesEditorShapeToolEllipse::InitShape(UAvaShapeDynamicMeshBase* InShape) const
-{
-	UAvaShapeEllipseDynamicMesh* Ellipse = Cast<UAvaShapeEllipseDynamicMesh>(InShape);
-	check(Ellipse);
-
-	Ellipse->SetSize2D({DefaultDim, DefaultDim});
-
-	Super::InitShape(InShape);
-}
diff --git a/Engine/Plugins/Experimental/Avalanche/Source/AvalancheShapesEditor/Private/Tools/AvaShapesEditorShapeToolEllipse.h b/Engine/Plugins/Experimental/Avalanche/Source/AvalancheShapesEditor/Private/Tools/AvaShapesEditorShapeToolEllipse.h
--- a/Engine/Plugins/Experimental/Avalanche/Source/AvalancheShapesEditor/Private/Tools/AvaShapesEditorShapeToolEllipse.h
+++ b/Engine/Plugins/Experimental/Avalanche/Source/AvalancheShapesEditor/Private/Tools/AvaShapesEditorShapeToolEllipse.h
@@ -13,12 +13,8 @@ class UAvaShapesEditorShapeToolEllipse : public UAvaShapesEditorShapeAreaToolBas
 public:
 	UAvaShapesEditorShapeToolEllipse();
 
+protected:
 	//~ Begin UAvaInteractiveToolsToolBase
 	virtual void OnRegisterTool(IAvalancheInteractiveToolsModule* InAITModule) override;
 	//~ End UAvaInteractiveToolsToolBase
-
-protected:
-	//~ Begin UAvaShapesEditorShapeToolBase
-	virtual void InitShape(UAvaShapeDynamicMeshBase* InShape) const override;
-	//~ End UAvaShapesEditorShapeToolBase
 };
diff --git a/Engine/Plugins/Experimental/Avalanche/Source/AvalancheShapesEditor/Private/Tools/AvaShapesEditorShapeToolIrregularPoly.cpp b/Engine/Plugins/Experimental/Avalanche/Source/AvalancheShapesEditor/Private/Tools/AvaShapesEditorShapeToolIrregularPoly.cpp
--- a/Engine/Plugins/Experimental/Avalanche/Source/AvalancheShapesEditor/Private/Tools/AvaShapesEditorShapeToolIrregularPoly.cpp
+++ b/Engine/Plugins/Experimental/Avalanche/Source/AvalancheShapesEditor/Private/Tools/AvaShapesEditorShapeToolIrregularPoly.cpp
@@ -34,7 +34,9 @@ void UAvaShapesEditorShapeToolIrregularPoly::OnRegisterTool(IAvalancheInteractiv
 			[](UEdMode* InEdMode)
 			{
 				return UAvaInteractiveToolsToolBuilder::CreateToolBuilder<UAvaShapesEditorShapeToolIrregularPoly>(InEdMode);
-			})
+			}),
+		nullptr,
+		CreateFactory<UAvaShapeIrregularPolygonDynamicMesh>()
 	};
 
 	InAITModule->RegisterTool(IAvalancheInteractiveToolsModule::CategoryName2D, MoveTemp(ToolParameters));
@@ -129,14 +131,6 @@ void UAvaShapesEditorShapeToolIrregularPoly::OnViewportPlannerUpdate()
 	}
 }
 
-void UAvaShapesEditorShapeToolIrregularPoly::InitShape(UAvaShapeDynamicMeshBase* InShape) const
-{
-	UAvaShapeIrregularPolygonDynamicMesh* IrregularPoly = Cast<UAvaShapeIrregularPolygonDynamicMesh>(InShape);
-	check(IrregularPoly);
-
-	Super::InitShape(InShape);
-}
-
 int32 UAvaShapesEditorShapeToolIrregularPoly::FindOverlappingPointIndex(const FVector2f& InPosition) const
 {
 	int32 Index = INDEX_NONE;
diff --git a/Engine/Plugins/Experimental/Avalanche/Source/AvalancheShapesEditor/Private/Tools/AvaShapesEditorShapeToolIrregularPoly.h b/Engine/Plugins/Experimental/Avalanche/Source/AvalancheShapesEditor/Private/Tools/AvaShapesEditorShapeToolIrregularPoly.h
--- a/Engine/Plugins/Experimental/Avalanche/Source/AvalancheShapesEditor/Private/Tools/AvaShapesEditorShapeToolIrregularPoly.h
+++ b/Engine/Plugins/Experimental/Avalanche/Source/AvalancheShapesEditor/Private/Tools/AvaShapesEditorShapeToolIrregularPoly.h
@@ -13,6 +13,10 @@ class UAvaShapesEditorShapeToolIrregularPoly : public UAvaShapesEditorShapeToolB
 public:
 	UAvaShapesEditorShapeToolIrregularPoly();
 
+protected:
+	// The minimum dimension
+	static constexpr float MinDim = 5;
+	
 	//~ Begin UAvaInteractiveToolsToolBase
 	virtual void OnRegisterTool(IAvalancheInteractiveToolsModule* InAITModule) override;
 	virtual void DrawHUD(FCanvas* InCanvas, IToolsContextRenderAPI* InRenderAPI) override;
@@ -20,12 +24,7 @@ public:
 	virtual void OnViewportPlannerComplete() override {} // Does nothing
 	//~ End UAvaInteractiveToolsToolBase
 
-protected:
-	// The minimum dimension
-	static constexpr float MinDim = 5;
-
 	//~ Begin UAvaShapesEditorShapeToolBase
-	virtual void InitShape(UAvaShapeDynamicMeshBase* InShape) const override;
 	virtual void SetShapeSize(AAvaShapeActor* InShapeActor, const FVector2D& InShapeSize) const override {}
 	//~ End UAvaShapesEditorShapeToolBase
 
diff --git a/Engine/Plugins/Experimental/Avalanche/Source/AvalancheShapesEditor/Private/Tools/AvaShapesEditorShapeToolLine.cpp b/Engine/Plugins/Experimental/Avalanche/Source/AvalancheShapesEditor/Private/Tools/AvaShapesEditorShapeToolLine.cpp
--- a/Engine/Plugins/Experimental/Avalanche/Source/AvalancheShapesEditor/Private/Tools/AvaShapesEditorShapeToolLine.cpp
+++ b/Engine/Plugins/Experimental/Avalanche/Source/AvalancheShapesEditor/Private/Tools/AvaShapesEditorShapeToolLine.cpp
@@ -74,7 +74,14 @@ void UAvaShapesEditorShapeToolLine::OnRegisterTool(IAvalancheInteractiveToolsMod
 				return UAvaInteractiveToolsToolBuilder::CreateToolBuilder<UAvaShapesEditorShapeToolLine>(InEdMode);
 			}),
 		nullptr,
-		CreateFactory<UAvaShapeLineDynamicMesh>()
+		CreateFactory<UAvaShapeLineDynamicMesh>(
+			DefaultDim,
+			[](UAvaShapeDynamicMeshBase* InMesh)
+			{
+				UAvaShapeLineDynamicMesh* Line = Cast<UAvaShapeLineDynamicMesh>(InMesh);
+				Line->SetLineWidth(3.f);
+			}
+		)
 	};
 
 	InAITModule->RegisterTool(IAvalancheInteractiveToolsModule::CategoryName2D, MoveTemp(ToolParameters));
@@ -134,16 +141,6 @@ void UAvaShapesEditorShapeToolLine::OnViewportPlannerUpdate()
 	}
 }
 
-void UAvaShapesEditorShapeToolLine::InitShape(UAvaShapeDynamicMeshBase* InShape) const
-{
-	UAvaShapeLineDynamicMesh* Line = Cast<UAvaShapeLineDynamicMesh>(InShape);
-	check(Line);
-
-	Line->SetLineWidth(3.f);
-
-	Super::InitShape(InShape);
-}
-
 void UAvaShapesEditorShapeToolLine::SetShapeSize(AAvaShapeActor* InShapeActor, const FVector2D& InShapeSize) const
 {
 	Super::SetShapeSize(InShapeActor, InShapeSize);
diff --git a/Engine/Plugins/Experimental/Avalanche/Source/AvalancheShapesEditor/Private/Tools/AvaShapesEditorShapeToolLine.h b/Engine/Plugins/Experimental/Avalanche/Source/AvalancheShapesEditor/Private/Tools/AvaShapesEditorShapeToolLine.h
--- a/Engine/Plugins/Experimental/Avalanche/Source/AvalancheShapesEditor/Private/Tools/AvaShapesEditorShapeToolLine.h
+++ b/Engine/Plugins/Experimental/Avalanche/Source/AvalancheShapesEditor/Private/Tools/AvaShapesEditorShapeToolLine.h
@@ -13,22 +13,21 @@ class UAvaShapesEditorShapeToolLine : public UAvaShapesEditorShapeToolBase
 public:
 	UAvaShapesEditorShapeToolLine();
 
+protected:
+	// The minimum dimension
+	static constexpr float MinDim = 5;
+
 	//~ Begin UAvaInteractiveToolsToolBase
 	virtual void OnRegisterTool(IAvalancheInteractiveToolsModule* InAITModule) override;
 	virtual void OnViewportPlannerUpdate() override;
 	virtual void OnViewportPlannerComplete() override {} // Does nothing
 	//~ End UAvaInteractiveToolsToolBase
-
-protected:
-	// The minimum dimension
-	static constexpr float MinDim = 5;
-
-	FVector2f LineEndLocation = FVector2f::ZeroVector;
-
+	
 	//~ Begin UAvaShapesEditorShapeToolBase
-	virtual void InitShape(UAvaShapeDynamicMeshBase* InShape) const override;
 	virtual void SetShapeSize(AAvaShapeActor* InShapeActor, const FVector2D& InShapeSize) const override;
 	//~ End UAvaShapesEditorShapeToolBase
 
 	void SetLineEnds(AAvaShapeActor* InShapeActor, const FVector2f& Start, const FVector2f& End);
+
+	FVector2f LineEndLocation = FVector2f::ZeroVector;
 };
diff --git a/Engine/Plugins/Experimental/Avalanche/Source/AvalancheShapesEditor/Private/Tools/AvaShapesEditorShapeToolNGon.cpp b/Engine/Plugins/Experimental/Avalanche/Source/AvalancheShapesEditor/Private/Tools/AvaShapesEditorShapeToolNGon.cpp
--- a/Engine/Plugins/Experimental/Avalanche/Source/AvalancheShapesEditor/Private/Tools/AvaShapesEditorShapeToolNGon.cpp
+++ b/Engine/Plugins/Experimental/Avalanche/Source/AvalancheShapesEditor/Private/Tools/AvaShapesEditorShapeToolNGon.cpp
@@ -31,13 +31,3 @@ void UAvaShapesEditorShapeToolNGon::OnRegisterTool(IAvalancheInteractiveToolsMod
 	
 	InAITModule->RegisterTool(IAvalancheInteractiveToolsModule::CategoryName2D, MoveTemp(ToolParameters));
 }
-
-void UAvaShapesEditorShapeToolNGon::InitShape(UAvaShapeDynamicMeshBase* InShape) const
-{
-	UAvaShapeNGonDynamicMesh* NGon = Cast<UAvaShapeNGonDynamicMesh>(InShape);
-	check(NGon);
-
-	NGon->SetSize2D({DefaultDim, DefaultDim});
-
-	Super::InitShape(InShape);
-}
diff --git a/Engine/Plugins/Experimental/Avalanche/Source/AvalancheShapesEditor/Private/Tools/AvaShapesEditorShapeToolNGon.h b/Engine/Plugins/Experimental/Avalanche/Source/AvalancheShapesEditor/Private/Tools/AvaShapesEditorShapeToolNGon.h
--- a/Engine/Plugins/Experimental/Avalanche/Source/AvalancheShapesEditor/Private/Tools/AvaShapesEditorShapeToolNGon.h
+++ b/Engine/Plugins/Experimental/Avalanche/Source/AvalancheShapesEditor/Private/Tools/AvaShapesEditorShapeToolNGon.h
@@ -13,12 +13,8 @@ class UAvaShapesEditorShapeToolNGon : public UAvaShapesEditorShapeAreaToolBase
 public:
 	UAvaShapesEditorShapeToolNGon();
 
+protected:
 	//~ Begin UAvaInteractiveToolsToolBase
 	virtual void OnRegisterTool(IAvalancheInteractiveToolsModule* InAITModule) override;
 	//~ End UAvaInteractiveToolsToolBase
-
-protected:
-	//~ Begin UAvaShapesEditorShapeToolBase
-	virtual void InitShape(UAvaShapeDynamicMeshBase* InShape) const override;
-	//~ End UAvaShapesEditorShapeToolBase
 };
diff --git a/Engine/Plugins/Experimental/Avalanche/Source/AvalancheShapesEditor/Private/Tools/AvaShapesEditorShapeToolRectangle.cpp b/Engine/Plugins/Experimental/Avalanche/Source/AvalancheShapesEditor/Private/Tools/AvaShapesEditorShapeToolRectangle.cpp
--- a/Engine/Plugins/Experimental/Avalanche/Source/AvalancheShapesEditor/Private/Tools/AvaShapesEditorShapeToolRectangle.cpp
+++ b/Engine/Plugins/Experimental/Avalanche/Source/AvalancheShapesEditor/Private/Tools/AvaShapesEditorShapeToolRectangle.cpp
@@ -8,9 +8,6 @@
 
 #define LOCTEXT_NAMESPACE "AvaShapesEditorShapeToolRectangle"
 
-const TCHAR* UAvaShapesEditorShapeToolRectangle::SquareToolIdentifier = TEXT("Parametric Square Tool");
-const TCHAR* UAvaShapesEditorShapeToolRectangle::RectangleToolIdentifier = TEXT("Parametric Rectangle Tool");
-
 UAvaShapesEditorShapeToolRectangle::UAvaShapesEditorShapeToolRectangle()
 {
 	ShapeClass = UAvaShapeRectangleDynamicMesh::StaticClass();
@@ -20,10 +17,10 @@ void UAvaShapesEditorShapeToolRectangle::OnRegisterTool(IAvalancheInteractiveToo
 {
 	Super::OnRegisterTool(InAITModule);
 
-	FAvaInteractiveToolsToolParameters SquareToolParameters = 
+	FAvaInteractiveToolsToolParameters RectangleToolParameters = 
 	{
-		FAvaShapesEditorCommands::Get().Tool_Shape_Square,
-		SquareToolIdentifier,
+		FAvaShapesEditorCommands::Get().Tool_Shape_Rectangle,
+		TEXT("Parametric Rectangle Tool"),
 		1000,
 		FAvalancheInteractiveToolsCreateBuilder::CreateLambda(
 			[](UEdMode* InEdMode)
@@ -31,15 +28,21 @@ void UAvaShapesEditorShapeToolRectangle::OnRegisterTool(IAvalancheInteractiveToo
 				return UAvaInteractiveToolsToolBuilder::CreateToolBuilder<UAvaShapesEditorShapeToolRectangle>(InEdMode);
 			}),
 		nullptr,
-		CreateFactory<UAvaShapeRectangleDynamicMesh>()
+		CreateFactory<UAvaShapeRectangleDynamicMesh>(
+			FVector(0, 160, 90),
+			[](UAvaShapeDynamicMeshBase* InMesh)
+			{
+				InMesh->SetMaterialUVMode(UAvaShapeDynamicMeshBase::MESH_INDEX_PRIMARY, EAvaShapeUVMode::Stretch);
+			}
+		)
 	};
 
-	InAITModule->RegisterTool(IAvalancheInteractiveToolsModule::CategoryName2D, MoveTemp(SquareToolParameters));
+	InAITModule->RegisterTool(IAvalancheInteractiveToolsModule::CategoryName2D, MoveTemp(RectangleToolParameters));
 
-	FAvaInteractiveToolsToolParameters RectangleToolParameters = 
+	FAvaInteractiveToolsToolParameters SquareToolParameters = 
 	{
-		FAvaShapesEditorCommands::Get().Tool_Shape_Rectangle,
-		RectangleToolIdentifier,
+		FAvaShapesEditorCommands::Get().Tool_Shape_Square,
+		TEXT("Parametric Square Tool"),
 		1001,
 		FAvalancheInteractiveToolsCreateBuilder::CreateLambda(
 			[](UEdMode* InEdMode)
@@ -50,24 +53,7 @@ void UAvaShapesEditorShapeToolRectangle::OnRegisterTool(IAvalancheInteractiveToo
 		CreateFactory<UAvaShapeRectangleDynamicMesh>()
 	};
 
-	InAITModule->RegisterTool(IAvalancheInteractiveToolsModule::CategoryName2D, MoveTemp(RectangleToolParameters));
-}
-
-void UAvaShapesEditorShapeToolRectangle::InitShape(UAvaShapeDynamicMeshBase* InShape) const
-{
-	UAvaShapeRectangleDynamicMesh* Rectangle = Cast<UAvaShapeRectangleDynamicMesh>(InShape);
-	check(Rectangle);
-
-	if (GetActiveToolIdentifier().Equals(SquareToolIdentifier))
-	{
-		Rectangle->SetSize2D({DefaultDim, DefaultDim});
-	}
-	else
-	{
-		Rectangle->SetSize2D({160, 90});
-	}
-
-	Super::InitShape(InShape);
+	InAITModule->RegisterTool(IAvalancheInteractiveToolsModule::CategoryName2D, MoveTemp(SquareToolParameters));
 }
 
 #undef LOCTEXT_NAMESPACE
diff --git a/Engine/Plugins/Experimental/Avalanche/Source/AvalancheShapesEditor/Private/Tools/AvaShapesEditorShapeToolRectangle.h b/Engine/Plugins/Experimental/Avalanche/Source/AvalancheShapesEditor/Private/Tools/AvaShapesEditorShapeToolRectangle.h
--- a/Engine/Plugins/Experimental/Avalanche/Source/AvalancheShapesEditor/Private/Tools/AvaShapesEditorShapeToolRectangle.h
+++ b/Engine/Plugins/Experimental/Avalanche/Source/AvalancheShapesEditor/Private/Tools/AvaShapesEditorShapeToolRectangle.h
@@ -11,17 +11,10 @@ class UAvaShapesEditorShapeToolRectangle : public UAvaShapesEditorShapeAreaToolB
 	GENERATED_BODY()
 
 public:
-	static const TCHAR* SquareToolIdentifier;
-	static const TCHAR* RectangleToolIdentifier;
-	
 	UAvaShapesEditorShapeToolRectangle();
 
+protected:
 	//~ Begin UAvaInteractiveToolsToolBase
 	virtual void OnRegisterTool(IAvalancheInteractiveToolsModule* InAITModule) override;
 	//~ End UAvaInteractiveToolsToolBase
-
-protected:
-	//~ Begin UAvaShapesEditorShapeToolBase
-	virtual void InitShape(UAvaShapeDynamicMeshBase* InShape) const override;
-	//~ End UAvaShapesEditorShapeToolBase
 };
diff --git a/Engine/Plugins/Experimental/Avalanche/Source/AvalancheShapesEditor/Private/Tools/AvaShapesEditorShapeToolRing.cpp b/Engine/Plugins/Experimental/Avalanche/Source/AvalancheShapesEditor/Private/Tools/AvaShapesEditorShapeToolRing.cpp
--- a/Engine/Plugins/Experimental/Avalanche/Source/AvalancheShapesEditor/Private/Tools/AvaShapesEditorShapeToolRing.cpp
+++ b/Engine/Plugins/Experimental/Avalanche/Source/AvalancheShapesEditor/Private/Tools/AvaShapesEditorShapeToolRing.cpp
@@ -31,13 +31,3 @@ void UAvaShapesEditorShapeToolRing::OnRegisterTool(IAvalancheInteractiveToolsMod
 
 	InAITModule->RegisterTool(IAvalancheInteractiveToolsModule::CategoryName2D, MoveTemp(ToolParameters));
 }
-
-void UAvaShapesEditorShapeToolRing::InitShape(UAvaShapeDynamicMeshBase* InShape) const
-{
-	UAvaShapeRingDynamicMesh* Ring = Cast<UAvaShapeRingDynamicMesh>(InShape);
-	check(Ring);
-
-	Ring->SetSize2D({DefaultDim, DefaultDim});
-
-	Super::InitShape(InShape);
-}
diff --git a/Engine/Plugins/Experimental/Avalanche/Source/AvalancheShapesEditor/Private/Tools/AvaShapesEditorShapeToolRing.h b/Engine/Plugins/Experimental/Avalanche/Source/AvalancheShapesEditor/Private/Tools/AvaShapesEditorShapeToolRing.h
--- a/Engine/Plugins/Experimental/Avalanche/Source/AvalancheShapesEditor/Private/Tools/AvaShapesEditorShapeToolRing.h
+++ b/Engine/Plugins/Experimental/Avalanche/Source/AvalancheShapesEditor/Private/Tools/AvaShapesEditorShapeToolRing.h
@@ -13,12 +13,8 @@ class UAvaShapesEditorShapeToolRing : public UAvaShapesEditorShapeAreaToolBase
 public:
 	UAvaShapesEditorShapeToolRing();
 
+protected:
 	//~ Begin UAvaInteractiveToolsToolBase
 	virtual void OnRegisterTool(IAvalancheInteractiveToolsModule* InAITModule) override;
 	//~ End UAvaInteractiveToolsToolBase
-
-protected:
-	//~ Begin UAvaShapesEditorShapeToolBase
-	virtual void InitShape(UAvaShapeDynamicMeshBase* InShape) const override;
-	//~ End UAvaShapesEditorShapeToolBase
 };
diff --git a/Engine/Plugins/Experimental/Avalanche/Source/AvalancheShapesEditor/Private/Tools/AvaShapesEditorShapeToolSphere.cpp b/Engine/Plugins/Experimental/Avalanche/Source/AvalancheShapesEditor/Private/Tools/AvaShapesEditorShapeToolSphere.cpp
--- a/Engine/Plugins/Experimental/Avalanche/Source/AvalancheShapesEditor/Private/Tools/AvaShapesEditorShapeToolSphere.cpp
+++ b/Engine/Plugins/Experimental/Avalanche/Source/AvalancheShapesEditor/Private/Tools/AvaShapesEditorShapeToolSphere.cpp
@@ -31,13 +31,3 @@ void UAvaShapesEditorShapeToolSphere::OnRegisterTool(IAvalancheInteractiveToolsM
 	
 	InAITModule->RegisterTool(IAvalancheInteractiveToolsModule::CategoryName3D, MoveTemp(ToolParameters));
 }
-
-void UAvaShapesEditorShapeToolSphere::InitShape(UAvaShapeDynamicMeshBase* InShape) const
-{
-	UAvaShapeSphereDynamicMesh* Sphere = Cast<UAvaShapeSphereDynamicMesh>(InShape);
-	check(Sphere);
-
-	Sphere->SetSize3D({DefaultDim, DefaultDim, DefaultDim});
-
-	Super::InitShape(InShape);
-}
diff --git a/Engine/Plugins/Experimental/Avalanche/Source/AvalancheShapesEditor/Private/Tools/AvaShapesEditorShapeToolSphere.h b/Engine/Plugins/Experimental/Avalanche/Source/AvalancheShapesEditor/Private/Tools/AvaShapesEditorShapeToolSphere.h
--- a/Engine/Plugins/Experimental/Avalanche/Source/AvalancheShapesEditor/Private/Tools/AvaShapesEditorShapeToolSphere.h
+++ b/Engine/Plugins/Experimental/Avalanche/Source/AvalancheShapesEditor/Private/Tools/AvaShapesEditorShapeToolSphere.h
@@ -13,12 +13,8 @@ class UAvaShapesEditorShapeToolSphere : public UAvaShapesEditorShapeAreaToolBase
 public:
 	UAvaShapesEditorShapeToolSphere();
 
+protected:
 	//~ Begin UAvaInteractiveToolsToolBase
 	virtual void OnRegisterTool(IAvalancheInteractiveToolsModule* InAITModule) override;
 	//~ End UAvaInteractiveToolsToolBase
-
-protected:
-	//~ Begin UAvaShapesEditorShapeToolBase
-	virtual void InitShape(UAvaShapeDynamicMeshBase* InShape) const override;
-	//~ End UAvaShapesEditorShapeToolBase
 };
diff --git a/Engine/Plugins/Experimental/Avalanche/Source/AvalancheShapesEditor/Private/Tools/AvaShapesEditorShapeToolStar.cpp b/Engine/Plugins/Experimental/Avalanche/Source/AvalancheShapesEditor/Private/Tools/AvaShapesEditorShapeToolStar.cpp
--- a/Engine/Plugins/Experimental/Avalanche/Source/AvalancheShapesEditor/Private/Tools/AvaShapesEditorShapeToolStar.cpp
+++ b/Engine/Plugins/Experimental/Avalanche/Source/AvalancheShapesEditor/Private/Tools/AvaShapesEditorShapeToolStar.cpp
@@ -31,13 +31,3 @@ void UAvaShapesEditorShapeToolStar::OnRegisterTool(IAvalancheInteractiveToolsMod
 	
 	InAITModule->RegisterTool(IAvalancheInteractiveToolsModule::CategoryName2D, MoveTemp(ToolParameters));
 }
-
-void UAvaShapesEditorShapeToolStar::InitShape(UAvaShapeDynamicMeshBase* InShape) const
-{
-	UAvaShapeStarDynamicMesh* Star = Cast<UAvaShapeStarDynamicMesh>(InShape);
-	check(Star);
-
-	Star->SetSize2D({DefaultDim, DefaultDim});
-
-	Super::InitShape(InShape);
-}
diff --git a/Engine/Plugins/Experimental/Avalanche/Source/AvalancheShapesEditor/Private/Tools/AvaShapesEditorShapeToolStar.h b/Engine/Plugins/Experimental/Avalanche/Source/AvalancheShapesEditor/Private/Tools/AvaShapesEditorShapeToolStar.h
--- a/Engine/Plugins/Experimental/Avalanche/Source/AvalancheShapesEditor/Private/Tools/AvaShapesEditorShapeToolStar.h
+++ b/Engine/Plugins/Experimental/Avalanche/Source/AvalancheShapesEditor/Private/Tools/AvaShapesEditorShapeToolStar.h
@@ -13,12 +13,8 @@ class UAvaShapesEditorShapeToolStar : public UAvaShapesEditorShapeAreaToolBase
 public:
 	UAvaShapesEditorShapeToolStar();
 
+protected:
 	//~ Begin UAvaInteractiveToolsToolBase
 	virtual void OnRegisterTool(IAvalancheInteractiveToolsModule* InAITModule) override;
 	//~ End UAvaInteractiveToolsToolBase
-
-protected:
-	//~ Begin UAvaShapesEditorShapeToolBase
-	virtual void InitShape(UAvaShapeDynamicMeshBase* InShape) const override;
-	//~ End UAvaShapesEditorShapeToolBase
 };
diff --git a/Engine/Plugins/Experimental/Avalanche/Source/AvalancheShapesEditor/Private/Tools/AvaShapesEditorShapeToolTorus.cpp b/Engine/Plugins/Experimental/Avalanche/Source/AvalancheShapesEditor/Private/Tools/AvaShapesEditorShapeToolTorus.cpp
--- a/Engine/Plugins/Experimental/Avalanche/Source/AvalancheShapesEditor/Private/Tools/AvaShapesEditorShapeToolTorus.cpp
+++ b/Engine/Plugins/Experimental/Avalanche/Source/AvalancheShapesEditor/Private/Tools/AvaShapesEditorShapeToolTorus.cpp
@@ -31,13 +31,3 @@ void UAvaShapesEditorShapeToolTorus::OnRegisterTool(IAvalancheInteractiveToolsMo
 	
 	InAITModule->RegisterTool(IAvalancheInteractiveToolsModule::CategoryName3D, MoveTemp(ToolParameters));
 }
-
-void UAvaShapesEditorShapeToolTorus::InitShape(UAvaShapeDynamicMeshBase* InShape) const
-{
-	UAvaShapeTorusDynamicMesh* Torus = Cast<UAvaShapeTorusDynamicMesh>(InShape);
-	check(Torus);
-
-	Torus->SetSize3D({DefaultDim, DefaultDim, DefaultDim});
-
-	Super::InitShape(InShape);
-}
diff --git a/Engine/Plugins/Experimental/Avalanche/Source/AvalancheShapesEditor/Private/Tools/AvaShapesEditorShapeToolTorus.h b/Engine/Plugins/Experimental/Avalanche/Source/AvalancheShapesEditor/Private/Tools/AvaShapesEditorShapeToolTorus.h
--- a/Engine/Plugins/Experimental/Avalanche/Source/AvalancheShapesEditor/Private/Tools/AvaShapesEditorShapeToolTorus.h
+++ b/Engine/Plugins/Experimental/Avalanche/Source/AvalancheShapesEditor/Private/Tools/AvaShapesEditorShapeToolTorus.h
@@ -13,12 +13,8 @@ class UAvaShapesEditorShapeToolTorus : public UAvaShapesEditorShapeAreaToolBase
 public:
 	UAvaShapesEditorShapeToolTorus();
 
+protected:
 	//~ Begin UAvaInteractiveToolsToolBase
 	virtual void OnRegisterTool(IAvalancheInteractiveToolsModule* InAITModule) override;
 	//~ End UAvaInteractiveToolsToolBase
-
-protected:
-	//~ Begin UAvaShapesEditorShapeToolBase
-	virtual void InitShape(UAvaShapeDynamicMeshBase* InShape) const override;
-	//~ End UAvaShapesEditorShapeToolBase
 };
diff --git a/Engine/Plugins/Experimental/Avalanche/Source/AvalancheTextEditor/Private/Tool/AvaTextActorTool.cpp b/Engine/Plugins/Experimental/Avalanche/Source/AvalancheTextEditor/Private/Tool/AvaTextActorTool.cpp
--- a/Engine/Plugins/Experimental/Avalanche/Source/AvalancheTextEditor/Private/Tool/AvaTextActorTool.cpp
+++ b/Engine/Plugins/Experimental/Avalanche/Source/AvalancheTextEditor/Private/Tool/AvaTextActorTool.cpp
@@ -21,7 +21,7 @@ void UAvaTextActorTool::OnRegisterTool(IAvalancheInteractiveToolsModule* InAITMo
 	{
 		FAvaTextEditorCommands::Get().Tool_Actor_Text3D,
 		TEXT("Text Actor Tool"),
-		2000,
+		1000,
 		FAvalancheInteractiveToolsCreateBuilder::CreateLambda(
 			[](UEdMode* InEdMode)
 			{
diff --git a/Engine/Plugins/Experimental/ClonerEffector/Source/ClonerEffectorEditor/Private/Cloner/CEClonerActorFactory.cpp b/Engine/Plugins/Experimental/ClonerEffector/Source/ClonerEffectorEditor/Private/Cloner/CEClonerActorFactory.cpp
--- a/Engine/Plugins/Experimental/ClonerEffector/Source/ClonerEffectorEditor/Private/Cloner/CEClonerActorFactory.cpp
+++ b/Engine/Plugins/Experimental/ClonerEffector/Source/ClonerEffectorEditor/Private/Cloner/CEClonerActorFactory.cpp
@@ -2,6 +2,7 @@
 
 #include "Cloner/CEClonerActorFactory.h"
 #include "Cloner/CEClonerActor.h"
+#include "Cloner/CEClonerComponent.h"
 #include "EngineAnalytics.h"
 #include "Subsystems/PlacementSubsystem.h"
 
@@ -10,6 +11,27 @@ UCEClonerActorFactory::UCEClonerActorFactory()
 	NewActorClass = ACEClonerActor::StaticClass();
 }
 
+void UCEClonerActorFactory::SetClonerLayout(FName InLayoutName)
+{
+	ClonerLayoutName = InLayoutName;
+}
+
+void UCEClonerActorFactory::PostSpawnActor(UObject* InAsset, AActor* InNewActor)
+{
+	Super::PostSpawnActor(InAsset, InNewActor);
+
+	if (!ClonerLayoutName.IsNone())
+	{
+		if (const ACEClonerActor* ClonerActor = Cast<ACEClonerActor>(InNewActor))
+		{
+			if (UCEClonerComponent* ClonerComponent = ClonerActor->GetClonerComponent())
+			{
+				ClonerComponent->SetLayoutName(ClonerLayoutName);
+			}
+		}
+	}
+}
+
 void UCEClonerActorFactory::PostPlaceAsset(TArrayView<const FTypedElementHandle> InHandle, const FAssetPlacementInfo& InPlacementInfo, const FPlacementOptions& InPlacementOptions)
 {
 	Super::PostPlaceAsset(InHandle, InPlacementInfo, InPlacementOptions);
diff --git a/Engine/Plugins/Experimental/ClonerEffector/Source/ClonerEffectorEditor/Private/Effector/CEEffectorActorFactory.cpp b/Engine/Plugins/Experimental/ClonerEffector/Source/ClonerEffectorEditor/Private/Effector/CEEffectorActorFactory.cpp
--- a/Engine/Plugins/Experimental/ClonerEffector/Source/ClonerEffectorEditor/Private/Effector/CEEffectorActorFactory.cpp
+++ b/Engine/Plugins/Experimental/ClonerEffector/Source/ClonerEffectorEditor/Private/Effector/CEEffectorActorFactory.cpp
@@ -2,6 +2,7 @@
 
 #include "Effector/CEEffectorActorFactory.h"
 #include "Effector/CEEffectorActor.h"
+#include "Effector/CEEffectorComponent.h"
 #include "EngineAnalytics.h"
 #include "Subsystems/PlacementSubsystem.h"
 
@@ -10,6 +11,27 @@ UCEEffectorActorFactory::UCEEffectorActorFactory()
 	NewActorClass = ACEEffectorActor::StaticClass();
 }
 
+void UCEEffectorActorFactory::SetEffectorTypeName(FName InEffectorTypeName)
+{
+	EffectorTypeName = InEffectorTypeName;
+}
+
+void UCEEffectorActorFactory::PostSpawnActor(UObject* InAsset, AActor* InNewActor)
+{
+	Super::PostSpawnActor(InAsset, InNewActor);
+
+	if (!EffectorTypeName.IsNone())
+	{
+		if (const ACEEffectorActor* EffectorActor = Cast<ACEEffectorActor>(InNewActor))
+		{
+			if (UCEEffectorComponent* EffectorComponent = EffectorActor->GetEffectorComponent())
+			{
+				EffectorComponent->SetTypeName(EffectorTypeName);
+			}
+		}
+	}
+}
+
 void UCEEffectorActorFactory::PostPlaceAsset(TArrayView<const FTypedElementHandle> InHandle, const FAssetPlacementInfo& InPlacementInfo, const FPlacementOptions& InPlacementOptions)
 {
 	Super::PostPlaceAsset(InHandle, InPlacementInfo, InPlacementOptions);
diff --git a/Engine/Plugins/Experimental/ClonerEffector/Source/ClonerEffectorEditor/Public/Cloner/CEClonerActorFactory.h b/Engine/Plugins/Experimental/ClonerEffector/Source/ClonerEffectorEditor/Public/Cloner/CEClonerActorFactory.h
--- a/Engine/Plugins/Experimental/ClonerEffector/Source/ClonerEffectorEditor/Public/Cloner/CEClonerActorFactory.h
+++ b/Engine/Plugins/Experimental/ClonerEffector/Source/ClonerEffectorEditor/Public/Cloner/CEClonerActorFactory.h
@@ -13,7 +13,13 @@ class UCEClonerActorFactory : public UActorFactory
 public:
 	UCEClonerActorFactory();
 
+	CLONEREFFECTOREDITOR_API void SetClonerLayout(FName InLayoutName);
+
+protected:
 	//~ Begin UActorFactory
+	virtual void PostSpawnActor(UObject* InAsset, AActor* InNewActor) override;
 	virtual void PostPlaceAsset(TArrayView<const FTypedElementHandle> InHandle, const FAssetPlacementInfo& InPlacementInfo, const FPlacementOptions& InPlacementOptions) override;
 	//~ End UActorFactory
+
+	FName ClonerLayoutName = NAME_None;
 };
diff --git a/Engine/Plugins/Experimental/ClonerEffector/Source/ClonerEffectorEditor/Public/Effector/CEEffectorActorFactory.h b/Engine/Plugins/Experimental/ClonerEffector/Source/ClonerEffectorEditor/Public/Effector/CEEffectorActorFactory.h
--- a/Engine/Plugins/Experimental/ClonerEffector/Source/ClonerEffectorEditor/Public/Effector/CEEffectorActorFactory.h
+++ b/Engine/Plugins/Experimental/ClonerEffector/Source/ClonerEffectorEditor/Public/Effector/CEEffectorActorFactory.h
@@ -13,7 +13,13 @@ class UCEEffectorActorFactory : public UActorFactory
 public:
 	UCEEffectorActorFactory();
 
+	CLONEREFFECTOREDITOR_API void SetEffectorTypeName(FName InEffectorTypeName);
+
+protected:
 	//~ Begin UActorFactory
+	virtual void PostSpawnActor(UObject* InAsset, AActor* InNewActor) override;
 	virtual void PostPlaceAsset(TArrayView<const FTypedElementHandle> InHandle, const FAssetPlacementInfo& InPlacementInfo, const FPlacementOptions& InPlacementOptions) override;
 	//~ End UActorFactory
+
+	FName EffectorTypeName;
 };
