Don't mention xyz explicitly in uniform scale tooltips

diff --git a/Engine/Plugins/Experimental/Avalanche/Source/AvalancheLevelViewport/Private/Customizations/AvaLevelViewportComponentTransformDetails.cpp b/Engine/Plugins/Experimental/Avalanche/Source/AvalancheLevelViewport/Private/Customizations/AvaLevelViewportComponentTransformDetails.cpp
--- a/Engine/Plugins/Experimental/Avalanche/Source/AvalancheLevelViewport/Private/Customizations/AvaLevelViewportComponentTransformDetails.cpp
+++ b/Engine/Plugins/Experimental/Avalanche/Source/AvalancheLevelViewport/Private/Customizations/AvaLevelViewportComponentTransformDetails.cpp
@@ -481,7 +481,7 @@ TSharedRef<SWidget> FAvaLevelViewportComponentTransformDetails::GetPreserveScale
 		.IsEnabled(this, &FAvaLevelViewportComponentTransformDetails::GetIsEnabled)
 		.OnCheckStateChanged(this, &FAvaLevelViewportComponentTransformDetails::OnPreserveScaleRatioToggled)
 		.Style(FAppStyle::Get(), "TransparentCheckBox")
-		.ToolTipText(LOCTEXT("PreserveScaleToolTip", "When locked, scales uniformly based on the current xyz scale values so the object maintains its shape in each direction when scaled"))
+		.ToolTipText(LOCTEXT("PreserveScaleToolTip", "When locked, all axis values scale together so the object maintains its proportions in all directions."))
 		[
 			SNew(SImage)
 			.Image(this, &FAvaLevelViewportComponentTransformDetails::GetPreserveScaleRatioImage)
diff --git a/Engine/Source/Developer/AnimationWidgets/Public/SAdvancedTransformInputBox.h b/Engine/Source/Developer/AnimationWidgets/Public/SAdvancedTransformInputBox.h
--- a/Engine/Source/Developer/AnimationWidgets/Public/SAdvancedTransformInputBox.h
+++ b/Engine/Source/Developer/AnimationWidgets/Public/SAdvancedTransformInputBox.h
@@ -1047,7 +1047,7 @@ public:
 			FOnScaleLockChanged OnScaleLockChanged = InArgs._OnScaleLockChanged;
 
 			static const FText TooltipText = NSLOCTEXT("SAdvancedTransformInputBox", "PreserveScaleToolTip",
-					 "When locked, scales uniformly based on the current xyz scale values so the object maintains its shape in each direction when scaled");
+					 "When locked, all axis values scale together so the object maintains its proportions in all directions.");
 
 			return SNew(SButton)
 				.ContentPadding(0.f)
diff --git a/Engine/Source/Editor/DetailCustomizations/Private/ComponentTransformDetails.cpp b/Engine/Source/Editor/DetailCustomizations/Private/ComponentTransformDetails.cpp
--- a/Engine/Source/Editor/DetailCustomizations/Private/ComponentTransformDetails.cpp
+++ b/Engine/Source/Editor/DetailCustomizations/Private/ComponentTransformDetails.cpp
@@ -242,7 +242,7 @@ TSharedRef<SWidget> FComponentTransformDetails::BuildTransformFieldLabel( ETrans
 				.IsEnabled(this, &FComponentTransformDetails::GetIsScaleEnabled)
 				.OnCheckStateChanged(this, &FComponentTransformDetails::OnPreserveScaleRatioToggled)
 				.Style(FAppStyle::Get(), "TransparentCheckBox")
-				.ToolTipText(LOCTEXT("PreserveScaleToolTip", "When locked, scales uniformly based on the current xyz scale values so the object maintains its shape in each direction when scaled"))
+				.ToolTipText(LOCTEXT("PreserveScaleToolTip", "When locked, all axis values scale together so the object maintains its proportions in all directions."))
 				[
 					SNew(SImage)
 					.Image(this, &FComponentTransformDetails::GetPreserveScaleRatioImage)
diff --git a/Engine/Source/Editor/DetailCustomizations/Private/MathStructCustomizations.cpp b/Engine/Source/Editor/DetailCustomizations/Private/MathStructCustomizations.cpp
--- a/Engine/Source/Editor/DetailCustomizations/Private/MathStructCustomizations.cpp
+++ b/Engine/Source/Editor/DetailCustomizations/Private/MathStructCustomizations.cpp
@@ -204,7 +204,7 @@ void FMathStructCustomization::MakeHeaderRow(TSharedRef<class IPropertyHandle>&
 			.IsChecked(this, &FMathStructCustomization::IsPreserveScaleRatioChecked)
 			.OnCheckStateChanged(this, &FMathStructCustomization::OnPreserveScaleRatioToggled, StructWeakHandlePtr)
 			.Style(FAppStyle::Get(), "TransparentCheckBox")
-			.ToolTipText(LOCTEXT("PreserveScaleToolTip", "When locked, scales uniformly based on the current xyz scale values so the object maintains its shape in each direction when scaled"))
+			.ToolTipText(LOCTEXT("PreserveScaleToolTip", "When locked, all axis values scale together so the object maintains its proportions in all directions."))
 			[
 				SNew(SImage)
 				.Image(this, &FMathStructCustomization::GetPreserveScaleRatioImage)
