MotionDesign : Shapes
- Add simple texture option within parametric material used by shapes to avoid having to create a material designer each time
- Size to texture button added next to parametric texture property that find/add size to texture modifier and links the texture to it
- Add GradientRotation property for LinearGradient style
- Updated toolbox materials to support cloner texture and color
- Default shape parametric material will be set to translucent and two sided disabled for 3D shapes

#virtualized

diff --git a/Engine/Plugins/Experimental/Avalanche/Source/AvalancheEditor/Private/DetailView/Customizations/AvaMeshesDetailCustomization.cpp b/Engine/Plugins/Experimental/Avalanche/Source/AvalancheEditor/Private/DetailView/Customizations/AvaMeshesDetailCustomization.cpp
--- a/Engine/Plugins/Experimental/Avalanche/Source/AvalancheEditor/Private/DetailView/Customizations/AvaMeshesDetailCustomization.cpp
+++ b/Engine/Plugins/Experimental/Avalanche/Source/AvalancheEditor/Private/DetailView/Customizations/AvaMeshesDetailCustomization.cpp
@@ -10,13 +10,18 @@
 #include "DetailView/Widgets/SAvaDynamicMaterialWidget.h"
 #include "DetailWidgetRow.h"
 #include "Dialogs/DlgPickAssetPath.h"
+#include "DynamicMeshes/AvaShape2DDynMeshBase.h"
 #include "DynamicMeshes/AvaShapeDynMeshBase.h"
 #include "DynamicMeshToMeshDescription.h"
 #include "Engine/StaticMesh.h"
 #include "IAssetTools.h"
+#include "Modifiers/AvaSizeToTextureModifier.h"
+#include "Modifiers/Utilities/ActorModifierCoreLibrary.h"
 #include "UObject/Object.h"
+#include "Widgets/Layout/SBox.h"
 #include "Widgets/Input/SButton.h"
 #include "Widgets/Text/STextBlock.h"
+#include "Widgets/SBoxPanel.h"
 
 #define LOCTEXT_NAMESPACE "AvaMeshesDetailCustomization"
 
@@ -138,7 +143,7 @@ void FAvaMeshesDetailCustomization::CustomizeDetails(IDetailLayoutBuilder& InDet
 				return MaterialHandle->IsEditable() ? EVisibility::Visible : EVisibility::Hidden;
 			}));
 
-			// Parametric Material Color
+			// Parametric Material
 			TSharedPtr<IPropertyHandle> ParametricMaterialHandle = MeshPropertyHandle->GetChildHandle(GET_MEMBER_NAME_CHECKED(FAvaShapeMeshData, ParametricMaterial));
 			check(ParametricMaterialHandle.IsValid());
 
@@ -150,14 +155,68 @@ void FAvaMeshesDetailCustomization::CustomizeDetails(IDetailLayoutBuilder& InDet
 				TSharedPtr<IPropertyHandle> ParametricChildHandle = ParametricMaterialHandle->GetChildHandle(ChildIdx);
 				check(ParametricChildHandle.IsValid());
 
-				IDetailPropertyRow& NewParametricRow = MaterialCategoryBuilder.AddProperty(ParametricChildHandle.ToSharedRef());
-				FString ParametricRowName = MeshName + ParametricChildHandle->GetPropertyDisplayName().ToString();
-				NewParametricRow.DisplayName(FText::FromString(ParametricRowName));
-				NewParametricRow.Visibility(MakeAttributeLambda([=]() {
-					return ParametricMaterialHandle->IsEditable() ?
-						(ParametricChildHandle->IsEditable() ? EVisibility::Visible : EVisibility::Hidden) :
-						EVisibility::Hidden;
-				}));
+				TAttribute<EVisibility> VisibilityAttribute = MakeAttributeLambda([ParametricMaterialHandle, ParametricChildHandle]()
+					{
+						return ParametricMaterialHandle->IsEditable() ? (ParametricChildHandle->IsEditable() ? EVisibility::Visible : EVisibility::Hidden) : EVisibility::Hidden;
+					});
+
+				if (ParametricChildHandle->GetProperty() && ParametricChildHandle->GetProperty()->GetFName() == TEXT("Texture"))
+				{
+					FDetailWidgetRow& ParametricTextureRow = MaterialCategoryBuilder.AddCustomRow(ParametricChildHandle->GetPropertyDisplayName());
+					ParametricTextureRow.PropertyHandleList({ParametricChildHandle});
+					ParametricTextureRow.Visibility(VisibilityAttribute);
+
+					ParametricTextureRow.NameContent()
+					[
+						ParametricChildHandle->CreatePropertyNameWidget()
+					];
+
+					const EVisibility SizeToTextureVisibility = CanSizeToTexture() ? EVisibility::Visible : EVisibility::Hidden;
+
+					ParametricTextureRow.ValueContent()
+					.HAlign(HAlign_Fill)
+					.VAlign(VAlign_Center)
+					.MaxDesiredWidth(250)
+					[
+						SNew(SHorizontalBox)
+						+ SHorizontalBox::Slot()
+						.AutoWidth()
+						[
+							ParametricChildHandle->CreatePropertyValueWidget()
+						]
+						+ SHorizontalBox::Slot()
+						.Padding(2.f, 0.f)
+						.AutoWidth()
+						[
+							SNew(SBox)
+							.HAlign(HAlign_Center)
+							.VAlign(VAlign_Center)
+							.Visibility(SizeToTextureVisibility)
+							.Padding(0.f)
+							[
+								SNew(SButton)
+								.ContentPadding(2.f)
+								.VAlign(VAlign_Center)
+								.HAlign(HAlign_Center)
+								.ToolTipText(LOCTEXT("SizeToTextureTooltip", "Find or Add a size to texture modifier and link this texture to it"))
+								.OnClicked(this, &FAvaMeshesDetailCustomization::OnSizeToTextureClicked)
+								.Content()
+								[
+									SNew(STextBlock)
+									.Font(IDetailLayoutBuilder::GetDetailFont())
+									.Text(LOCTEXT("SizeToTextureLabel", "Size to Texture"))
+								]
+							]
+						]
+					];
+				}
+				else
+				{
+					IDetailPropertyRow& NewParametricRow = MaterialCategoryBuilder.AddProperty(ParametricChildHandle.ToSharedRef());
+					FString ParametricRowName = MeshName + ParametricChildHandle->GetPropertyDisplayName().ToString();
+					NewParametricRow.DisplayName(FText::FromString(ParametricRowName));
+					NewParametricRow.Visibility(VisibilityAttribute);
+				}
 			}
 
 			// Use primary uv params
@@ -301,6 +360,84 @@ FReply FAvaMeshesDetailCustomization::OnConvertToStaticMeshClicked()
 	return FReply::Handled();
 }
 
+bool FAvaMeshesDetailCustomization::CanSizeToTexture() const
+{
+	bool bAllowed = false;
+
+	for (const TWeakObjectPtr<UAvaShapeDynamicMeshBase>& MeshGeneratorWeak : MeshGeneratorsWeak)
+	{
+		if (const UAvaShapeDynamicMeshBase* MeshGenerator = MeshGeneratorWeak.Get())
+		{
+			if (MeshGenerator->IsA<UAvaShape2DDynMeshBase>())
+			{
+				bAllowed = true;
+			}
+			else
+			{
+				bAllowed = false;
+				break;
+			}
+		}
+	}
+
+	return bAllowed;
+}
+
+FReply FAvaMeshesDetailCustomization::OnSizeToTextureClicked()
+{
+	if (!CanSizeToTexture())
+	{
+		return FReply::Handled();
+	}
+
+	for (const TWeakObjectPtr<UAvaShapeDynamicMeshBase>& MeshGeneratorWeak : MeshGeneratorsWeak)
+	{
+		UAvaShapeDynamicMeshBase* MeshGenerator = MeshGeneratorWeak.Get();
+		if (!MeshGenerator)
+		{
+			continue;
+		}
+
+		const FAvaShapeParametricMaterial* PrimaryParametricMaterial = MeshGenerator->GetParametricMaterial(UAvaShapeDynamicMeshBase::MESH_INDEX_PRIMARY);
+		if (!PrimaryParametricMaterial)
+		{
+			continue;
+		}
+
+		AActor* ShapeActor = MeshGenerator->GetShapeActor();
+		if (!ShapeActor)
+		{
+			continue;
+		}
+
+		UActorModifierCoreStack* Stack = nullptr;
+		if (!UActorModifierCoreLibrary::FindModifierStack(ShapeActor, Stack, /** CreateIfNone */true))
+		{
+			continue;
+		}
+
+		UActorModifierCoreBase* SizeToTextureModifier = UActorModifierCoreLibrary::FindModifierByClass(Stack, UAvaSizeToTextureModifier::StaticClass());
+		if (!SizeToTextureModifier)
+		{
+			FActorModifierCoreInsertOperation InsertOp;
+			InsertOp.ModifierClass = UAvaSizeToTextureModifier::StaticClass();
+			InsertOp.InsertPosition = EActorModifierCoreStackPosition::After;
+			InsertOp.InsertPositionContext = nullptr;
+
+			if (!UActorModifierCoreLibrary::InsertModifier(Stack, InsertOp, SizeToTextureModifier))
+			{
+				continue;
+			}
+		}
+
+		UAvaSizeToTextureModifier* SizeToTextureModifierCasted = Cast<UAvaSizeToTextureModifier>(SizeToTextureModifier);
+		SizeToTextureModifierCasted->SetTexture(PrimaryParametricMaterial->GetTexture());
+		MeshGenerator->SetMaterialUVMode(UAvaShapeDynamicMeshBase::MESH_INDEX_PRIMARY, EAvaShapeUVMode::Stretch);
+	}
+
+	return FReply::Handled();
+}
+
 bool FAvaMeshesDetailCustomization::CanConvertToStaticMesh() const
 {
 	return MeshGeneratorsWeak.Num() == 1 && MeshGeneratorsWeak[0].IsValid();
diff --git a/Engine/Plugins/Experimental/Avalanche/Source/AvalancheEditor/Private/DetailView/Customizations/AvaMeshesDetailCustomization.h b/Engine/Plugins/Experimental/Avalanche/Source/AvalancheEditor/Private/DetailView/Customizations/AvaMeshesDetailCustomization.h
--- a/Engine/Plugins/Experimental/Avalanche/Source/AvalancheEditor/Private/DetailView/Customizations/AvaMeshesDetailCustomization.h
+++ b/Engine/Plugins/Experimental/Avalanche/Source/AvalancheEditor/Private/DetailView/Customizations/AvaMeshesDetailCustomization.h
@@ -25,11 +25,15 @@ public:
 	//~ End IDetailCustomization
 
 protected:
+	/** Enable or disable the button when selected object is not compatible */
+	bool CanConvertToStaticMesh() const;
+
 	/** Handler when the convert button is clicked */
 	FReply OnConvertToStaticMeshClicked();
 
-	/** Enable or disable the button when selected object is not compatible */
-	bool CanConvertToStaticMesh() const;
+	bool CanSizeToTexture() const;
+
+	FReply OnSizeToTextureClicked();
 
 	TArray<TWeakObjectPtr<UAvaShapeDynamicMeshBase>> MeshGeneratorsWeak;
 };
\ No newline at end of file
diff --git a/Engine/Plugins/Experimental/Avalanche/Source/AvalancheMask/Private/Handling/AvaMaskParametricMaterialHandle.cpp b/Engine/Plugins/Experimental/Avalanche/Source/AvalancheMask/Private/Handling/AvaMaskParametricMaterialHandle.cpp
--- a/Engine/Plugins/Experimental/Avalanche/Source/AvalancheMask/Private/Handling/AvaMaskParametricMaterialHandle.cpp
+++ b/Engine/Plugins/Experimental/Avalanche/Source/AvalancheMask/Private/Handling/AvaMaskParametricMaterialHandle.cpp
@@ -20,9 +20,9 @@ FAvaMaskParametricMaterialHandle::FAvaMaskParametricMaterialHandle(const FStruct
 			LastMaterialInstance = MID;
 		}
 
-		if (!ParametricMtl->OnMaterialParameterChanged().IsBoundToObject(this))
+		if (!FAvaShapeParametricMaterial::OnMaterialParameterChanged().IsBoundToObject(this))
 		{
-			MaterialChangedHandle = ParametricMtl->OnMaterialParameterChanged().AddRaw(this, &FAvaMaskParametricMaterialHandle::OnMaterialChanged);
+			FAvaShapeParametricMaterial::OnMaterialParameterChanged().AddRaw(this, &FAvaMaskParametricMaterialHandle::OnMaterialChanged);
 		}
 	}
 }
@@ -32,10 +32,7 @@ FAvaMaskParametricMaterialHandle::~FAvaMaskParametricMaterialHandle()
 	// @note: if this doesn't cover all situations, the handle needs to be aware of the struct's outer to check validity
 	if (!GIsGarbageCollecting)
 	{
-		if (FAvaShapeParametricMaterial* ParametricMtl = GetParametricMaterial())
-		{
-			ParametricMtl->OnMaterialParameterChanged().RemoveAll(this);
-		}
+		FAvaShapeParametricMaterial::OnMaterialParameterChanged().RemoveAll(this);
 	}
 }
 
@@ -256,8 +253,13 @@ UMaterialInstanceDynamic* FAvaMaskParametricMaterialHandle::GetMaterialInstance(
 	return FAvaParametricMaterialHandle::GetMaterialInstance();
 }
 
-void FAvaMaskParametricMaterialHandle::OnMaterialChanged()
+void FAvaMaskParametricMaterialHandle::OnMaterialChanged(const FAvaShapeParametricMaterial& InMaterial)
 {
+	if ((&InMaterial) != GetParametricMaterial())
+	{
+		return;
+	}
+
 	if (const FAvaShapeParametricMaterial* ParametricMtl = GetParametricMaterial())
 	{
 		if (UMaterialInstanceDynamic* MaterialInstance = ParametricMtl->GetMaterial())
diff --git a/Engine/Plugins/Experimental/Avalanche/Source/AvalancheMask/Private/Handling/AvaMaskParametricMaterialHandle.h b/Engine/Plugins/Experimental/Avalanche/Source/AvalancheMask/Private/Handling/AvaMaskParametricMaterialHandle.h
--- a/Engine/Plugins/Experimental/Avalanche/Source/AvalancheMask/Private/Handling/AvaMaskParametricMaterialHandle.h
+++ b/Engine/Plugins/Experimental/Avalanche/Source/AvalancheMask/Private/Handling/AvaMaskParametricMaterialHandle.h
@@ -8,6 +8,8 @@
 
 #include "AvaMaskParametricMaterialHandle.generated.h"
 
+struct FAvaShapeParametricMaterial;
+
 USTRUCT()
 struct FAvaMaskParametricMaterialData
 {
@@ -23,8 +25,6 @@ struct FAvaMaskParametricMaterialData
 	EAvaShapeParametricMaterialStyle OriginalMaterialStyle = EAvaShapeParametricMaterialStyle::Solid;
 };
 
-struct FAvaShapeParametricMaterial;
-
 class FAvaMaskParametricMaterialHandle
 	: public FAvaParametricMaterialHandle
 	, public TAvaMaskMaterialHandle<FAvaMaskParametricMaterialData>
@@ -64,12 +64,11 @@ public:
 protected:
 	virtual UMaterialInstanceDynamic* GetMaterialInstance() override;
 	
-	void OnMaterialChanged();
+	void OnMaterialChanged(const FAvaShapeParametricMaterial& InMaterial);
 
 private:
 	FMaterialShadingModelField LastShadingModel;
 	EBlendMode LastBlendMode;
 	TWeakObjectPtr<UMaterialInstance> LastMaterialInstance;
 	TWeakObjectPtr<UMaterialInstanceDynamic> ChildMaterial = nullptr;
-	FDelegateHandle MaterialChangedHandle;
 };
diff --git a/Engine/Plugins/Experimental/Avalanche/Source/AvalancheModifiers/Public/Modifiers/AvaSizeToTextureModifier.h b/Engine/Plugins/Experimental/Avalanche/Source/AvalancheModifiers/Public/Modifiers/AvaSizeToTextureModifier.h
--- a/Engine/Plugins/Experimental/Avalanche/Source/AvalancheModifiers/Public/Modifiers/AvaSizeToTextureModifier.h
+++ b/Engine/Plugins/Experimental/Avalanche/Source/AvalancheModifiers/Public/Modifiers/AvaSizeToTextureModifier.h
@@ -35,7 +35,7 @@ class UAvaSizeToTextureModifier : public UAvaGeometryBaseModifier
 
 public:
 	UFUNCTION(BlueprintCallable, Category="Motion Design|Modifiers|TextureSize")
-	void SetTexture(UTexture* InTexture);
+	AVALANCHEMODIFIERS_API void SetTexture(UTexture* InTexture);
 
 	UFUNCTION(BlueprintPure, Category="Motion Design|Modifiers|TextureSize")
 	UTexture* GetTexture() const
@@ -44,7 +44,7 @@ public:
 	}
 
 	UFUNCTION(BlueprintCallable, Category="Motion Design|Modifiers|TextureSize")
-	void SetRule(EAvaSizeToTextureRule InRule);
+	AVALANCHEMODIFIERS_API void SetRule(EAvaSizeToTextureRule InRule);
 
 	UFUNCTION(BlueprintPure, Category="Motion Design|Modifiers|TextureSize")
 	EAvaSizeToTextureRule GetRule() const
@@ -53,7 +53,7 @@ public:
 	}
 
 	UFUNCTION(BlueprintCallable, Category="Motion Design|Modifiers|TextureSize")
-	void SetFixedHeight(float InFixedHeight);
+	AVALANCHEMODIFIERS_API void SetFixedHeight(float InFixedHeight);
 
 	UFUNCTION(BlueprintPure, Category="Motion Design|Modifiers|TextureSize")
 	float GetFixedHeight() const
@@ -62,7 +62,7 @@ public:
 	}
 
 	UFUNCTION(BlueprintCallable, Category="Motion Design|Modifiers|TextureSize")
-	void SetFixedWidth(float InFixedWidth);
+	AVALANCHEMODIFIERS_API void SetFixedWidth(float InFixedWidth);
 
 	UFUNCTION(BlueprintPure, Category="Motion Design|Modifiers|TextureSize")
 	float GetFixedWidth() const
@@ -102,7 +102,7 @@ protected:
 
 	/** Rule for texture resize */
 	UPROPERTY(EditInstanceOnly, Setter, Getter, Category="TextureSize", meta=(AllowPrivateAccess="true"))
-	EAvaSizeToTextureRule Rule = EAvaSizeToTextureRule::FixedHeight;
+	EAvaSizeToTextureRule Rule = EAvaSizeToTextureRule::AdaptiveHeight;
 
 	/** The fixed height size */
 	UPROPERTY(EditInstanceOnly, Setter, Getter, Category="TextureSize", meta=(ClampMin="0", EditCondition="Rule == EAvaSizeToTextureRule::FixedHeight", EditConditionHides, AllowPrivateAccess="true"))
diff --git a/Engine/Plugins/Experimental/Avalanche/Source/AvalancheShapes/Private/AvaShapeParametricMaterial.cpp b/Engine/Plugins/Experimental/Avalanche/Source/AvalancheShapes/Private/AvaShapeParametricMaterial.cpp
--- a/Engine/Plugins/Experimental/Avalanche/Source/AvalancheShapes/Private/AvaShapeParametricMaterial.cpp
+++ b/Engine/Plugins/Experimental/Avalanche/Source/AvalancheShapes/Private/AvaShapeParametricMaterial.cpp
@@ -2,17 +2,27 @@
 
 #include "AvaShapeParametricMaterial.h"
 
+#include "Engine/Texture.h"
 #include "Materials/Material.h"
 #include "Materials/MaterialInstance.h"
 #include "Materials/MaterialInstanceDynamic.h"
 
 FAvaShapeParametricMaterial::FOnMaterialChanged FAvaShapeParametricMaterial::OnMaterialChangedDelegate;
+FAvaShapeParametricMaterial::FOnMaterialParameterChanged FAvaShapeParametricMaterial::OnMaterialParameterChangedDelegate;
+
+const FName FAvaShapeParametricMaterial::StyleParameterName = TEXT("Style");
+const FName FAvaShapeParametricMaterial::TextureParameterName = TEXT("Texture");
+const FName FAvaShapeParametricMaterial::ColorAParameterName = TEXT("ColorA");
+const FName FAvaShapeParametricMaterial::ColorBParameterName = TEXT("ColorB");
+const FName FAvaShapeParametricMaterial::GradientOffsetParameterName = TEXT("GradientOffset");
+const FName FAvaShapeParametricMaterial::GradientRotationParameterName = TEXT("GradientRotation");
 
 FAvaShapeParametricMaterial::FAvaShapeParametricMaterial(const FAvaShapeParametricMaterial& Other)
 {
 	bUseUnlitMaterial        = Other.bUseUnlitMaterial;
 	Translucency             = Other.Translucency;
 	bUseTwoSidedMaterial     = Other.bUseTwoSidedMaterial;
+	Texture                  = Other.Texture;
 	ColorA                   = Other.ColorA;
 	ColorB                   = Other.ColorB;
 	GradientOffset           = Other.GradientOffset;
@@ -28,6 +38,7 @@ FAvaShapeParametricMaterial& FAvaShapeParametricMaterial::operator=(const FAvaSh
 	bUseUnlitMaterial        = Other.bUseUnlitMaterial;
 	Translucency             = Other.Translucency;
 	bUseTwoSidedMaterial     = Other.bUseTwoSidedMaterial;
+	Texture                  = Other.Texture;
 	ColorA                   = Other.ColorA;
 	ColorB                   = Other.ColorB;
 	GradientOffset           = Other.GradientOffset;
@@ -119,18 +130,21 @@ void FAvaShapeParametricMaterial::SetMaterialParameterValues(UMaterialInstanceDy
 		return;
 	}
 
+	InMaterialInstance->SetScalarParameterValue(FAvaShapeParametricMaterial::StyleParameterName, static_cast<float>(Style));
+
 	switch (Style)
 	{
 		case EAvaShapeParametricMaterialStyle::Solid:
-			InMaterialInstance->SetScalarParameterValue(TEXT("LinearGradient"), 0.f);
-			InMaterialInstance->SetVectorParameterValue(TEXT("ColorA"), ColorA);
+			InMaterialInstance->SetVectorParameterValue(FAvaShapeParametricMaterial::ColorAParameterName, ColorA);
 		break;
-
 		case EAvaShapeParametricMaterialStyle::LinearGradient:
-			InMaterialInstance->SetScalarParameterValue(TEXT("LinearGradient"), 1.f);
-			InMaterialInstance->SetVectorParameterValue(TEXT("ColorA"), ColorA);
-			InMaterialInstance->SetVectorParameterValue(TEXT("ColorB"), ColorB);
-			InMaterialInstance->SetScalarParameterValue(TEXT("Offset"), GradientOffset);
+			InMaterialInstance->SetVectorParameterValue(FAvaShapeParametricMaterial::ColorAParameterName, ColorA);
+			InMaterialInstance->SetVectorParameterValue(FAvaShapeParametricMaterial::ColorBParameterName, ColorB);
+			InMaterialInstance->SetScalarParameterValue(FAvaShapeParametricMaterial::GradientOffsetParameterName, GradientOffset);
+			InMaterialInstance->SetScalarParameterValue(FAvaShapeParametricMaterial::GradientRotationParameterName, GradientRotation);
+		break;
+		case EAvaShapeParametricMaterialStyle::Texture:
+			InMaterialInstance->SetTextureParameterValue(FAvaShapeParametricMaterial::TextureParameterName, Texture);
 		break;
 	}
 
@@ -140,11 +154,34 @@ void FAvaShapeParametricMaterial::SetMaterialParameterValues(UMaterialInstanceDy
 	{
 		if (OnMaterialParameterChangedDelegate.IsBound())
 		{
-			OnMaterialParameterChangedDelegate.Broadcast();
+			OnMaterialParameterChangedDelegate.Broadcast(*this);
 		}
 	}
 }
 
+void FAvaShapeParametricMaterial::PostEditChange(TConstArrayView<FName> InPropertyNames)
+{
+	static const TArray<FName> PropertyNames =
+	{
+		GET_MEMBER_NAME_CHECKED(FAvaShapeParametricMaterial, Style),
+		GET_MEMBER_NAME_CHECKED(FAvaShapeParametricMaterial, Texture),
+		GET_MEMBER_NAME_CHECKED(FAvaShapeParametricMaterial, ColorA),
+		GET_MEMBER_NAME_CHECKED(FAvaShapeParametricMaterial, ColorB),
+		GET_MEMBER_NAME_CHECKED(FAvaShapeParametricMaterial, GradientOffset),
+		GET_MEMBER_NAME_CHECKED(FAvaShapeParametricMaterial, GradientRotation),
+		GET_MEMBER_NAME_CHECKED(FAvaShapeParametricMaterial, bUseUnlitMaterial),
+		GET_MEMBER_NAME_CHECKED(FAvaShapeParametricMaterial, bUseTwoSidedMaterial),
+		GET_MEMBER_NAME_CHECKED(FAvaShapeParametricMaterial, Translucency),
+	};
+
+	const FName MemberPropertyName = !InPropertyNames.IsEmpty() ? InPropertyNames.Last() : NAME_None;
+
+	if (PropertyNames.Contains(MemberPropertyName))
+	{
+		OnMaterialParameterUpdated();
+	}
+}
+
 bool FAvaShapeParametricMaterial::ShouldUseTranslucentMaterial() const
 {
 	if (Translucency == EAvaShapeParametricMaterialTranslucency::Auto)
@@ -158,13 +195,9 @@ bool FAvaShapeParametricMaterial::ShouldUseTranslucentMaterial() const
 void FAvaShapeParametricMaterial::OnMaterialParameterUpdated()
 {
 	const int32 PreviousActiveIndex = ActiveInstanceIndex;
+	const int32 NewActiveIndex = GetActiveInstanceIndex();
 
-	for (const TObjectPtr<UMaterialInstanceDynamic>& Material : InstanceMaterials)
-	{
-		SetMaterialParameterValues(Material, true);
-	}
-
-	if (PreviousActiveIndex != GetActiveInstanceIndex())
+	if (PreviousActiveIndex != NewActiveIndex)
 	{
 		if (OnMaterialChangedDelegate.IsBound())
 		{
@@ -172,6 +205,10 @@ void FAvaShapeParametricMaterial::OnMaterialParameterUpdated()
 		}
 	}
 
+	for (const TObjectPtr<UMaterialInstanceDynamic>& Material : InstanceMaterials)
+	{
+		SetMaterialParameterValues(Material, true);
+	}
 }
 
 int32 FAvaShapeParametricMaterial::GetActiveInstanceIndex() const
@@ -219,30 +256,36 @@ bool FAvaShapeParametricMaterial::CopyFromMaterialParameters(UMaterialInstance*
 		return false;
 	}
 
-	float IsLinearGradient;
-	if (InMaterial->GetScalarParameterValue(FMaterialParameterInfo(TEXT("LinearGradient")), IsLinearGradient))
+	float InputStyle;
+	if (InMaterial->GetScalarParameterValue(FMaterialParameterInfo(FAvaShapeParametricMaterial::StyleParameterName), InputStyle))
 	{
-		Style = IsLinearGradient >= 0.5f ? EAvaShapeParametricMaterialStyle::LinearGradient : EAvaShapeParametricMaterialStyle::Solid;
+		Style = static_cast<EAvaShapeParametricMaterialStyle>(FMath::TruncToInt(InputStyle));
 	}
 
 	FLinearColor InputColorA;
-	if (InMaterial->GetVectorParameterValue(FMaterialParameterInfo(TEXT("ColorA")), InputColorA))
+	if (InMaterial->GetVectorParameterValue(FMaterialParameterInfo(FAvaShapeParametricMaterial::ColorAParameterName), InputColorA))
 	{
 		ColorA = InputColorA;
 	}
 
 	FLinearColor InputColorB;
-	if (InMaterial->GetVectorParameterValue(FMaterialParameterInfo(TEXT("ColorB")), InputColorB))
+	if (InMaterial->GetVectorParameterValue(FMaterialParameterInfo(FAvaShapeParametricMaterial::ColorBParameterName), InputColorB))
 	{
 		ColorB = InputColorB;
 	}
 
 	float InputOffset;
-	if (InMaterial->GetScalarParameterValue(FMaterialParameterInfo(TEXT("Offset")), InputOffset))
+	if (InMaterial->GetScalarParameterValue(FMaterialParameterInfo(FAvaShapeParametricMaterial::GradientOffsetParameterName), InputOffset))
 	{
 		GradientOffset = InputOffset;
 	}
 
+	UTexture* InputTexture;
+	if (InMaterial->GetTextureParameterValue(FMaterialParameterInfo(FAvaShapeParametricMaterial::TextureParameterName), InputTexture))
+	{
+		Texture = InputTexture;
+	}
+
 	return true;
 }
 
@@ -375,9 +418,20 @@ void FAvaShapeParametricMaterial::SetStyle(EAvaShapeParametricMaterialStyle InSt
 	OnMaterialParameterUpdated();
 }
 
+void FAvaShapeParametricMaterial::SetTexture(UTexture* InTexture)
+{
+	if (Texture == InTexture)
+	{
+		return;
+	}
+
+	Texture = InTexture;
+	OnMaterialParameterUpdated();
+}
+
 void FAvaShapeParametricMaterial::SetPrimaryColor(const FLinearColor& InColor)
 {
-	if (InColor == ColorA)
+	if (InColor.Equals(ColorA))
 	{
 		return;
 	}
@@ -388,7 +442,7 @@ void FAvaShapeParametricMaterial::SetPrimaryColor(const FLinearColor& InColor)
 
 void FAvaShapeParametricMaterial::SetSecondaryColor(const FLinearColor& InColor)
 {
-	if (InColor == ColorB)
+	if (InColor.Equals(ColorB))
 	{
 		return;
 	}
@@ -399,7 +453,7 @@ void FAvaShapeParametricMaterial::SetSecondaryColor(const FLinearColor& InColor)
 
 void FAvaShapeParametricMaterial::SetGradientOffset(float InOffset)
 {
-	if (InOffset == GradientOffset)
+	if (FMath::IsNearlyEqual(InOffset, GradientOffset))
 	{
 		return;
 	}
@@ -408,6 +462,17 @@ void FAvaShapeParametricMaterial::SetGradientOffset(float InOffset)
 	OnMaterialParameterUpdated();
 }
 
+void FAvaShapeParametricMaterial::SetGradientRotation(float InRotation)
+{
+	if (FMath::IsNearlyEqual(InRotation, GradientRotation))
+	{
+		return;
+	}
+
+	GradientRotation = InRotation;
+	OnMaterialParameterUpdated();
+}
+
 void FAvaShapeParametricMaterial::SetUseUnlitMaterial(bool bInUse)
 {
 	if (bInUse == bUseUnlitMaterial)
diff --git a/Engine/Plugins/Experimental/Avalanche/Source/AvalancheShapes/Private/DynamicMeshes/AvaShapeDynMeshBase.cpp b/Engine/Plugins/Experimental/Avalanche/Source/AvalancheShapes/Private/DynamicMeshes/AvaShapeDynMeshBase.cpp
--- a/Engine/Plugins/Experimental/Avalanche/Source/AvalancheShapes/Private/DynamicMeshes/AvaShapeDynMeshBase.cpp
+++ b/Engine/Plugins/Experimental/Avalanche/Source/AvalancheShapes/Private/DynamicMeshes/AvaShapeDynMeshBase.cpp
@@ -59,6 +59,7 @@ UAvaShapeDynamicMeshBase::UAvaShapeDynamicMeshBase(const FLinearColor& InVertexC
 #endif
 
 		FAvaShapeParametricMaterial::OnMaterialChanged().AddUObject(this, &UAvaShapeDynamicMeshBase::OnParametricMaterialChanged);
+		FAvaShapeParametricMaterial::OnMaterialParameterChanged().AddUObject(this, &UAvaShapeDynamicMeshBase::OnParametricMaterialParameterChanged);
 	}
 }
 
@@ -110,16 +111,22 @@ void UAvaShapeDynamicMeshBase::SetupMeshes()
 			// register other meshes
 			RegisterMeshes();
 
+			const bool bIs2DShape = MeshDatas.Num() == 1;
 			for (const int32& MeshIdx : GetMeshesIndexes())
 			{
+				FAvaShapeMeshData* MeshData = GetMeshData(MeshIdx);
+
 				if (MeshIdx != MESH_INDEX_PRIMARY)
 				{
-					if (FAvaShapeMeshData* MeshData = GetMeshData(MeshIdx))
-					{
-						MeshData->bUsesPrimaryMaterialParams = bUsePrimaryMaterialEverywhere;
-					}
+					MeshData->bUsesPrimaryMaterialParams = bUsePrimaryMaterialEverywhere;
 				}
 
+				// Enable translucency by default
+				MeshData->ParametricMaterial.SetTranslucency(EAvaShapeParametricMaterialTranslucency::Enabled);
+
+				// Disable two sided materials to avoid inverted normal for 3D shapes
+				MeshData->ParametricMaterial.SetUseTwoSidedMaterial(bIs2DShape);
+
 				OnMaterialTypeChanged(MeshIdx);
 			}
 		}
@@ -216,6 +223,9 @@ void UAvaShapeDynamicMeshBase::OnComponentDestroyed(bool bDestroyingHierarchy)
 	{
 		MeshComponent->TransformUpdated.RemoveAll(this);
 	}
+
+	FAvaShapeParametricMaterial::OnMaterialChanged().RemoveAll(this);
+	FAvaShapeParametricMaterial::OnMaterialParameterChanged().RemoveAll(this);
 }
 
 void UAvaShapeDynamicMeshBase::ToggleGizmo_Implementation(
@@ -269,9 +279,7 @@ void UAvaShapeDynamicMeshBase::ToggleGizmo_Implementation(
 
 void UAvaShapeDynamicMeshBase::UpdateDynamicMesh()
 {
-	PRAGMA_DISABLE_DEPRECATION_WARNINGS
-	if (IsThisNotNull(this, "UAvaShapeDynamicMeshBase::UpdateDynamicMesh") && IsValidChecked(this) && !HasAnyFlags(RF_BeginDestroyed | RF_FinishDestroyed | RF_ClassDefaultObject))
-	PRAGMA_ENABLE_DEPRECATION_WARNINGS
+	if (IsValid(this) && !HasAnyFlags(RF_BeginDestroyed | RF_FinishDestroyed | RF_ClassDefaultObject))
 	{
 		RestoreCachedMesh();
 		CachedMesh.Reset();
@@ -1631,15 +1639,18 @@ void UAvaShapeDynamicMeshBase::PostEditChangeProperty(FPropertyChangedEvent& Pro
 	static const FName ScaledSizeName = GET_MEMBER_NAME_CHECKED(UAvaShapeDynamicMeshBase, UniformScaledSize);
 	static const FName UsePrimaryMaterialEverywhereName = GET_MEMBER_NAME_CHECKED(UAvaShapeDynamicMeshBase, bUsePrimaryMaterialEverywhere);
 
-	if (PropertyChangedEvent.MemberProperty->GetFName() == VertexColorName)
+	const FName MemberPropertyName = PropertyChangedEvent.GetMemberPropertyName();
+	const FName PropertyName = PropertyChangedEvent.GetPropertyName();
+
+	if (MemberPropertyName.IsEqual(VertexColorName))
 	{
 		OnVertexColorChanged();
 	}
-	else if (PropertyChangedEvent.MemberProperty->GetFName() == ScaledSizeName)
+	else if (MemberPropertyName.IsEqual(ScaledSizeName))
 	{
 		OnScaledSizeChanged();
 	}
-	else if (PropertyChangedEvent.Property->GetFName() == UsePrimaryMaterialEverywhereName)
+	else if (PropertyName.IsEqual(UsePrimaryMaterialEverywhereName))
 	{
 		OnUsePrimaryMaterialEverywhereChanged();
 	}
@@ -1649,9 +1660,35 @@ void UAvaShapeDynamicMeshBase::PostEditChangeChainProperty(FPropertyChangedChain
 {
 	Super::PostEditChangeChainProperty(PropertyChangedEvent);
 
+	auto GetChainPropertyNames = [](const FPropertyChangedChainEvent& InPropertyChangedEvent)->TArray<FName>
+	{
+		TArray<FName> PropertyNames;
+		FEditPropertyChain::TDoubleLinkedListNode* Node = InPropertyChangedEvent.PropertyChain.GetActiveMemberNode();
+
+		while (Node)
+		{
+			const FProperty* Property = Node->GetValue();
+
+			if (!Property)
+			{
+				break;
+			}
+
+			PropertyNames.Add(Property->GetFName());
+			Node = Node->GetNextNode();
+		}
+
+		Algo::Reverse(PropertyNames);
+		return PropertyNames;
+	};
+
 	static const FName MeshDatasName = GET_MEMBER_NAME_CHECKED(UAvaShapeDynamicMeshBase, MeshDatas);
 
-	if (PropertyChangedEvent.PropertyChain.GetActiveMemberNode()->GetValue()->GetFName() == MeshDatasName)
+	TArray<FName> PropertyNames = GetChainPropertyNames(PropertyChangedEvent);
+
+	const FName MemberPropertyName = !PropertyNames.IsEmpty() ? PropertyNames.Pop() : NAME_None;
+
+	if (MemberPropertyName == MeshDatasName)
 	{
 		const int32 MeshDatasIndex = PropertyChangedEvent.GetArrayIndex(MeshDatasName.ToString());
 
@@ -1669,26 +1706,25 @@ void UAvaShapeDynamicMeshBase::PostEditChangeChainProperty(FPropertyChangedChain
 		static const FName UsesPrimaryUVParamsName = GET_MEMBER_NAME_CHECKED(FAvaShapeMeshData, bOverridePrimaryUVParams);
 		static const FName MaterialUVParamsName = GET_MEMBER_NAME_CHECKED(FAvaShapeMeshData, MaterialUVParams);
 
-		const FProperty* Prop = PropertyChangedEvent.PropertyChain.GetActiveMemberNode()->GetNextNode()->GetValue();
-		const FName PropName = Prop->GetFName();
+		const FName MeshSectionProperty = !PropertyNames.IsEmpty() ? PropertyNames.Pop() : NAME_None;
 
-		if (PropName == MaterialTypeName)
+		if (MeshSectionProperty == MaterialTypeName)
 		{
 			OnMaterialTypeChanged(MeshIndex);
 		}
-		else if (PropName == MaterialName)
+		else if (MeshSectionProperty == MaterialName)
 		{
 			OnMaterialChanged(MeshIndex);
 		}
-		else if (PropName == ParametricMaterialName)
+		else if (MeshSectionProperty == ParametricMaterialName)
 		{
-			OnParametricMaterialChanged(MeshIndex);
+			MeshDatas[MeshIndex].ParametricMaterial.PostEditChange(PropertyNames);
 		}
-		else if (PropName == UsesPrimaryUVParamsName)
+		else if (MeshSectionProperty == UsesPrimaryUVParamsName)
 		{
 			OnUsesPrimaryUVParamsChanged(MeshIndex);
 		}
-		else if (PropName == MaterialUVParamsName)
+		else if (MeshSectionProperty == MaterialUVParamsName)
 		{
 			static const FName AnchorPresetName = GET_MEMBER_NAME_CHECKED(FAvaShapeMaterialUVParameters, AnchorPreset);
 			static const FName AnchorName = GET_MEMBER_NAME_CHECKED(FAvaShapeMaterialUVParameters, Anchor);
@@ -2018,14 +2054,26 @@ bool UAvaShapeDynamicMeshBase::ApplyUVsTransform(FAvaShapeMesh& InMesh, FAvaShap
 	return bResult;
 }
 
-void UAvaShapeDynamicMeshBase::OnParametricMaterialChanged(FAvaShapeParametricMaterial& InMaterial)
+void UAvaShapeDynamicMeshBase::OnParametricMaterialChanged(const FAvaShapeParametricMaterial& InMaterial)
 {
 	for (const int32 Index : GetMeshesIndexes())
 	{
-		if (GetParametricMaterial(Index) == &InMaterial
-			&& GetMaterialType(Index) == EMaterialType::Parametric)
+		if (GetParametricMaterial(Index) == &InMaterial && GetMaterialType(Index) == EMaterialType::Parametric)
 		{
 			OnParametricMaterialChanged(Index);
+			break;
+		}
+	}
+}
+
+void UAvaShapeDynamicMeshBase::OnParametricMaterialParameterChanged(const FAvaShapeParametricMaterial& InMaterial)
+{
+	for (const int32 Index : GetMeshesIndexes())
+	{
+		if (GetParametricMaterial(Index) == &InMaterial && GetMaterialType(Index) == EMaterialType::Parametric)
+		{
+			OnParametricMaterialChanged(Index);
+			break;
 		}
 	}
 }
diff --git a/Engine/Plugins/Experimental/Avalanche/Source/AvalancheShapes/Public/AvaShapeParametricMaterial.h b/Engine/Plugins/Experimental/Avalanche/Source/AvalancheShapes/Public/AvaShapeParametricMaterial.h
--- a/Engine/Plugins/Experimental/Avalanche/Source/AvalancheShapes/Public/AvaShapeParametricMaterial.h
+++ b/Engine/Plugins/Experimental/Avalanche/Source/AvalancheShapes/Public/AvaShapeParametricMaterial.h
@@ -13,6 +13,7 @@ class UMaterialInstance;
 class UMaterialInstanceDynamic;
 class UMaterialInterface;
 class UObject;
+class UTexture;
 struct FAvaShapeParametricMaterial;
 
 UENUM(BlueprintType)
@@ -45,19 +46,26 @@ struct FAvaShapeParametricMaterial
 		MaterialTypeCount = Translucent + Unlit + OneSided + 1
 	};
 
-	DECLARE_MULTICAST_DELEGATE_OneParam(FOnMaterialChanged, FAvaShapeParametricMaterial&)
-	DECLARE_MULTICAST_DELEGATE(FOnMaterialParameterChanged)
+	DECLARE_MULTICAST_DELEGATE_OneParam(FOnMaterialChanged, const FAvaShapeParametricMaterial&)
+	DECLARE_MULTICAST_DELEGATE_OneParam(FOnMaterialParameterChanged, const FAvaShapeParametricMaterial&)
+
+	static const FName StyleParameterName;
+	static const FName TextureParameterName;
+	static const FName ColorAParameterName;
+	static const FName ColorBParameterName;
+	static const FName GradientOffsetParameterName;
+	static const FName GradientRotationParameterName;
 	
-	/** Called when active material index changed */
+	/** Called when active material changed */
 	static FOnMaterialChanged::RegistrationType& OnMaterialChanged();
 
+	/** Called when active material parameters changed */
+	static AVALANCHESHAPES_API FOnMaterialParameterChanged::RegistrationType& OnMaterialParameterChanged();
+
 	AVALANCHESHAPES_API FAvaShapeParametricMaterial() = default;
 	AVALANCHESHAPES_API FAvaShapeParametricMaterial(const FAvaShapeParametricMaterial& Other);
 	AVALANCHESHAPES_API FAvaShapeParametricMaterial& operator=(const FAvaShapeParametricMaterial& Other);
 
-	/** Called when active material parameters are updated */
-	AVALANCHESHAPES_API FOnMaterialParameterChanged::RegistrationType& OnMaterialParameterChanged();
-
 	AVALANCHESHAPES_API bool CopyFromMaterialParameters(UMaterialInstance* InMaterial);
 
 	/** Check if input material is a parametric material */
@@ -76,12 +84,12 @@ struct FAvaShapeParametricMaterial
 	void SetMaterial(UMaterialInstanceDynamic* InMaterial);
 
 	AVALANCHESHAPES_API void SetTranslucency(EAvaShapeParametricMaterialTranslucency InTranslucency);
-	
+
 	EAvaShapeParametricMaterialTranslucency GetTranslucency() const
 	{
 		return Translucency;
 	}
-	
+
 	AVALANCHESHAPES_API bool ShouldUseTranslucentMaterial() const;
 
 	EAvaShapeParametricMaterialStyle GetStyle() const
@@ -96,6 +104,13 @@ struct FAvaShapeParametricMaterial
 		return ColorA;
 	}
 
+	AVALANCHESHAPES_API void SetTexture(UTexture* InTexture);
+
+	UTexture* GetTexture() const
+	{
+		return Texture;
+	}
+
 	AVALANCHESHAPES_API void SetPrimaryColor(const FLinearColor& InColor);
 
 	const FLinearColor& GetSecondaryColor() const
@@ -112,6 +127,13 @@ struct FAvaShapeParametricMaterial
 
 	AVALANCHESHAPES_API void SetGradientOffset(float InOffset);
 
+	float GetGradientRotation() const
+	{
+		return GradientRotation;
+	}
+
+	AVALANCHESHAPES_API void SetGradientRotation(float InRotation);
+
 	bool GetUseUnlitMaterial() const
 	{
 		return bUseUnlitMaterial;
@@ -129,28 +151,38 @@ struct FAvaShapeParametricMaterial
 	/** Set parameter values on a material instance */
 	AVALANCHESHAPES_API void SetMaterialParameterValues(UMaterialInstanceDynamic* InMaterialInstance, bool bInNotifyUpdate = false) const;
 
+	void PostEditChange(TConstArrayView<FName> InPropertyNames);
+
 protected:
-	// Default style for the material
-	UPROPERTY(EditInstanceOnly, BlueprintReadWrite, Category="Material", meta=(DisplayName="Simple Style", AllowPrivateAccess = "true"))
+	/** Default style for the material */
+	UPROPERTY(EditInstanceOnly, BlueprintReadWrite, Category="Material", meta=(DisplayName="Material Style", AllowPrivateAccess = "true"))
 	EAvaShapeParametricMaterialStyle Style = EAvaShapeParametricMaterialStyle::Solid;
 
-	// Primary colour for the material
-	UPROPERTY(EditInstanceOnly, BlueprintReadWrite, Category="Material", meta=(DisplayName="Simple Color", AllowPrivateAccess = "true"))
+	/** Simple texture for the material */
+	UPROPERTY(EditInstanceOnly, BlueprintReadWrite, Category="Material", meta=(DisplayName="Texture", EditCondition="Style == EAvaShapeParametricMaterialStyle::Texture", EditConditionHides, AllowPrivateAccess = "true"))
+	TObjectPtr<UTexture> Texture = nullptr;
+
+	/** Primary colour for the material */
+	UPROPERTY(EditInstanceOnly, BlueprintReadWrite, Category="Material", meta=(DisplayName="Primary Color", EditCondition="Style != EAvaShapeParametricMaterialStyle::Texture", EditConditionHides, AllowPrivateAccess = "true"))
 	FLinearColor ColorA = FLinearColor::White;
 
-	// Secondary colour for the material
-	UPROPERTY(EditInstanceOnly, BlueprintReadWrite, Category="Material", meta=(DisplayName="Gradient Color", EditCondition="Style == EAvaShapeParametricMaterialStyle::LinearGradient", EditConditionHides, AllowPrivateAccess = "true"))
+	/** Secondary colour for the material */
+	UPROPERTY(EditInstanceOnly, BlueprintReadWrite, Category="Material", meta=(DisplayName="Secondary Color", EditCondition="Style == EAvaShapeParametricMaterialStyle::LinearGradient", EditConditionHides, AllowPrivateAccess = "true"))
 	FLinearColor ColorB = FLinearColor::Black;
 
-	// Offset for gradient style material
+	/** Offset for gradient style material */
 	UPROPERTY(EditInstanceOnly, BlueprintReadWrite, Category="Material", meta=(ClampMin="0", ClampMax="1", DisplayName="Gradient Offset", EditCondition = "Style == EAvaShapeParametricMaterialStyle::LinearGradient", EditConditionHides, AllowPrivateAccess = "true"))
 	float GradientOffset = 0.5f;
 
-	// whether the material is unlit or default lit
+	/** Rotation for gradient style material */
+	UPROPERTY(EditInstanceOnly, BlueprintReadWrite, Category="Material", meta=(ClampMin="0", ClampMax="1", DisplayName="Gradient Rotation", EditCondition = "Style == EAvaShapeParametricMaterialStyle::LinearGradient", EditConditionHides, AllowPrivateAccess = "true"))
+	float GradientRotation = 0.f;
+
+	/** whether the material is unlit or default lit */
 	UPROPERTY(EditInstanceOnly, BlueprintReadWrite, Category="Material", meta=(DisplayName="Use Unlit Material", AllowPrivateAccess = "true"))
 	bool bUseUnlitMaterial = false;
 
-	// whether the material is one sided or two sided
+	/** whether the material is one sided or two sided */
 	UPROPERTY(EditInstanceOnly, BlueprintReadWrite, Category = "Material", meta = (DisplayName = "Use Two Sided Material", AllowPrivateAccess = "true"))
 	bool bUseTwoSidedMaterial = true;
 
@@ -160,6 +192,7 @@ protected:
 
 private:
 	static FOnMaterialChanged OnMaterialChangedDelegate;
+	static FOnMaterialParameterChanged OnMaterialParameterChangedDelegate;
 
 	/** Parent material
 	 * 1. Opaque Lit
@@ -181,8 +214,6 @@ private:
 	UPROPERTY(Transient)
 	TArray<TObjectPtr<UMaterialInstanceDynamic>> InstanceMaterials;
 
-	FOnMaterialParameterChanged OnMaterialParameterChangedDelegate;
-
 	int32 ActiveInstanceIndex = INDEX_NONE;
 
 	/** Load parents materials to create instance materials */
diff --git a/Engine/Plugins/Experimental/Avalanche/Source/AvalancheShapes/Public/AvaShapesDefs.h b/Engine/Plugins/Experimental/Avalanche/Source/AvalancheShapes/Public/AvaShapesDefs.h
--- a/Engine/Plugins/Experimental/Avalanche/Source/AvalancheShapes/Public/AvaShapesDefs.h
+++ b/Engine/Plugins/Experimental/Avalanche/Source/AvalancheShapes/Public/AvaShapesDefs.h
@@ -51,7 +51,8 @@ UENUM(BlueprintType)
 enum class EAvaShapeParametricMaterialStyle : uint8
 {
 	Solid,
-	LinearGradient
+	LinearGradient,
+	Texture
 };
 
 UENUM(BlueprintType)
diff --git a/Engine/Plugins/Experimental/Avalanche/Source/AvalancheShapes/Public/DynamicMeshes/AvaShapeDynMeshBase.h b/Engine/Plugins/Experimental/Avalanche/Source/AvalancheShapes/Public/DynamicMeshes/AvaShapeDynMeshBase.h
--- a/Engine/Plugins/Experimental/Avalanche/Source/AvalancheShapes/Public/DynamicMeshes/AvaShapeDynMeshBase.h
+++ b/Engine/Plugins/Experimental/Avalanche/Source/AvalancheShapes/Public/DynamicMeshes/AvaShapeDynMeshBase.h
@@ -17,6 +17,7 @@
 class AAvaShapeActor;
 class UDynamicMeshComponent;
 class UStaticMesh;
+struct FAvaShapeParametricMaterial;
 
 DECLARE_MULTICAST_DELEGATE_OneParam(FMaskEnabledDelegate, AActor* /** NewMaskActor */);
 DECLARE_MULTICAST_DELEGATE_OneParam(FMaskDisabledDelegate, AActor* /** OldMaskActor */);
@@ -56,7 +57,7 @@ public:
 	{}
 
 	UAvaShapeDynamicMeshBase(const FLinearColor& InVertexColor = FLinearColor::White,
-		float InUniformScaledSize = 1.f, bool bInAllowEditSize = true);
+	                         float InUniformScaledSize = 1.f, bool bInAllowEditSize = true);
 
 	// Converts the Outer of this object to a AAvaShapeActor.
 	AVALANCHESHAPES_API AAvaShapeActor* GetShapeActor() const;
@@ -221,9 +222,6 @@ protected:
 	/** Only update vertices from mesh */
 	void MarkVerticesDirty();
 
-	// When a parametric material is updated
-	void OnParametricMaterialChanged(FAvaShapeParametricMaterial& InMaterial);
-
 	/** Register all meshes once, calls RegisterMeshes() */
 	void SetupMeshes();
 
@@ -401,6 +399,12 @@ private:
 
 	void OnMeshTransformChanged(USceneComponent* InComponent, EUpdateTransformFlags InTransformFlags, ETeleportType InTeleport);
 
+	/** When a parametric material is updated, eg: opaque -> translucent */
+	void OnParametricMaterialChanged(const FAvaShapeParametricMaterial& InMaterial);
+
+	/** When a parametric material parameter is updated */
+	void OnParametricMaterialParameterChanged(const FAvaShapeParametricMaterial& InMaterial);
+
 	// used to quickly restore mesh that has been modified after OnMeshUpdateFinished was called (eg extrude, masking)
 	TOptional<UE::Geometry::FDynamicMesh3> CachedMesh;
 
