[Motion Design] Stagger Tool:
- Process section models as layer bars

diff --git a/Engine/Plugins/Experimental/Avalanche/Source/AvalancheSequencer/Private/StaggerTool/AvaStaggerBarElement.cpp b/Engine/Plugins/Experimental/Avalanche/Source/AvalancheSequencer/Private/StaggerTool/AvaStaggerBarElement.cpp
new file mode 100644
--- /dev/null
+++ b/Engine/Plugins/Experimental/Avalanche/Source/AvalancheSequencer/Private/StaggerTool/AvaStaggerBarElement.cpp
@@ -0,0 +1,73 @@
+﻿// Copyright Epic Games, Inc. All Rights Reserved.
+
+#include "AvaStaggerBarElement.h"
+#include "MVVM/ViewModelPtr.h"
+#include "MVVM/Extensions/ITrackAreaExtension.h"
+
+using namespace UE::Sequencer;
+
+FAvaStaggerBarElement FAvaStaggerBarElement::FromTrack(const TViewModelPtr<ITrackAreaExtension>& InTrackAreaExtension)
+{
+	for (const FViewModelPtr& TrackAreaModel : InTrackAreaExtension->GetTopLevelChildTrackAreaModels())
+	{
+		if (const TViewModelPtr<FLayerBarModel> BarModel = TrackAreaModel.ImplicitCast())
+		{
+			return BarModel;
+		}
+	}
+	for (const TViewModelPtr<ILayerBarExtension>& TrackAreaModel : InTrackAreaExtension->GetTrackAreaModelListAs<ILayerBarExtension>())
+	{
+		if (const TViewModelPtr<ILayerBarExtension> BarModel = TrackAreaModel.ImplicitCast())
+		{
+			return BarModel;
+		}
+	}
+	return FAvaStaggerBarElement();
+}
+
+FAvaStaggerBarElement::FAvaStaggerBarElement(const TViewModelPtr<FLayerBarModel>& InBarModel)
+{
+	if (const TViewModelPtr<IOutlinerExtension> LinkedOutlinerItem = InBarModel->GetLinkedOutlinerItem())
+	{
+		BarModel.Emplace<TViewModelPtr<FLayerBarModel>>(InBarModel);
+		OutlinerItem = LinkedOutlinerItem;
+		Range = InBarModel->ComputeRange();
+		OriginalFrame = Range.GetLowerBoundValue();
+	}
+}
+
+FAvaStaggerBarElement::FAvaStaggerBarElement(const TViewModelPtr<ILayerBarExtension>& InBarModel)
+{
+	if (const TViewModelPtr<FLinkedOutlinerExtension> LinkedOutlinerExtension = InBarModel.ImplicitCast())
+	{
+		BarModel.Emplace<TViewModelPtr<ILayerBarExtension>>(InBarModel);
+		OutlinerItem = LinkedOutlinerExtension->GetLinkedOutlinerItem();
+		Range = InBarModel->GetLayerBarRange();
+		OriginalFrame = Range.GetLowerBoundValue();
+	}
+}
+
+bool FAvaStaggerBarElement::IsValid() const
+{
+	if (const TViewModelPtr<FLayerBarModel>* LayerBarModel = BarModel.TryGet<TViewModelPtr<FLayerBarModel>>())
+	{
+		return LayerBarModel->IsValid();
+	}
+	if (const TViewModelPtr<ILayerBarExtension>* LayerBarExtension = BarModel.TryGet<TViewModelPtr<ILayerBarExtension>>())
+	{
+		return LayerBarExtension->IsValid();
+	}
+	return false;
+}
+
+void FAvaStaggerBarElement::Offset(const FFrameNumber InDelta)
+{
+	if (const TViewModelPtr<FLayerBarModel>* LayerBarModel = BarModel.TryGet<TViewModelPtr<FLayerBarModel>>())
+	{
+		(*LayerBarModel)->Offset(InDelta);
+	}
+	else if (const TViewModelPtr<ILayerBarExtension>* LayerBarExtension = BarModel.TryGet<TViewModelPtr<ILayerBarExtension>>())
+	{
+		(*LayerBarExtension)->OffsetLayerBar(InDelta);
+	}
+}
diff --git a/Engine/Plugins/Experimental/Avalanche/Source/AvalancheSequencer/Private/StaggerTool/AvaStaggerBarElement.h b/Engine/Plugins/Experimental/Avalanche/Source/AvalancheSequencer/Private/StaggerTool/AvaStaggerBarElement.h
new file mode 100644
--- /dev/null
+++ b/Engine/Plugins/Experimental/Avalanche/Source/AvalancheSequencer/Private/StaggerTool/AvaStaggerBarElement.h
@@ -0,0 +1,57 @@
+﻿// Copyright Epic Games, Inc. All Rights Reserved.
+
+#pragma once
+
+#include "Math/Range.h"
+#include "MVVM/ViewModelPtr.h"
+#include "MVVM/ViewModels/LayerBarModel.h"
+
+struct FFrameNumber;
+
+namespace UE::Sequencer
+{
+	class FLayerBarModel;
+	class ILayerBarExtension;
+	class IOutlinerExtension;
+	class ITrackAreaExtension;
+}
+
+/** Structure that holds cached data for a stagger element operation */
+struct FAvaStaggerBarElement
+{
+	static FAvaStaggerBarElement FromTrack(const UE::Sequencer::TViewModelPtr<UE::Sequencer::ITrackAreaExtension>& InTrackAreaExtension);
+
+	FAvaStaggerBarElement() {}
+	FAvaStaggerBarElement(const UE::Sequencer::TViewModelPtr<UE::Sequencer::FLayerBarModel>& InLayerBar);
+	FAvaStaggerBarElement(const UE::Sequencer::TViewModelPtr<UE::Sequencer::ILayerBarExtension>& InSection);
+
+	bool operator==(const FAvaStaggerBarElement& InOther)
+	{
+		if (BarModel.IsType<UE::Sequencer::TViewModelPtr<UE::Sequencer::FLayerBarModel>>())
+		{
+			return InOther.BarModel.IsType<UE::Sequencer::TViewModelPtr<UE::Sequencer::FLayerBarModel>>()
+				&& InOther.BarModel.Get<UE::Sequencer::TViewModelPtr<UE::Sequencer::FLayerBarModel>>()
+					== BarModel.Get<UE::Sequencer::TViewModelPtr<UE::Sequencer::FLayerBarModel>>();
+		}
+		if (BarModel.IsType<UE::Sequencer::TViewModelPtr<UE::Sequencer::ILayerBarExtension>>())
+		{
+			return InOther.BarModel.IsType<UE::Sequencer::TViewModelPtr<UE::Sequencer::ILayerBarExtension>>()
+				&& InOther.BarModel.Get<UE::Sequencer::TViewModelPtr<UE::Sequencer::ILayerBarExtension>>()
+					== BarModel.Get<UE::Sequencer::TViewModelPtr<UE::Sequencer::ILayerBarExtension>>();
+		}
+		return false;
+	}
+
+	bool IsValid() const;
+
+	void Offset(const FFrameNumber InDelta);
+
+	TVariant<UE::Sequencer::TViewModelPtr<UE::Sequencer::FLayerBarModel>
+		, UE::Sequencer::TViewModelPtr<UE::Sequencer::ILayerBarExtension>> BarModel;
+
+	UE::Sequencer::TViewModelPtr<UE::Sequencer::IOutlinerExtension> OutlinerItem;
+	TRange<FFrameNumber> Range;
+
+	FFrameNumber OriginalFrame;
+	FFrameNumber LastOffset;
+};
diff --git a/Engine/Plugins/Experimental/Avalanche/Source/AvalancheSequencer/Private/StaggerTool/AvaStaggerTool.cpp b/Engine/Plugins/Experimental/Avalanche/Source/AvalancheSequencer/Private/StaggerTool/AvaStaggerTool.cpp
--- a/Engine/Plugins/Experimental/Avalanche/Source/AvalancheSequencer/Private/StaggerTool/AvaStaggerTool.cpp
+++ b/Engine/Plugins/Experimental/Avalanche/Source/AvalancheSequencer/Private/StaggerTool/AvaStaggerTool.cpp
@@ -1,6 +1,7 @@
 // Copyright Epic Games, Inc. All Rights Reserved.
 
 #include "AvaStaggerTool.h"
+#include "AvaStaggerBarElement.h"
 #include "Channels/MovieSceneChannel.h"
 #include "Commands/AvaSequencerCommands.h"
 #include "Framework/Application/SlateApplication.h"
@@ -11,6 +12,7 @@
 #include "MVVM/Selection/Selection.h"
 #include "MVVM/ViewModels/ChannelModel.h"
 #include "MVVM/ViewModels/LayerBarModel.h"
+#include "MVVM/ViewModels/SectionModel.h"
 #include "MVVM/ViewModels/SequencerEditorViewModel.h"
 #include "ScopedTransaction.h"
 #include "Tools/SequencerSelectionAlignmentUtils.h"
@@ -20,14 +22,7 @@ using namespace UE::Sequencer;
 
 #define LOCTEXT_NAMESPACE "AvaSequencerStagger"
 
-FAvaStaggerTool::FStaggerLayerBarElement::FStaggerLayerBarElement(TViewModelPtr<FLayerBarModel> InLayerBarModel)
-	: LayerBarModel(MoveTemp(InLayerBarModel))
-{
-	const TRange<FFrameNumber> ComputedRange = LayerBarModel->ComputeRange();
-	OriginalFrame = ComputedRange.GetLowerBoundValue();
-}
-
-FAvaStaggerTool::FStaggerKeyFrameElement::FStaggerKeyFrameElement(TViewModelPtr<FChannelModel> InKeyChannelModel
+FAvaStaggerTool::FAvaStaggerKeyElement::FAvaStaggerKeyElement(TViewModelPtr<FChannelModel> InKeyChannelModel
 	, FKeyHandle InKeyHandle, FFrameNumber InOriginalFrame)
 	: KeyChannelModel(MoveTemp(InKeyChannelModel))
 	, KeyHandle(MoveTemp(InKeyHandle))
@@ -188,29 +183,29 @@ void FAvaStaggerTool::CloseToolWindow()
 
 bool FAvaStaggerTool::HasValidSelection() const
 {
-	return IsLayerBarSelection() || IsKeyFrameSelection();
+	return IsBarSelection() || IsKeySelection();
 }
 
-bool FAvaStaggerTool::IsLayerBarSelection() const
+bool FAvaStaggerTool::IsBarSelection() const
 {
-	return OriginalLayerBarElements.Num() > 1 && OriginalKeyFrameElements.Num() == 0;
+	return OriginalBarElements.Num() > 1 && OriginalKeyElements.Num() == 0;
 }
 
-bool FAvaStaggerTool::IsKeyFrameSelection() const
+bool FAvaStaggerTool::IsKeySelection() const
 {
-	return OriginalKeyFrameElements.Num() > 1 && OriginalLayerBarElements.Num() == 0;
+	return OriginalKeyElements.Num() > 1 && OriginalBarElements.Num() == 0;
 }
 
 int32 FAvaStaggerTool::GetSelectionCount() const
 {
-	if (IsLayerBarSelection())
+	if (IsBarSelection())
 	{
-		return OriginalLayerBarElements.Num();
+		return OriginalBarElements.Num();
 	}
 
-	if (IsKeyFrameSelection())
+	if (IsKeySelection())
 	{
-		return OriginalKeyFrameElements.Num();
+		return OriginalKeyElements.Num();
 	}
 
 	return 0;
@@ -262,26 +257,25 @@ FFrameNumber FAvaStaggerTool::GetInterval(const FFrameNumber InRangeSize, const
 	}
 }
 
-FFrameNumber FAvaStaggerTool::CalculateLayerBarFrameSpan(const TArray<FStaggerLayerBarElement>& InStaggerElements, const int32 InStopIndex)
+FFrameNumber FAvaStaggerTool::CalculateBarElementFrameSpan(const TArray<FAvaStaggerBarElement>& InElements, const int32 InStopIndex)
 {
 	FFrameNumber OutFrame;
-	for (int32 Index = 0; Index < InStaggerElements.Num(); ++Index)
+	for (int32 Index = 0; Index < InElements.Num(); ++Index)
 	{
 		if (Index == InStopIndex)
 		{
 			break;
 		}
-		OutFrame += InStaggerElements[Index].LayerBarModel->ComputeRange().Size<FFrameNumber>();
+		OutFrame += InElements[Index].Range.Size<FFrameNumber>();
 	}
 	return MoveTemp(OutFrame);
 }
 
-FFrameNumber FAvaStaggerTool::GetLayerBarOperationOffset(const TViewModelPtr<FLayerBarModel>& InLayerBarModel) const
+FFrameNumber FAvaStaggerTool::GetBarElementOperationOffset(const FAvaStaggerBarElement& InElement) const
 {
-	const TRange<FFrameNumber> Range = InLayerBarModel->ComputeRange();
-	const FFrameNumber RangeSize = Range.Size<FFrameNumber>();
-	const float LayerBarPosition = Settings->ToolOptions.OperationPoint * RangeSize.Value;
-	return FMath::FloorToInt32(LayerBarPosition);
+	const FFrameNumber RangeSize = InElement.Range.Size<FFrameNumber>();
+	const float Position = Settings->ToolOptions.OperationPoint * RangeSize.Value;
+	return FMath::FloorToInt32(Position);
 }
 
 void FAvaStaggerTool::CacheOriginalElements()
@@ -296,24 +290,24 @@ void FAvaStaggerTool::CacheOriginalElements()
 	const int32 TrackAreaCount = SequencerSelection->TrackArea.Num();
 	const int32 KeyFrameCount = SequencerSelection->KeySelection.Num();
 
-	OriginalLayerBarElements.Empty();
-	OriginalKeyFrameElements.Empty();
+	OriginalBarElements.Empty();
+	OriginalKeyElements.Empty();
 
 	if ((OutlinerCount > 1 && TrackAreaCount == 0 && KeyFrameCount == 0)
 		|| (OutlinerCount <= 1 && TrackAreaCount > 1 && KeyFrameCount == 0))
 	{
-		OriginalLayerBarElements = GatherSelectionLayerBarElements();
+		OriginalBarElements = GatherSelectionBarElements(SequencerSelection);
 	}
 	else if (OutlinerCount == 0 && TrackAreaCount == 0 && KeyFrameCount > 1)
 	{
-		OriginalKeyFrameElements = GatherSelectionKeyFrameElements();
+		OriginalKeyElements = GatherSelectionKeyElements();
 	}
 
-	CachedLayerBarCount = OriginalLayerBarElements.Num();
-	CachedKeyframeCount = OriginalKeyFrameElements.Num();
+	CachedBarCount = OriginalBarElements.Num();
+	CachedKeyCount = OriginalKeyElements.Num();
 }
 
-FFrameNumber FAvaStaggerTool::FindFirstLayerBarStaggerPoint() const
+FFrameNumber FAvaStaggerTool::FindFirstBarStaggerPoint() const
 {
 	FFrameNumber OutFrameNumber;
 
@@ -341,115 +335,94 @@ FFrameNumber FAvaStaggerTool::FindFirstLayerBarStaggerPoint() const
 		break;
 
 	case EAvaSequencerStaggerDistribution::Random:
-		const FFrameNumber ElementSize = OriginalLayerBarElements[0].LayerBarModel->ComputeRange().Size<FFrameNumber>();
-		const FFrameNumber Interval = GetInterval(CachedRangeSize, CachedLayerBarCount, ElementSize);
+		const FFrameNumber ElementSize = OriginalBarElements[0].Range.Size<FFrameNumber>();
+		const FFrameNumber Interval = GetInterval(CachedRangeSize, CachedBarCount, ElementSize);
 		return CachedRange.GetLowerBoundValue() + Interval + CachedShiftFrames;
 	}
 
-	return GetStartPosition(OriginalLayerBarElements);
+	return GetStartPosition(OriginalBarElements);
 }
 
-FFrameNumber FAvaStaggerTool::FindFirstKeyframeStaggerPoint() const
+FFrameNumber FAvaStaggerTool::FindFirstKeyStaggerPoint() const
 {
 	if (Settings->ToolOptions.Distribution == EAvaSequencerStaggerDistribution::Random)
 	{
-		const FFrameNumber Interval = GetInterval(CachedRangeSize, CachedKeyframeCount);
+		const FFrameNumber Interval = GetInterval(CachedRangeSize, CachedKeyCount);
 		return CachedRange.GetLowerBoundValue() + Interval + CachedShiftFrames;
 	}
 
-	return GetStartPosition(OriginalLayerBarElements);
+	return GetStartPosition(OriginalBarElements);
 }
 
-TArray<TViewModelPtr<FLayerBarModel>> FAvaStaggerTool::GatherAllSelectedLayerBars() const
+TArray<FAvaStaggerBarElement> FAvaStaggerTool::GatherSelectionBarElements(const TSharedPtr<FSequencerSelection>& InSequencerSelection) const
 {
-	const TSharedPtr<FSequencerSelection> SequencerSelection = GetSequencerSelection();
-	if (!SequencerSelection.IsValid())
+	if (!InSequencerSelection.IsValid())
 	{
 		return {};
 	}
 
-	TArray<TViewModelPtr<FLayerBarModel>> LayerBarModels;
+	TArray<FAvaStaggerBarElement> OutBarElements;
 
-	const int32 OutlinerCount = SequencerSelection->Outliner.Num();
-	const int32 TrackAreaCount = SequencerSelection->TrackArea.Num();
+	const int32 OutlinerCount = InSequencerSelection->Outliner.Num();
+	const int32 TrackAreaCount = InSequencerSelection->TrackArea.Num();
 
 	if (OutlinerCount > 1)
 	{
-		LayerBarModels.Reserve(OutlinerCount);
+		OutBarElements.Reserve(OutlinerCount);
 
-		for (const FViewModelPtr& ViewModel : SequencerSelection->Outliner)
+		for (const FViewModelPtr& ViewModel : InSequencerSelection->Outliner)
 		{
-			const TViewModelPtr<ITrackAreaExtension> Track = ViewModel.ImplicitCast();
-			if (!Track.IsValid())
+			const TViewModelPtr<ITrackAreaExtension> TrackArea = ViewModel.ImplicitCast();
+			if (!TrackArea.IsValid())
 			{
 				continue;
 			}
 
-			for (const FViewModelPtr& TrackAreaModel : Track->GetTopLevelChildTrackAreaModels())
-			{
-				if (const TViewModelPtr<FLayerBarModel> LayerBarModel = TrackAreaModel.ImplicitCast())
-				{
-					LayerBarModels.AddUnique(LayerBarModel);
-				}
-			}
-
-			for (const TViewModelPtr<FLayerBarModel>& LayerBarModel : Track->GetTrackAreaModelListAs<FLayerBarModel>())
+			FAvaStaggerBarElement NewElement = FAvaStaggerBarElement::FromTrack(TrackArea);
+			if (NewElement.IsValid())
 			{
-				LayerBarModels.AddUnique(LayerBarModel);
+				OutBarElements.Add(NewElement);
 			}
 		}
 	}
 	else if (OutlinerCount <= 1 && TrackAreaCount > 0)
 	{
-		LayerBarModels.Reserve(TrackAreaCount);
+		OutBarElements.Reserve(TrackAreaCount);
 
-		for (const FViewModelPtr& ViewModel : SequencerSelection->TrackArea)
+		for (const FViewModelPtr& ViewModel : InSequencerSelection->TrackArea)
 		{
 			if (const TViewModelPtr<FLayerBarModel> LayerBarModel = ViewModel.ImplicitCast())
 			{
-				LayerBarModels.AddUnique(LayerBarModel);
+				OutBarElements.Add(LayerBarModel);
+			}
+			else if (const TViewModelPtr<ILayerBarExtension> LayerBarExtension = ViewModel.ImplicitCast())
+			{
+				OutBarElements.Add(LayerBarExtension);
 			}
 		}
 	}
 
 	// Disallow operations on all layer bars that have a descendant that is selected
-	for (const TViewModelPtr<FLayerBarModel>& InLayerBarModel : LayerBarModels)
+	for (const FAvaStaggerBarElement& BarElement : OutBarElements)
 	{
-		const TViewModelPtr<IOutlinerExtension> OutlinerExtension = InLayerBarModel->GetLinkedOutlinerItem();
-		if (!OutlinerExtension.IsValid())
+		if (!BarElement.OutlinerItem.IsValid())
 		{
 			continue;
 		}
 
-		for (const TViewModelPtr<IOutlinerExtension>& ChildOutlinerExtension : OutlinerExtension.AsModel()->GetDescendantsOfType<IOutlinerExtension>())
+		for (const TViewModelPtr<IOutlinerExtension>& ChildOutlinerExtension : BarElement.OutlinerItem.AsModel()->GetDescendantsOfType<IOutlinerExtension>())
 		{
-			if (SequencerSelection->Outliner.IsSelected(ChildOutlinerExtension))
+			if (InSequencerSelection->Outliner.IsSelected(ChildOutlinerExtension))
 			{
 				return {};
 			}
 		}
 	}
 
-	return LayerBarModels;
-}
-
-TArray<FAvaStaggerTool::FStaggerLayerBarElement> FAvaStaggerTool::GatherSelectionLayerBarElements() const
-{
-	TArray<TViewModelPtr<FLayerBarModel>> LayerBarModels = GatherAllSelectedLayerBars();
-
-	// Compute range and set original frame
-	TArray<FStaggerLayerBarElement> OutElements;
-	OutElements.Append(MoveTemp(LayerBarModels));
-
-	for (FStaggerLayerBarElement& Element : OutElements)
-	{
-		Element.OriginalFrame = Element.LayerBarModel->ComputeRange().GetLowerBoundValue();
-	}
-
-	return OutElements;
+	return OutBarElements;
 }
 
-TArray<FAvaStaggerTool::FStaggerKeyFrameElement> FAvaStaggerTool::GatherSelectionKeyFrameElements() const
+TArray<FAvaStaggerTool::FAvaStaggerKeyElement> FAvaStaggerTool::GatherSelectionKeyElements() const
 {
 	const TSharedPtr<FSequencerSelection> SequencerSelection = GetSequencerSelection();
 	if (!SequencerSelection.IsValid())
@@ -457,7 +430,7 @@ TArray<FAvaStaggerTool::FStaggerKeyFrameElement> FAvaStaggerTool::GatherSelectio
 		return {};
 	}
 
-	TArray<FStaggerKeyFrameElement> OutElements;
+	TArray<FAvaStaggerKeyElement> OutElements;
 
 	// Gather all the selected keyframe handles
 	OutElements.Reserve(SequencerSelection->KeySelection.Num());
@@ -474,7 +447,7 @@ TArray<FAvaStaggerTool::FStaggerKeyFrameElement> FAvaStaggerTool::GatherSelectio
 					FFrameNumber KeyFrameNumber;
 					KeyModel->GetChannel()->GetKeyTime(KeyHandle, KeyFrameNumber);
 
-					OutElements.Add(FStaggerKeyFrameElement(KeyModel, KeyHandle, KeyFrameNumber));
+					OutElements.Add(FAvaStaggerKeyElement(KeyModel, KeyHandle, KeyFrameNumber));
 				}
 			}
 		}
@@ -530,14 +503,14 @@ FFrameNumber FAvaStaggerTool::CalculateLocalCurveOffset(const float InCurveTime,
 	return FMath::RoundToInt32(LocalFrameOffset);
 }
 
-FFrameNumber FAvaStaggerTool::FindNextStaggerLocation(const FStaggerLayerBarElement& InElement
+FFrameNumber FAvaStaggerTool::FindNextStaggerLocation(const FAvaStaggerBarElement& InElement
 	, const int32 InElementIndex
 	, const FFrameNumber InFirstFrame
 	, const FFrameNumber InCurrentFrame) const
 {
 	const FRichCurve* const RichCurve = Settings->ToolOptions.Curve.GetRichCurve();
-	const FFrameNumber ElementSize = InElement.LayerBarModel->ComputeRange().Size<FFrameNumber>();
-	const FFrameNumber Interval = GetInterval(CachedRangeSize, CachedLayerBarCount, ElementSize);
+	const FFrameNumber ElementSize = InElement.Range.Size<FFrameNumber>();
+	const FFrameNumber Interval = GetInterval(CachedRangeSize, CachedBarCount, ElementSize);
 
 	if (Settings->ToolOptions.bUseCurve)
 	{
@@ -546,8 +519,8 @@ FFrameNumber FAvaStaggerTool::FindNextStaggerLocation(const FStaggerLayerBarElem
 			return InFirstFrame;
 		}
 
-		const FFrameNumber FullFrameSpan = CalculateLayerBarFrameSpan(OriginalLayerBarElements);
-		const FFrameNumber ElementFrameSpan = CalculateLayerBarFrameSpan(OriginalLayerBarElements, InElementIndex + 1);
+		const FFrameNumber FullFrameSpan = CalculateBarElementFrameSpan(OriginalBarElements);
+		const FFrameNumber ElementFrameSpan = CalculateBarElementFrameSpan(OriginalBarElements, InElementIndex + 1);
 		const float ElementStartPercentOfRange = static_cast<float>(ElementFrameSpan.Value) / static_cast<float>(FullFrameSpan.Value);
 
 		if (Settings->ToolOptions.Distribution == EAvaSequencerStaggerDistribution::Increment)
@@ -556,7 +529,7 @@ FFrameNumber FAvaStaggerTool::FindNextStaggerLocation(const FStaggerLayerBarElem
 			return InFirstFrame + LocalCurveOffset + Interval;
 		}
 
-		const float CurveStepTime = 1.f / CachedLayerBarCount;
+		const float CurveStepTime = 1.f / CachedBarCount;
 		const float CurveTime = CurveStepTime * InElementIndex;
 
 		const FFrameNumber LocalCurveOffset = CalculateLocalCurveOffset(CurveTime, CachedRange);
@@ -565,27 +538,27 @@ FFrameNumber FAvaStaggerTool::FindNextStaggerLocation(const FStaggerLayerBarElem
 
 	if (Settings->ToolOptions.Distribution == EAvaSequencerStaggerDistribution::Random)
 	{
-		if (OriginalLayerBarElements.IsValidIndex(InElementIndex + 1))
+		if (OriginalBarElements.IsValidIndex(InElementIndex + 1))
 		{
-			const FStaggerLayerBarElement& NextElement = OriginalLayerBarElements[InElementIndex + 1];
-			const FFrameNumber NextElementSize = NextElement.LayerBarModel->ComputeRange().Size<FFrameNumber>();
-			const FFrameNumber RandomInterval = GetInterval(CachedRangeSize, CachedLayerBarCount, NextElementSize);
+			const FAvaStaggerBarElement& NextElement = OriginalBarElements[InElementIndex + 1];
+			const FFrameNumber NextElementSize = NextElement.Range.Size<FFrameNumber>();
+			const FFrameNumber RandomInterval = GetInterval(CachedRangeSize, CachedBarCount, NextElementSize);
 			return CachedRange.GetLowerBoundValue() + RandomInterval + CachedShiftFrames;
 		}
 		return CachedRange.GetLowerBoundValue() + CachedShiftFrames;
 	}
 
-	const FFrameNumber OperationOffset = GetLayerBarOperationOffset(InElement.LayerBarModel);
+	const FFrameNumber OperationOffset = GetBarElementOperationOffset(InElement);
 	return InCurrentFrame + OperationOffset + Interval;
 }
 
-FFrameNumber FAvaStaggerTool::FindNextStaggerLocation(const FStaggerKeyFrameElement& InElement
+FFrameNumber FAvaStaggerTool::FindNextStaggerLocation(const FAvaStaggerKeyElement& InElement
 	, const int32 InElementIndex
 	, const FFrameNumber InFirstFrame
 	, const FFrameNumber InCurrentFrame) const
 {
 	const FRichCurve* const RichCurve = Settings->ToolOptions.Curve.GetRichCurve();
-	const FFrameNumber Interval = GetInterval(CachedRangeSize, CachedKeyframeCount);
+	const FFrameNumber Interval = GetInterval(CachedRangeSize, CachedKeyCount);
 
 	if (Settings->ToolOptions.bUseCurve)
 	{
@@ -594,7 +567,7 @@ FFrameNumber FAvaStaggerTool::FindNextStaggerLocation(const FStaggerKeyFrameElem
 			return InFirstFrame;
 		}
 
-		const float CurveStepTime = 1.f / CachedKeyframeCount;
+		const float CurveStepTime = 1.f / CachedKeyCount;
 		const FFrameNumber NextStaggerFrameNumber = CalculateLocalCurveOffset(CurveStepTime * InElementIndex, CachedRange);
 
 		return InFirstFrame + NextStaggerFrameNumber + Interval;
@@ -620,26 +593,26 @@ void FAvaStaggerTool::Stagger()
 		RandomStream = FRandomStream(Settings->ToolOptions.RandomSeed);
 	}
 
-	if (IsLayerBarSelection())
+	if (IsBarSelection())
 	{
 		// Reset layer bar offsets that were last applied
-		for (FStaggerLayerBarElement& Element : OriginalLayerBarElements)
+		for (FAvaStaggerBarElement& Element : OriginalBarElements)
 		{
-			Element.LayerBarModel->Offset(-Element.LastOffset);
+			Element.Offset(-Element.LastOffset);
 			Element.LastOffset = 0;
 		}
 
-		StaggerLayers();
+		StaggerBarElements();
 	}
-	else if (IsKeyFrameSelection())
+	else if (IsKeySelection())
 	{
 		// Reset key frames to their last location
-		for (FStaggerKeyFrameElement& Element : OriginalKeyFrameElements)
+		for (FAvaStaggerKeyElement& Element : OriginalKeyElements)
 		{
-			SetKeyFrameElementTime(Element, Element.OriginalFrame);
+			SetKeyElementTime(Element, Element.OriginalFrame);
 		}
 
-		StaggerKeyFrames();
+		StaggerKeyElements();
 	}
 }
 
@@ -648,25 +621,25 @@ bool FAvaStaggerTool::IsAutoAppying() const
 	return Settings->bAutoApply;
 }
 
-void FAvaStaggerTool::StaggerLayers()
+void FAvaStaggerTool::StaggerBarElements()
 {
-	if (CachedLayerBarCount <= 1)
+	if (CachedBarCount <= 1)
 	{
 		return;
 	}
 
 	FScopedTransaction Transaction(LOCTEXT("StaggerLayers", "Stagger Layers"));
 
-	const FFrameNumber FirstStaggerPoint = FindFirstLayerBarStaggerPoint();
+	const FFrameNumber FirstStaggerPoint = FindFirstBarStaggerPoint();
 
 	FFrameNumber NextStaggerPoint = FirstStaggerPoint;
 	int32 CurrentElementNum = 0;
 	int32 CurrentGroupNum = 0;
 
-	for (FStaggerLayerBarElement& Element : OriginalLayerBarElements)
+	for (FAvaStaggerBarElement& Element : OriginalBarElements)
 	{
 		Element.LastOffset = NextStaggerPoint - Element.OriginalFrame;
-		Element.LayerBarModel->Offset(Element.LastOffset);
+		Element.Offset(Element.LastOffset);
 
 		++CurrentGroupNum;
 
@@ -682,9 +655,9 @@ void FAvaStaggerTool::StaggerLayers()
 	}
 }
 
-void FAvaStaggerTool::StaggerKeyFrames()
+void FAvaStaggerTool::StaggerKeyElements()
 {
-	if (CachedKeyframeCount <= 1)
+	if (CachedKeyCount <= 1)
 	{
 		return;
 	}
@@ -692,15 +665,15 @@ void FAvaStaggerTool::StaggerKeyFrames()
 	// Stagger all elements
 	FScopedTransaction Transaction(LOCTEXT("StaggerKeyFrames", "Stagger Key Frames"));
 
-	const FFrameNumber FirstStaggerPoint = FindFirstKeyframeStaggerPoint();
+	const FFrameNumber FirstStaggerPoint = FindFirstKeyStaggerPoint();
 
 	FFrameNumber NextStaggerPoint = FirstStaggerPoint;
 	int32 CurrentElementNum = 0;
 	int32 CurrentGroupNum = 0;
 
-	for (FStaggerKeyFrameElement& Element : OriginalKeyFrameElements)
+	for (FAvaStaggerKeyElement& Element : OriginalKeyElements)
 	{
-		SetKeyFrameElementTime(Element, NextStaggerPoint);
+		SetKeyElementTime(Element, NextStaggerPoint);
 
 		++CurrentElementNum;
 		++CurrentGroupNum;
@@ -725,7 +698,7 @@ void FAvaStaggerTool::AlignToPlayhead()
 	FSequencerSelectionAlignmentUtils::AlignSelectionToPlayhead(*Owner.GetSequencer());
 }
 
-void FAvaStaggerTool::SetKeyFrameElementTime(const FStaggerKeyFrameElement& InElement, const FFrameNumber InKeyTime)
+void FAvaStaggerTool::SetKeyElementTime(const FAvaStaggerKeyElement& InElement, const FFrameNumber InKeyTime)
 {
 	if (!InElement.KeyChannelModel.IsValid())
 	{
diff --git a/Engine/Plugins/Experimental/Avalanche/Source/AvalancheSequencer/Private/StaggerTool/AvaStaggerTool.h b/Engine/Plugins/Experimental/Avalanche/Source/AvalancheSequencer/Private/StaggerTool/AvaStaggerTool.h
--- a/Engine/Plugins/Experimental/Avalanche/Source/AvalancheSequencer/Private/StaggerTool/AvaStaggerTool.h
+++ b/Engine/Plugins/Experimental/Avalanche/Source/AvalancheSequencer/Private/StaggerTool/AvaStaggerTool.h
@@ -16,25 +16,25 @@
 class SAvaStaggerTool;
 class SWindow;
 class UAvaStaggerCurve;
+struct FAvaStaggerBarElement;
 
 namespace UE::Sequencer
 {
 	class FChannelModel;
-	class FLayerBarModel;
 }
 
 class FAvaStaggerTool : public FAvaSequencerAction
 {
 public:
 	explicit FAvaStaggerTool(FAvaSequencer& InOwner);
-	virtual ~FAvaStaggerTool();
+	virtual ~FAvaStaggerTool() override;
 
 	TObjectPtr<UAvaSequencerStaggerSettings> GetSettings() const;
 
 	bool HasValidSelection() const;
 
-	bool IsLayerBarSelection() const;
-	bool IsKeyFrameSelection() const;
+	bool IsBarSelection() const;
+	bool IsKeySelection() const;
 
 	int32 GetSelectionCount() const;
 
@@ -50,18 +50,9 @@ public:
 	void CloseToolWindow();
 
 private:
-	struct FStaggerLayerBarElement
+	struct FAvaStaggerKeyElement
 	{
-		explicit FStaggerLayerBarElement(UE::Sequencer::TViewModelPtr<UE::Sequencer::FLayerBarModel> InLayerBarModel);
-
-		UE::Sequencer::TViewModelPtr<UE::Sequencer::FLayerBarModel> LayerBarModel;
-		FFrameNumber OriginalFrame;
-		FFrameNumber LastOffset;
-	};
-
-	struct FStaggerKeyFrameElement
-	{
-		explicit FStaggerKeyFrameElement(UE::Sequencer::TViewModelPtr<UE::Sequencer::FChannelModel> InKeyChannelModel
+		explicit FAvaStaggerKeyElement(UE::Sequencer::TViewModelPtr<UE::Sequencer::FChannelModel> InKeyChannelModel
 			, FKeyHandle InKeyHandle, FFrameNumber InOriginalFrame = 0);
 
 		UE::Sequencer::TViewModelPtr<UE::Sequencer::FChannelModel> KeyChannelModel;
@@ -126,35 +117,33 @@ private:
 
 	FFrameNumber GetInterval(const FFrameNumber InRangeSize, const int32 InElementCount, const FFrameNumber InElementSize = 0) const;
 
-	FFrameNumber GetLayerBarOperationOffset(const UE::Sequencer::TViewModelPtr<UE::Sequencer::FLayerBarModel>& InLayerBarModel) const;
+	FFrameNumber GetBarElementOperationOffset(const FAvaStaggerBarElement& InElement) const;
 
-	static FFrameNumber CalculateLayerBarFrameSpan(const TArray<FStaggerLayerBarElement>& InStaggerElements, const int32 InStopIndex = INDEX_NONE);
+	static FFrameNumber CalculateBarElementFrameSpan(const TArray<FAvaStaggerBarElement>& InElements, const int32 InStopIndex = INDEX_NONE);
 
 	void CacheOriginalElements();
 
-	FFrameNumber FindFirstLayerBarStaggerPoint() const;
-	FFrameNumber FindFirstKeyframeStaggerPoint() const;
-
-	TArray<UE::Sequencer::TViewModelPtr<UE::Sequencer::FLayerBarModel>> GatherAllSelectedLayerBars() const;
+	FFrameNumber FindFirstBarStaggerPoint() const;
+	FFrameNumber FindFirstKeyStaggerPoint() const;
 
-	TArray<FStaggerLayerBarElement> GatherSelectionLayerBarElements() const;
-	TArray<FStaggerKeyFrameElement> GatherSelectionKeyFrameElements() const;
+	TArray<FAvaStaggerBarElement> GatherSelectionBarElements(const TSharedPtr<UE::Sequencer::FSequencerSelection>& InSequencerSelection) const;
+	TArray<FAvaStaggerKeyElement> GatherSelectionKeyElements() const;
 
 	FFrameNumber CalculateLocalCurveOffset(const float InCurveTime, const TRange<FFrameNumber>& InRange) const;
 
-	FFrameNumber FindNextStaggerLocation(const FStaggerLayerBarElement& InElement
+	FFrameNumber FindNextStaggerLocation(const FAvaStaggerBarElement& InElement
 		, const int32 InElementIndex
 		, const FFrameNumber InFirstFrame
 		, const FFrameNumber InCurrentFrame) const;
-	FFrameNumber FindNextStaggerLocation(const FStaggerKeyFrameElement& InElement
+	FFrameNumber FindNextStaggerLocation(const FAvaStaggerKeyElement& InElement
 		, const int32 InElementIndex
 		, const FFrameNumber InFirstFrame
 		, const FFrameNumber InCurrentFrame) const;
 
-	void StaggerLayers();
-    void StaggerKeyFrames();
+	void StaggerBarElements();
+    void StaggerKeyElements();
 
-	static void SetKeyFrameElementTime(const FStaggerKeyFrameElement& InElement, const FFrameNumber InKeyTime);
+	static void SetKeyElementTime(const FAvaStaggerKeyElement& InElement, const FFrameNumber InKeyTime);
 
 	template<typename InElementType>
 	static void FindFirstFrameInTimeline(const TArray<InElementType>& InElements, FFrameNumber& OutFrame)
@@ -177,11 +166,11 @@ private:
 	TStrongObjectPtr<UAvaSequencerStaggerSettings> Settings;
 
 	// Cached values at the start of auto-updating
-	TArray<FStaggerLayerBarElement> OriginalLayerBarElements;
-	TArray<FStaggerKeyFrameElement> OriginalKeyFrameElements;
+	TArray<FAvaStaggerBarElement> OriginalBarElements;
+	TArray<FAvaStaggerKeyElement> OriginalKeyElements;
 
-	int32 CachedLayerBarCount = 0;
-	int32 CachedKeyframeCount = 0;
+	int32 CachedBarCount = 0;
+	int32 CachedKeyCount = 0;
 
 	TRange<FFrameNumber> CachedRange;
 	FFrameNumber CachedRangeSize = 0;
diff --git a/Engine/Plugins/Experimental/Avalanche/Source/AvalancheSequencer/Private/StaggerTool/Customization/AvaStaggerToolSettingsDetailsCustomization.cpp b/Engine/Plugins/Experimental/Avalanche/Source/AvalancheSequencer/Private/StaggerTool/Customization/AvaStaggerToolSettingsDetailsCustomization.cpp
--- a/Engine/Plugins/Experimental/Avalanche/Source/AvalancheSequencer/Private/StaggerTool/Customization/AvaStaggerToolSettingsDetailsCustomization.cpp
+++ b/Engine/Plugins/Experimental/Avalanche/Source/AvalancheSequencer/Private/StaggerTool/Customization/AvaStaggerToolSettingsDetailsCustomization.cpp
@@ -174,7 +174,7 @@ EVisibility FAvaStaggerToolSettingsDetailsCustomization::GetStartPositionPropert
 EVisibility FAvaStaggerToolSettingsDetailsCustomization::GetOperationPointPropertyVisibility() const
 {
 	// Keyframes have no range to use
-	if (WeakTool.IsValid() && !WeakTool.Pin()->IsLayerBarSelection())
+	if (WeakTool.IsValid() && WeakTool.Pin()->IsKeySelection())
 	{
 		return EVisibility::Collapsed;
 	}
diff --git a/Engine/Plugins/Experimental/Avalanche/Source/AvalancheSequencer/Private/StaggerTool/Widgets/SAvaStaggerTool.cpp b/Engine/Plugins/Experimental/Avalanche/Source/AvalancheSequencer/Private/StaggerTool/Widgets/SAvaStaggerTool.cpp
--- a/Engine/Plugins/Experimental/Avalanche/Source/AvalancheSequencer/Private/StaggerTool/Widgets/SAvaStaggerTool.cpp
+++ b/Engine/Plugins/Experimental/Avalanche/Source/AvalancheSequencer/Private/StaggerTool/Widgets/SAvaStaggerTool.cpp
@@ -191,13 +191,13 @@ FText SAvaStaggerTool::GetSelectionText() const
 		return FText::GetEmpty();
 	}
 
-	if (StaggerTool->IsLayerBarSelection())
+	if (StaggerTool->IsBarSelection())
 	{
 		const FText SelectionCountText = FText::FromString(FString::FromInt(StaggerTool->GetSelectionCount()));
 		return FText::Format(LOCTEXT("LayerBarSelectionDetails", "{0} Layer Bars Selected"), SelectionCountText);
 	}
 
-	if (StaggerTool->IsKeyFrameSelection())
+	if (StaggerTool->IsKeySelection())
 	{
 		const FText SelectionCountText = FText::FromString(FString::FromInt(StaggerTool->GetSelectionCount()));
 		return FText::Format(LOCTEXT("KeyFrameSelectionDetails", "{0} Key Frames Selected"), SelectionCountText);
