[MotionDesign] Stagger Tool: Fix warning about FAvaSequencerStaggerOptions::RandomSeed not being initialized

diff --git a/Engine/Plugins/Experimental/Avalanche/Source/AvalancheSequencer/Private/StaggerTool/AvaStaggerToolSettings.h b/Engine/Plugins/Experimental/Avalanche/Source/AvalancheSequencer/Private/StaggerTool/AvaStaggerToolSettings.h
--- a/Engine/Plugins/Experimental/Avalanche/Source/AvalancheSequencer/Private/StaggerTool/AvaStaggerToolSettings.h
+++ b/Engine/Plugins/Experimental/Avalanche/Source/AvalancheSequencer/Private/StaggerTool/AvaStaggerToolSettings.h
@@ -81,7 +81,7 @@ struct FAvaSequencerStaggerOptions
 	EAvaSequencerStaggerDistribution Distribution = EAvaSequencerStaggerDistribution::Default;
 
 	UPROPERTY(EditAnywhere, Category = "Stagger Tool Options", meta = (ClampMin = "0", UIMin = "0", Tooltip = "Seed value for randomization"))
-	int32 RandomSeed = FMath::Rand();
+	int32 RandomSeed = 0;
 
 	UPROPERTY(EditAnywhere, Category = "Stagger Tool Options", meta = (Tooltip = "Time range to perform the operation on"))
 	EAvaSequencerStaggerRange Range = EAvaSequencerStaggerRange::Default;
@@ -122,6 +122,8 @@ class UAvaSequencerStaggerSettings : public UObject
 public:
 	UAvaSequencerStaggerSettings()
 	{
+		ToolOptions.RandomSeed = FMath::Rand();
+
 		ResetToolOptions();
 	}
 
