MotionDesign : InteractiveTools
- Fix temp actor was not flagged as preview when spawning through ITF
- Fix cancel cloner spawn in ITF would still spawn default clone
- Fix attaching cloner actors outliner action would not spawn a cloner and move selected actors below it

diff --git a/Engine/Plugins/Experimental/Avalanche/Source/AvalancheInteractiveTools/Private/Tools/AvaInteractiveToolsStaticMeshActorTool.cpp b/Engine/Plugins/Experimental/Avalanche/Source/AvalancheInteractiveTools/Private/Tools/AvaInteractiveToolsStaticMeshActorTool.cpp
--- a/Engine/Plugins/Experimental/Avalanche/Source/AvalancheInteractiveTools/Private/Tools/AvaInteractiveToolsStaticMeshActorTool.cpp
+++ b/Engine/Plugins/Experimental/Avalanche/Source/AvalancheInteractiveTools/Private/Tools/AvaInteractiveToolsStaticMeshActorTool.cpp
@@ -33,6 +33,7 @@ AActor* UAvaInteractiveToolsStaticMeshActorTool::SpawnActor(TSubclassOf<AActor>
 	if (AStaticMeshActor* StaticMeshActor = Cast<AStaticMeshActor>(Actor))
 	{
 		StaticMeshActor->GetStaticMeshComponent()->SetStaticMesh(StaticMesh);
+		StaticMeshActor->SetMobility(EComponentMobility::Type::Movable);
 	}
 
 	return Actor;
diff --git a/Engine/Plugins/Experimental/Avalanche/Source/AvalancheInteractiveTools/Private/Tools/AvaInteractiveToolsToolBase.cpp b/Engine/Plugins/Experimental/Avalanche/Source/AvalancheInteractiveTools/Private/Tools/AvaInteractiveToolsToolBase.cpp
--- a/Engine/Plugins/Experimental/Avalanche/Source/AvalancheInteractiveTools/Private/Tools/AvaInteractiveToolsToolBase.cpp
+++ b/Engine/Plugins/Experimental/Avalanche/Source/AvalancheInteractiveTools/Private/Tools/AvaInteractiveToolsToolBase.cpp
@@ -345,6 +345,9 @@ AActor* UAvaInteractiveToolsToolBase::SpawnActor(TSubclassOf<AActor> InActorClas
 		NewActor = World->SpawnActor<AActor>(InActorClass, SpawnLocation, SpawnRotation, SpawnParams);
 	}
 
+	checkf(!!NewActor, TEXT("Motion Design Interactive Tool : Failed to spawn actor of class %s"), *InActorClass->GetName())
+	NewActor->bIsEditorPreviewActor = bInPreview;
+
 	if (bUseOverride)
 	{
 		FActorLabelUtilities::SetActorLabelUnique(NewActor, *InActorLabelOverride);
