MotionDesign : Modifiers
- Added ensure when restoring materials in mask modifier to catch refresh material calls during material handle iteration

diff --git a/Engine/Plugins/Experimental/Avalanche/Source/AvalancheMask/Private/Handling/AvaMaskActorMaterialCollectionHandle.cpp b/Engine/Plugins/Experimental/Avalanche/Source/AvalancheMask/Private/Handling/AvaMaskActorMaterialCollectionHandle.cpp
--- a/Engine/Plugins/Experimental/Avalanche/Source/AvalancheMask/Private/Handling/AvaMaskActorMaterialCollectionHandle.cpp
+++ b/Engine/Plugins/Experimental/Avalanche/Source/AvalancheMask/Private/Handling/AvaMaskActorMaterialCollectionHandle.cpp
@@ -242,7 +242,8 @@ void FAvaMaskActorMaterialCollectionHandle::ForEachMaterialHandle(TFunctionRef<b
 	{
 		// Effectively refreshes material handles if necessary
 		GetMaterialHandles();
-		
+
+		TGuardValue<bool> IteratingGuard(bIsIteratingMaterials, true);
 		for (TPair<FSoftComponentReference, TArray<TSharedPtr<IAvaMaskMaterialHandle>>>& ComponentMaterialHandles : MaterialHandles)
 		{
 			for (int32 SlotIdx = 0; SlotIdx < ComponentMaterialHandles.Value.Num(); ++SlotIdx)
@@ -298,6 +299,7 @@ void FAvaMaskActorMaterialCollectionHandle::MapEachMaterialHandle(TFunctionRef<T
 		TBitArray<> SetFlags;
 		SetFlags.Init(false, MappedMaterials.Num());
 
+		TGuardValue<bool> IteratingGuard(bIsIteratingMaterials, true);
 		for (TPair<FSoftComponentReference, TArray<TSharedPtr<IAvaMaskMaterialHandle>>>& ComponentMaterialHandles : MaterialHandles)
 		{
 			for (int32 SlotIdx = 0; SlotIdx < ComponentMaterialHandles.Value.Num(); ++SlotIdx)
@@ -483,6 +485,8 @@ void FAvaMaskActorMaterialCollectionHandle::RefreshComponents()
 
 bool FAvaMaskActorMaterialCollectionHandle::RefreshMaterials()
 {
+	ensureMsgf(!bIsIteratingMaterials, TEXT("Refresh should not happen when iterating over material handles"));
+
 	const int32 ProbableNumMaterials = FMath::Max(1, GetNumMaterials());
 
 	OccupiedMaterialSlots.Reset();
diff --git a/Engine/Plugins/Experimental/Avalanche/Source/AvalancheMask/Private/Handling/AvaMaskActorMaterialCollectionHandle.h b/Engine/Plugins/Experimental/Avalanche/Source/AvalancheMask/Private/Handling/AvaMaskActorMaterialCollectionHandle.h
--- a/Engine/Plugins/Experimental/Avalanche/Source/AvalancheMask/Private/Handling/AvaMaskActorMaterialCollectionHandle.h
+++ b/Engine/Plugins/Experimental/Avalanche/Source/AvalancheMask/Private/Handling/AvaMaskActorMaterialCollectionHandle.h
@@ -117,4 +117,5 @@ protected:
 	TMap<FSoftComponentReference, TArray<TWeakObjectPtr<UMaterialInterface>>> WeakMaterials;
 	TMap<FSoftComponentReference, TArray<TSharedPtr<IAvaMaskMaterialHandle>>> MaterialHandles;
 	TMap<FSoftComponentReference, TUniquePtr<FAvaMaskMaterialAssignmentObserver>> MaterialAssignmentObservers;
+	bool bIsIteratingMaterials = false;
 };
