[Motion Design] Rundown: collecting cook dependencies from page values.

diff --git a/Engine/Plugins/Experimental/Avalanche/Source/AvalancheMedia/Private/Playable/AvaPlayableRemoteControlValues.cpp b/Engine/Plugins/Experimental/Avalanche/Source/AvalancheMedia/Private/Playable/AvaPlayableRemoteControlValues.cpp
--- a/Engine/Plugins/Experimental/Avalanche/Source/AvalancheMedia/Private/Playable/AvaPlayableRemoteControlValues.cpp
+++ b/Engine/Plugins/Experimental/Avalanche/Source/AvalancheMedia/Private/Playable/AvaPlayableRemoteControlValues.cpp
@@ -3,10 +3,14 @@
 #include "Playable/AvaPlayableRemoteControlValues.h"
 
 #include "AvaMediaSerializationUtils.h"
+#include "Dom/JsonObject.h"
+#include "Misc/PackageName.h"
 #include "Playable/AvaPlayableRemoteControl.h"
 #include "Playable/AvaPlayableRemoteControlValuesPrivate.h"
 #include "RCVirtualProperty.h"
 #include "Serialization/CustomVersion.h"
+#include "Serialization/JsonReader.h"
+#include "Serialization/JsonSerializer.h"
 
 namespace UE::AvaPlayableRemoteControlValues::Private
 {
@@ -131,6 +135,46 @@ namespace UE::AvaPlayableRemoteControlValues::Private
 	{
 		return bInModified ? InModifiedChanges : EAvaPlayableRemoteControlChanges::None;
 	}
+
+	void CollectReferencedAssetPaths(const TSharedPtr<FJsonValue>& InJsonValue, TSet<FSoftObjectPath>& OutReferencedPaths)
+	{
+		if (!InJsonValue.IsValid())
+		{
+			return;
+		}
+
+		switch (InJsonValue->Type)
+		{
+		case EJson::String:
+			{
+				// See if that is a valid package or asset path.
+				const FString ValueAsString = InJsonValue->AsString();
+				if ((FPackageName::IsValidTextForLongPackageName(ValueAsString) || FPackageName::IsValidObjectPath(ValueAsString))
+					&& FPackageName::DoesPackageExist(ValueAsString))
+				{
+					OutReferencedPaths.Add(FSoftObjectPath(*ValueAsString));
+				}
+			}
+			break;
+		case EJson::Array:
+			for(const TSharedPtr<FJsonValue>& JsonValue : InJsonValue->AsArray())
+			{
+				CollectReferencedAssetPaths(JsonValue, OutReferencedPaths);
+			}
+			break;
+		case EJson::Object:
+			if (const TSharedPtr<FJsonObject>& JsonObject = InJsonValue->AsObject())
+			{
+				for (const TPair<FString, TSharedPtr<FJsonValue>>& JsonValue : JsonObject->Values)
+				{
+					CollectReferencedAssetPaths(JsonValue.Value, OutReferencedPaths);
+				}
+			}
+			break;
+		default:
+			break;	// don't care about number and bool.
+		}
+	}
 }
 
 const FGuid FAvaPlayableRemoteControlValueCustomVersion::Key(0x85218F83, 0xEDF141CA, 0x800EF947, 0x2F14CB06);
@@ -442,3 +486,24 @@ const FAvaPlayableRemoteControlValues& FAvaPlayableRemoteControlValues::GetDefau
 	const static FAvaPlayableRemoteControlValues Empty;
 	return Empty;
 }
+
+void FAvaPlayableRemoteControlValues::CollectReferencedAssetPaths(const TMap<FGuid, FAvaPlayableRemoteControlValue>& InValues,TSet<FSoftObjectPath>& OutReferencedPaths)
+{
+	// Try to extract a package reference from the value
+	for (const TPair<FGuid, FAvaPlayableRemoteControlValue>& Value : InValues)
+	{
+		TSharedPtr<FJsonObject> ValueObject;
+		if (!FJsonSerializer::Deserialize(TJsonReaderFactory<>::Create(Value.Value.Value), ValueObject) || !ValueObject.IsValid())
+		{
+			UE_LOG(LogAvaPlayableRemoteControl, Warning, TEXT("CollectReferencedPaths: Unable to parse json '%s'"), *Value.Value.Value);
+			continue;
+		}
+
+		// For now, we don't care about the field name.
+		// Just trying to find values that could possibly be an existing asset reference.
+		for (const TPair<FString, TSharedPtr<FJsonValue>>& JsonValue : ValueObject->Values)
+		{
+			UE::AvaPlayableRemoteControlValues::Private::CollectReferencedAssetPaths(JsonValue.Value, OutReferencedPaths);
+		}
+	}
+}
diff --git a/Engine/Plugins/Experimental/Avalanche/Source/AvalancheMedia/Private/Rundown/AvaRundown.cpp b/Engine/Plugins/Experimental/Avalanche/Source/AvalancheMedia/Private/Rundown/AvaRundown.cpp
--- a/Engine/Plugins/Experimental/Avalanche/Source/AvalancheMedia/Private/Rundown/AvaRundown.cpp
+++ b/Engine/Plugins/Experimental/Avalanche/Source/AvalancheMedia/Private/Rundown/AvaRundown.cpp
@@ -17,6 +17,7 @@
 #include "Rundown/Transition/AvaRundownPageTransitionBuilder.h"
 
 #if WITH_EDITOR
+#include "Cooker/CookEvents.h"
 #include "Editor.h"
 #endif
 
@@ -126,6 +127,20 @@ namespace UE::AvaMedia::Rundown::Private
 			Page.PostLoad();
 		}
 	}
+
+	/**
+	 * Collect the referenced asset paths from page RC values.
+	 */
+	void CollectReferencedAssetPaths(const TArray<FAvaRundownPage>& InPages, TSet<FSoftObjectPath>& OutReferencedPaths)
+	{
+		for (const FAvaRundownPage& Page : InPages)
+		{
+			// Collect asset package references from values
+			const FAvaPlayableRemoteControlValues& Values = Page.GetRemoteControlValues();
+			FAvaPlayableRemoteControlValues::CollectReferencedAssetPaths(Values.EntityValues, OutReferencedPaths);
+			FAvaPlayableRemoteControlValues::CollectReferencedAssetPaths(Values.ControllerValues, OutReferencedPaths);
+		}
+	}
 }
 
 void UAvaRundown::BeginDestroy()
@@ -234,6 +249,37 @@ void UAvaRundown::PostEditUndo()
 		PreUndoBackup.Reset();
 	}
 }
+
+void UAvaRundown::OnCookEvent(UE::Cook::ECookEvent InCookEvent, UE::Cook::FCookEventContext& InCookContext)
+{
+	Super::OnCookEvent(InCookEvent, InCookContext);
+
+	if (InCookEvent == UE::Cook::ECookEvent::PlatformCookDependencies)
+	{
+		TSet<FSoftObjectPath> ReferencedPaths;
+		UE::AvaMedia::Rundown::Private::CollectReferencedAssetPaths(TemplatePages.Pages, ReferencedPaths);
+		UE::AvaMedia::Rundown::Private::CollectReferencedAssetPaths(InstancedPages.Pages, ReferencedPaths);
+
+		TSet<FName> ReferencedPackages;
+		Algo::Transform(ReferencedPaths, ReferencedPackages, [](const FSoftObjectPath& InPath){ return InPath.GetLongPackageFName(); });
+		
+		// Let the referenced packages be picked up by StormSync.
+		for (const FName PackageName : ReferencedPackages)
+		{
+			InCookContext.AddLoadBuildDependency(UE::Cook::FCookDependency::Package(PackageName));
+		}
+
+		if (InCookContext.IsCooking())
+		{
+			// Seems to be necessary for references to be picked up when cooking.
+			for (const FSoftObjectPath& ObjectPath : ReferencedPaths)
+			{
+				InCookContext.AddRuntimeDependency(ObjectPath);
+			}
+		}
+	}
+}
+
 #endif
 
 bool UAvaRundown::IsEmpty() const
diff --git a/Engine/Plugins/Experimental/Avalanche/Source/AvalancheMedia/Public/Playable/AvaPlayableRemoteControlValues.h b/Engine/Plugins/Experimental/Avalanche/Source/AvalancheMedia/Public/Playable/AvaPlayableRemoteControlValues.h
--- a/Engine/Plugins/Experimental/Avalanche/Source/AvalancheMedia/Public/Playable/AvaPlayableRemoteControlValues.h
+++ b/Engine/Plugins/Experimental/Avalanche/Source/AvalancheMedia/Public/Playable/AvaPlayableRemoteControlValues.h
@@ -193,4 +193,9 @@ struct AVALANCHEMEDIA_API FAvaPlayableRemoteControlValues
 	static bool HasIdCollisions(const TMap<FGuid, FAvaPlayableRemoteControlValue>& InValues, const TMap<FGuid, FAvaPlayableRemoteControlValue>& InOtherValues);
 	
 	static const FAvaPlayableRemoteControlValues& GetDefaultEmpty();
+
+	/**
+	 * Collect the referenced asset paths from the given values.
+	 */
+	static void CollectReferencedAssetPaths(const TMap<FGuid, FAvaPlayableRemoteControlValue>& InValues,TSet<FSoftObjectPath>& OutReferencedPaths);
 };
diff --git a/Engine/Plugins/Experimental/Avalanche/Source/AvalancheMedia/Public/Rundown/AvaRundown.h b/Engine/Plugins/Experimental/Avalanche/Source/AvalancheMedia/Public/Rundown/AvaRundown.h
--- a/Engine/Plugins/Experimental/Avalanche/Source/AvalancheMedia/Public/Rundown/AvaRundown.h
+++ b/Engine/Plugins/Experimental/Avalanche/Source/AvalancheMedia/Public/Rundown/AvaRundown.h
@@ -335,6 +335,7 @@ public:
 #if WITH_EDITOR
 	virtual void PreEditUndo() override;
 	virtual void PostEditUndo() override;
+	virtual void OnCookEvent(UE::Cook::ECookEvent InCookEvent, UE::Cook::FCookEventContext& InCookContext) override;
 #endif
 	//~ End UObject Interface
 
