Motion Design: show properties to only show/hide controlled properties. non-controlled properties are expected to be user-set (like controllers) so should always show.

diff --git a/Engine/Plugins/Experimental/Avalanche/Source/AvalancheMediaEditor/Private/Rundown/AvaRundownEditorSettings.h b/Engine/Plugins/Experimental/Avalanche/Source/AvalancheMediaEditor/Private/Rundown/AvaRundownEditorSettings.h
--- a/Engine/Plugins/Experimental/Avalanche/Source/AvalancheMediaEditor/Private/Rundown/AvaRundownEditorSettings.h
+++ b/Engine/Plugins/Experimental/Avalanche/Source/AvalancheMediaEditor/Private/Rundown/AvaRundownEditorSettings.h
@@ -34,7 +34,8 @@ public:
 	UPROPERTY(Config, EditAnywhere, Category="Preview")
 	bool bPreviewCheckerBoard = false;
 
-	UPROPERTY(Config, EditAnywhere, Category = "Page Details", meta = (DisplayName = "Show RC Properties"))
+	/** Whether to show controller properties in Rundown Page Details */
+	UPROPERTY(Config, EditAnywhere, Category = "Page Details", meta = (DisplayName = "Show RC Controlled Properties"))
 	bool bPageDetailsShowProperties = false;
 
 	/** Current macro collection used by the rundown editor. */
diff --git a/Engine/Plugins/Experimental/Avalanche/Source/AvalancheMediaEditor/Private/Rundown/DetailsView/RemoteControl/Properties/SAvaRundownPageRemoteControlProps.cpp b/Engine/Plugins/Experimental/Avalanche/Source/AvalancheMediaEditor/Private/Rundown/DetailsView/RemoteControl/Properties/SAvaRundownPageRemoteControlProps.cpp
--- a/Engine/Plugins/Experimental/Avalanche/Source/AvalancheMediaEditor/Private/Rundown/DetailsView/RemoteControl/Properties/SAvaRundownPageRemoteControlProps.cpp
+++ b/Engine/Plugins/Experimental/Avalanche/Source/AvalancheMediaEditor/Private/Rundown/DetailsView/RemoteControl/Properties/SAvaRundownPageRemoteControlProps.cpp
@@ -9,6 +9,7 @@
 #include "RemoteControlEntity.h"
 #include "RemoteControlPreset.h"
 #include "Rundown/AvaRundownEditor.h"
+#include "Rundown/AvaRundownEditorSettings.h"
 #include "Rundown/AvaRundownEditorUtils.h"
 #include "Rundown/AvaRundownManagedInstanceCache.h"
 #include "Rundown/AvaRundownPage.h"
@@ -82,6 +83,9 @@ void SAvaRundownPageRemoteControlProps::Construct(const FArguments& InArgs, TSha
 
 	ContextMenu = MakeShared<FAvaRundownPagePropertyContextMenu>(CommandList);
 
+	const UAvaRundownEditorSettings* RundownEditorSettings = UAvaRundownEditorSettings::Get();
+	bShowControlledProperties = !RundownEditorSettings || RundownEditorSettings->bPageDetailsShowProperties;
+
 	TSharedRef<SHeaderRow> HeaderRow =
 		SNew(SHeaderRow)
 		+ SHeaderRow::Column(PropertyColumnName)
@@ -320,7 +324,11 @@ void SAvaRundownPageRemoteControlProps::Refresh(const TArray<int32>& InSelectedP
 					}
 
 					const bool bEntityControlled = ManagedInstance->GetDefaultRemoteControlValues().EntitiesControlledByController.Contains(FieldEntity->GetId());
-					ItemsToAdd.Add({ FieldEntity, bEntityControlled });
+
+					if (!bEntityControlled || bShowControlledProperties)
+					{
+						ItemsToAdd.Add({ FieldEntity, bEntityControlled });
+					}
 				}
 			}
 
@@ -386,9 +394,21 @@ void SAvaRundownPageRemoteControlProps::Refresh(const TArray<int32>& InSelectedP
 				FieldItems.Add(MoveTemp(Item));
 			}
 		}
+
+		FieldItems.StableSort(
+			[](const TSharedPtr<FAvaRundownRCFieldItem>& InItemA, const TSharedPtr<FAvaRundownRCFieldItem>& InItemB)
+			{
+				return InItemA->IsEntityControlled() < InItemB->IsEntityControlled();
+			});
 	}
 }
 
+void SAvaRundownPageRemoteControlProps::SetShowControlledProperties(bool bInShowControlledProperties)
+{
+	bShowControlledProperties = bInShowControlledProperties;
+	Refresh({ActivePageId});
+}
+
 void SAvaRundownPageRemoteControlProps::Tick(const FGeometry& InAllottedGeometry, const double InCurrentTime, const float InDeltaTime)
 {
 	// This tick happens before Tree View tick
diff --git a/Engine/Plugins/Experimental/Avalanche/Source/AvalancheMediaEditor/Private/Rundown/DetailsView/RemoteControl/Properties/SAvaRundownPageRemoteControlProps.h b/Engine/Plugins/Experimental/Avalanche/Source/AvalancheMediaEditor/Private/Rundown/DetailsView/RemoteControl/Properties/SAvaRundownPageRemoteControlProps.h
--- a/Engine/Plugins/Experimental/Avalanche/Source/AvalancheMediaEditor/Private/Rundown/DetailsView/RemoteControl/Properties/SAvaRundownPageRemoteControlProps.h
+++ b/Engine/Plugins/Experimental/Avalanche/Source/AvalancheMediaEditor/Private/Rundown/DetailsView/RemoteControl/Properties/SAvaRundownPageRemoteControlProps.h
@@ -75,6 +75,9 @@ public:
 	/** Refreshes the content of this widget. */
 	void Refresh(const TArray<int32>& InSelectedPageIds);
 
+	/** Sets the visibility on the properties that are set by controllers */
+	void SetShowControlledProperties(bool bInShowControlledProperties);
+
 	//~ Begin SWidget
 	virtual void Tick(const FGeometry& InAllottedGeometry, const double InCurrentTime, const float InDeltaTime) override;
 	//~ End SWidget
@@ -167,4 +170,7 @@ private:
 	friend class FAvaRundownPageRCPropsNotifyHookImpl;
 
 	bool bRefreshRequested = false;
+
+	/** Whether the properties that are set by controllers are shown */
+	bool bShowControlledProperties = false;
 };
\ No newline at end of file
diff --git a/Engine/Plugins/Experimental/Avalanche/Source/AvalancheMediaEditor/Private/Rundown/DetailsView/RemoteControl/Properties/SAvaRundownRCPropertyItemRow.cpp b/Engine/Plugins/Experimental/Avalanche/Source/AvalancheMediaEditor/Private/Rundown/DetailsView/RemoteControl/Properties/SAvaRundownRCPropertyItemRow.cpp
--- a/Engine/Plugins/Experimental/Avalanche/Source/AvalancheMediaEditor/Private/Rundown/DetailsView/RemoteControl/Properties/SAvaRundownRCPropertyItemRow.cpp
+++ b/Engine/Plugins/Experimental/Avalanche/Source/AvalancheMediaEditor/Private/Rundown/DetailsView/RemoteControl/Properties/SAvaRundownRCPropertyItemRow.cpp
@@ -139,6 +139,7 @@ TSharedRef<SWidget> SAvaRundownRCPropertyItemRow::CreateValue()
 
 	if (FieldItem->IsEntityControlled())
 	{
+		ValueRow->SetEnabled(false);
 		ValueRow->AddSlot()
 			.VAlign(EVerticalAlignment::VAlign_Center)
 			.Padding(3.f, 0.f, 0.f, 0.f)
diff --git a/Engine/Plugins/Experimental/Avalanche/Source/AvalancheMediaEditor/Private/Rundown/DetailsView/SAvaRundownPageDetails.cpp b/Engine/Plugins/Experimental/Avalanche/Source/AvalancheMediaEditor/Private/Rundown/DetailsView/SAvaRundownPageDetails.cpp
--- a/Engine/Plugins/Experimental/Avalanche/Source/AvalancheMediaEditor/Private/Rundown/DetailsView/SAvaRundownPageDetails.cpp
+++ b/Engine/Plugins/Experimental/Avalanche/Source/AvalancheMediaEditor/Private/Rundown/DetailsView/SAvaRundownPageDetails.cpp
@@ -164,7 +164,6 @@ void SAvaRundownPageDetails::Construct(const FArguments& InArgs, const TSharedPt
 			.AutoHeight()
 			[
 				SAssignNew(RemoteControlProps, SAvaRundownPageRemoteControlProps, InRundownEditor)
-				.Visibility(UE::AvaMedia::RundownEditor::Private::ShouldPageDetailsShowProperties() ? EVisibility::Visible : EVisibility::Collapsed)
 			]
 		]
 	];
@@ -232,17 +231,8 @@ FReply SAvaRundownPageDetails::ToggleExposedPropertiesVisibility()
 	{
 		RundownEditorSettings->bPageDetailsShowProperties = !RundownEditorSettings->bPageDetailsShowProperties;
 		RundownEditorSettings->SaveConfig();
-
-		if (RundownEditorSettings->bPageDetailsShowProperties)
-		{
-			RemoteControlProps->SetVisibility(EVisibility::SelfHitTestInvisible);
-		}
-		else
-		{
-			RemoteControlProps->SetVisibility(EVisibility::Collapsed);
-		}
+		RemoteControlProps->SetShowControlledProperties(RundownEditorSettings->bPageDetailsShowProperties);
 	}
-
 	return FReply::Handled();
 }
 
