Motion Design: fix outliner crash

diff --git a/Engine/Plugins/Experimental/Avalanche/Source/AvalancheOutliner/Private/AvaOutlinerView.cpp b/Engine/Plugins/Experimental/Avalanche/Source/AvalancheOutliner/Private/AvaOutlinerView.cpp
--- a/Engine/Plugins/Experimental/Avalanche/Source/AvalancheOutliner/Private/AvaOutlinerView.cpp
+++ b/Engine/Plugins/Experimental/Avalanche/Source/AvalancheOutliner/Private/AvaOutlinerView.cpp
@@ -592,7 +592,7 @@ void FAvaOutlinerView::ClearRootItemSource()
 	{
 		if (TSharedPtr<SAvaOutlinerTreeView> TreeView = OutlinerWidget->GetTreeView())
 		{
-			TreeView->RebuildList();	
+			TreeView->ClearItems();
 		}
 	}
 }
diff --git a/Engine/Plugins/Experimental/Avalanche/Source/AvalancheOutliner/Private/Slate/SAvaOutlinerTreeView.cpp b/Engine/Plugins/Experimental/Avalanche/Source/AvalancheOutliner/Private/Slate/SAvaOutlinerTreeView.cpp
--- a/Engine/Plugins/Experimental/Avalanche/Source/AvalancheOutliner/Private/Slate/SAvaOutlinerTreeView.cpp
+++ b/Engine/Plugins/Experimental/Avalanche/Source/AvalancheOutliner/Private/Slate/SAvaOutlinerTreeView.cpp
@@ -20,6 +20,13 @@ int32 SAvaOutlinerTreeView::GetItemIndex(FAvaOutlinerItemPtr Item) const
 	return INDEX_NONE;
 }
 
+void SAvaOutlinerTreeView::ClearItems()
+{
+	ClearExpandedItems();
+	PreviousSelectedItems.Reset();
+	RebuildList();
+}
+
 void SAvaOutlinerTreeView::FocusOnItem(const FAvaOutlinerItemPtr& InItem)
 {
 	SelectorItem = InItem;
diff --git a/Engine/Plugins/Experimental/Avalanche/Source/AvalancheOutliner/Private/Slate/SAvaOutlinerTreeView.h b/Engine/Plugins/Experimental/Avalanche/Source/AvalancheOutliner/Private/Slate/SAvaOutlinerTreeView.h
--- a/Engine/Plugins/Experimental/Avalanche/Source/AvalancheOutliner/Private/Slate/SAvaOutlinerTreeView.h
+++ b/Engine/Plugins/Experimental/Avalanche/Source/AvalancheOutliner/Private/Slate/SAvaOutlinerTreeView.h
@@ -18,6 +18,8 @@ public:
 
 	int32 GetItemIndex(FAvaOutlinerItemPtr Item) const;
 
+	void ClearItems();
+
 	void FocusOnItem(const FAvaOutlinerItemPtr& InItem);
 
 	void ScrollItemIntoView(const FAvaOutlinerItemPtr& InItem);
