MotionDesign : InteractiveTools
- Added button to toggle viewport toolbar visibility
- Added detailed tooltip for cloner and effector interactive tool

diff --git a/Engine/Plugins/Experimental/Avalanche/Source/AvalancheEffectorsEditor/Private/AvaEffectorsEditorCommands.cpp b/Engine/Plugins/Experimental/Avalanche/Source/AvalancheEffectorsEditor/Private/AvaEffectorsEditorCommands.cpp
--- a/Engine/Plugins/Experimental/Avalanche/Source/AvalancheEffectorsEditor/Private/AvaEffectorsEditorCommands.cpp
+++ b/Engine/Plugins/Experimental/Avalanche/Source/AvalancheEffectorsEditor/Private/AvaEffectorsEditorCommands.cpp
@@ -32,7 +32,7 @@ void FAvaEffectorsEditorCommands::RegisterCommands()
 			TypeCommand,
 			TypeName,
 			FText::FromName(TypeName),
-			LOCTEXT("EffectorCommandTooltip", "Create a Effector Actor in the viewport."),
+			FText::Format(LOCTEXT("EffectorCommandTooltip", "Create a {0} Effector Actor in the viewport."), FText::FromName(TypeName)),
 			FSlateIconFinder::FindIcon(TEXT("AvaEffectorsEditor.Tool_Actor_Effector")),
 			EUserInterfaceActionType::ToggleButton,
 			FInputChord()
@@ -51,7 +51,7 @@ void FAvaEffectorsEditorCommands::RegisterCommands()
 			LayoutCommand,
 			LayoutName,
 			FText::FromName(LayoutName),
-			LOCTEXT("ClonerCommandTooltip", "Create a Cloner Actor in the viewport."),
+			FText::Format(LOCTEXT("ClonerCommandTooltip", "Create a {0} Cloner Actor in the viewport."), FText::FromName(LayoutName)),
 			FSlateIconFinder::FindIcon(TEXT("AvaEffectorsEditor.Tool_Actor_Cloner")),
 			EUserInterfaceActionType::ToggleButton,
 			FInputChord()
diff --git a/Engine/Plugins/Experimental/Avalanche/Source/AvalancheInteractiveTools/Private/AvaInteractiveToolsCommands.cpp b/Engine/Plugins/Experimental/Avalanche/Source/AvalancheInteractiveTools/Private/AvaInteractiveToolsCommands.cpp
--- a/Engine/Plugins/Experimental/Avalanche/Source/AvalancheInteractiveTools/Private/AvaInteractiveToolsCommands.cpp
+++ b/Engine/Plugins/Experimental/Avalanche/Source/AvalancheInteractiveTools/Private/AvaInteractiveToolsCommands.cpp
@@ -5,6 +5,16 @@
 
 #define LOCTEXT_NAMESPACE "AvaInteractiveToolsCommands"
 
+const FAvaInteractiveToolsCommands* FAvaInteractiveToolsCommands::GetExternal()
+{
+	if (IsRegistered())
+	{
+		return &FAvaInteractiveToolsCommands::Get();
+	}
+
+	return nullptr;
+}
+
 FAvaInteractiveToolsCommands::FAvaInteractiveToolsCommands()
 	: TCommands<FAvaInteractiveToolsCommands>(
 		TEXT("AvaInteractiveTools")
diff --git a/Engine/Plugins/Experimental/Avalanche/Source/AvalancheInteractiveTools/Private/AvaInteractiveToolsCommands.h b/Engine/Plugins/Experimental/Avalanche/Source/AvalancheInteractiveTools/Private/AvaInteractiveToolsCommands.h
deleted file mode 100644
diff --git a/Engine/Plugins/Experimental/Avalanche/Source/AvalancheInteractiveTools/Private/AvaInteractiveToolsSettings.cpp b/Engine/Plugins/Experimental/Avalanche/Source/AvalancheInteractiveTools/Private/AvaInteractiveToolsSettings.cpp
--- a/Engine/Plugins/Experimental/Avalanche/Source/AvalancheInteractiveTools/Private/AvaInteractiveToolsSettings.cpp
+++ b/Engine/Plugins/Experimental/Avalanche/Source/AvalancheInteractiveTools/Private/AvaInteractiveToolsSettings.cpp
@@ -2,6 +2,9 @@
 
 #include "AvaInteractiveToolsSettings.h"
 
+#include "EditorModeManager.h"
+#include "EdMode/AvaInteractiveToolsEdModeToolkit.h"
+#include "IAvalancheInteractiveToolsModule.h"
 #include "ISettingsModule.h"
 #include "Modules/ModuleManager.h"
 
@@ -23,3 +26,37 @@ void UAvaInteractiveToolsSettings::OpenEditorSettingsWindow() const
 		SettingsModule->ShowViewer(GetContainerName(), GetCategoryName(), GetSectionName());
 	}
 }
+
+bool UAvaInteractiveToolsSettings::IsViewportToolbarSupported() const
+{
+	const FEditorModeTools& EditorModeTools = GLevelEditorModeTools();
+	return EditorModeTools.IsModeActive(IAvalancheInteractiveToolsModule::EM_AvaInteractiveToolsEdModeId);
+}
+
+void UAvaInteractiveToolsSettings::SetViewportToolbarVisible(bool bInVisible) const
+{
+	const FEditorModeTools& EditorModeTools = GLevelEditorModeTools();
+	if (UEdMode* ITEdMode = EditorModeTools.GetActiveScriptableMode(IAvalancheInteractiveToolsModule::EM_AvaInteractiveToolsEdModeId))
+	{
+		if (const TSharedPtr<FModeToolkit> ModeToolkit = ITEdMode->GetToolkit().Pin())
+		{
+			FAvaInteractiveToolsEdModeToolkit* ITMode = static_cast<FAvaInteractiveToolsEdModeToolkit*>(ModeToolkit.Get());
+			ITMode->SetViewportToolbarVisibility(bInVisible);
+		}
+	}
+}
+
+bool UAvaInteractiveToolsSettings::GetViewportToolbarVisible() const
+{
+	const FEditorModeTools& EditorModeTools = GLevelEditorModeTools();
+	if (UEdMode* ITEdMode = EditorModeTools.GetActiveScriptableMode(IAvalancheInteractiveToolsModule::EM_AvaInteractiveToolsEdModeId))
+	{
+		if (const TSharedPtr<FModeToolkit> ModeToolkit = ITEdMode->GetToolkit().Pin())
+		{
+			FAvaInteractiveToolsEdModeToolkit* ITMode = static_cast<FAvaInteractiveToolsEdModeToolkit*>(ModeToolkit.Get());
+			return ITMode->GetViewportToolbarVisibility();
+		}
+	}
+
+	return false;
+}
diff --git a/Engine/Plugins/Experimental/Avalanche/Source/AvalancheInteractiveTools/Public/AvaInteractiveToolsCommands.h b/Engine/Plugins/Experimental/Avalanche/Source/AvalancheInteractiveTools/Public/AvaInteractiveToolsCommands.h
new file mode 100644
--- /dev/null
+++ b/Engine/Plugins/Experimental/Avalanche/Source/AvalancheInteractiveTools/Public/AvaInteractiveToolsCommands.h
@@ -0,0 +1,36 @@
+// Copyright Epic Games, Inc. All Rights Reserved.
+
+#pragma once
+
+#include "Framework/Commands/Commands.h"
+
+class FAvaInteractiveToolsCommands : public TCommands<FAvaInteractiveToolsCommands>
+{
+public:
+	static AVALANCHEINTERACTIVETOOLS_API const FAvaInteractiveToolsCommands* GetExternal();
+
+	FAvaInteractiveToolsCommands();
+
+	//~ Begin TCommands
+	virtual void RegisterCommands() override;
+	//~ End TCommands
+
+	void RegisterEdModeCommands();
+	void RegisterCategoryCommands();
+	void RegisterActorCommands();
+
+	// Ed Mode
+	TSharedPtr<FUICommandInfo> CancelActiveTool;
+	TSharedPtr<FUICommandInfo> ToggleViewportToolbar;
+	TSharedPtr<FUICommandInfo> OpenViewportToolbarSettings;
+
+	// Categories
+	TSharedPtr<FUICommandInfo> Category_2D;
+	TSharedPtr<FUICommandInfo> Category_3D;
+	TSharedPtr<FUICommandInfo> Category_Actor;
+	TSharedPtr<FUICommandInfo> Category_Cloner;
+	TSharedPtr<FUICommandInfo> Category_Effector;
+
+	TSharedPtr<FUICommandInfo> Tool_Actor_Null;
+	TSharedPtr<FUICommandInfo> Tool_Actor_Spline;
+};
diff --git a/Engine/Plugins/Experimental/Avalanche/Source/AvalancheInteractiveTools/Public/AvaInteractiveToolsSettings.h b/Engine/Plugins/Experimental/Avalanche/Source/AvalancheInteractiveTools/Public/AvaInteractiveToolsSettings.h
--- a/Engine/Plugins/Experimental/Avalanche/Source/AvalancheInteractiveTools/Public/AvaInteractiveToolsSettings.h
+++ b/Engine/Plugins/Experimental/Avalanche/Source/AvalancheInteractiveTools/Public/AvaInteractiveToolsSettings.h
@@ -40,6 +40,15 @@ public:
 	/** Open these settings in the editor */
 	void OpenEditorSettingsWindow() const;
 
+	/** Checks if active mode supports viewport toolbar */
+	bool IsViewportToolbarSupported() const;
+
+	/** Set visibility of viewport toolbar */
+	void SetViewportToolbarVisible(bool bInVisible) const;
+
+	/** Get visibility of viewport toolbar*/
+	bool GetViewportToolbarVisible() const;
+
 	/** Distance from the camera at which actors are created. */
 	UPROPERTY(Config, EditAnywhere, Category = "Motion Design")
 	float CameraDistance = 500.f;
diff --git a/Engine/Plugins/Experimental/Avalanche/Source/AvalancheLevelViewport/AvalancheLevelViewport.Build.cs b/Engine/Plugins/Experimental/Avalanche/Source/AvalancheLevelViewport/AvalancheLevelViewport.Build.cs
--- a/Engine/Plugins/Experimental/Avalanche/Source/AvalancheLevelViewport/AvalancheLevelViewport.Build.cs
+++ b/Engine/Plugins/Experimental/Avalanche/Source/AvalancheLevelViewport/AvalancheLevelViewport.Build.cs
@@ -36,6 +36,7 @@ public class AvalancheLevelViewport : ModuleRules
 				"Avalanche",
 				"AvalancheComponentVisualizers",
 				"AvalancheEditorCore",
+				"AvalancheInteractiveTools",
 				"AvalancheViewport",
 				"CinematicCamera",
 				"CustomDetailsView",
diff --git a/Engine/Plugins/Experimental/Avalanche/Source/AvalancheLevelViewport/Private/Widgets/SAvaLevelViewportStatusBarButtons.cpp b/Engine/Plugins/Experimental/Avalanche/Source/AvalancheLevelViewport/Private/Widgets/SAvaLevelViewportStatusBarButtons.cpp
--- a/Engine/Plugins/Experimental/Avalanche/Source/AvalancheLevelViewport/Private/Widgets/SAvaLevelViewportStatusBarButtons.cpp
+++ b/Engine/Plugins/Experimental/Avalanche/Source/AvalancheLevelViewport/Private/Widgets/SAvaLevelViewportStatusBarButtons.cpp
@@ -1,14 +1,19 @@
 // Copyright Epic Games, Inc. All Rights Reserved.
 
 #include "Widgets/SAvaLevelViewportStatusBarButtons.h"
+
 #include "AvaEditorWidgetUtils.h"
+#include "AvaInteractiveToolsCommands.h"
+#include "AvaInteractiveToolsSettings.h"
 #include "AvaLevelViewportCommands.h"
 #include "AvaLevelViewportModule.h"
 #include "AvaLevelViewportStyle.h"
 #include "AvaViewportPostProcessManager.h"
 #include "AvaViewportSettings.h"
 #include "CoreGlobals.h"
+#include "EditorSupportDelegates.h"
 #include "Engine/Texture.h"
+#include "IAvalancheInteractiveToolsModule.h"
 #include "Interaction/AvaIsolateActorsOperation.h"
 #include "LevelEditor.h"
 #include "LevelViewportActions.h"
@@ -328,6 +333,21 @@ void SAvaLevelViewportStatusBarButtons::PopulateViewportButtons(TSharedPtr<SHori
 
 	const FAvaLevelViewportCommands& CommandsRef = FAvaLevelViewportCommands::GetInternal();
 	const FLevelViewportCommands& ViewportActionsRef = FLevelViewportCommands::Get();
+	const FAvaInteractiveToolsCommands* IFCommands = FAvaInteractiveToolsCommands::GetExternal();
+
+	InContainer->AddSlot()
+		.AutoWidth()
+		.Padding(ViewportStatusBarButton::Padding)
+		[
+			ViewportStatusBarButton::MakeButton(
+				this,
+				IFCommands->ToggleViewportToolbar,
+				IFCommands->ToggleViewportToolbar->GetIcon().GetIcon(),
+				&SAvaLevelViewportStatusBarButtons::ToggleViewportToolbar,
+				&SAvaLevelViewportStatusBarButtons::GetToggleViewportToolbarEnabled,
+				&SAvaLevelViewportStatusBarButtons::GetToggleViewportToolbarColor
+			)
+		];
 
 	InContainer->AddSlot()
 		.AutoWidth()
@@ -737,6 +757,44 @@ void SAvaLevelViewportStatusBarButtons::TogglePostProcessMenu(const FToolMenuCon
 	UE::Ava::LevelViewportStatusBarButtons::Private::TogglePostProcess(ViewportFrameWeak, InPostProcessType);
 }
 
+FSlateColor SAvaLevelViewportStatusBarButtons::GetToggleViewportToolbarColor() const
+{
+	if (const UAvaInteractiveToolsSettings* ITSettings = UAvaInteractiveToolsSettings::Get())
+	{
+		if (ITSettings->IsViewportToolbarSupported())
+		{
+			const bool bIsVisible = ITSettings->GetViewportToolbarVisible();
+			return bIsVisible
+				? UE::AvaLevelViewport::Private::ViewportStatusBarButton::ActiveColor
+				: UE::AvaLevelViewport::Private::ViewportStatusBarButton::EnabledColor;
+		}
+	}
+	
+	return UE::AvaLevelViewport::Private::ViewportStatusBarButton::DisabledColor;
+}
+
+bool SAvaLevelViewportStatusBarButtons::GetToggleViewportToolbarEnabled() const
+{
+	if (const UAvaInteractiveToolsSettings* ITSettings = UAvaInteractiveToolsSettings::Get())
+	{
+		return ITSettings->IsViewportToolbarSupported();
+	}
+
+	return false;
+}
+
+FReply SAvaLevelViewportStatusBarButtons::ToggleViewportToolbar()
+{
+	if (const UAvaInteractiveToolsSettings* ITSettings = UAvaInteractiveToolsSettings::Get())
+	{
+		const bool bIsVisible = ITSettings->GetViewportToolbarVisible();
+		ITSettings->SetViewportToolbarVisible(!bIsVisible);
+		return FReply::Handled();
+	}
+
+	return FReply::Unhandled();
+}
+
 FSlateColor SAvaLevelViewportStatusBarButtons::GetHighResScreenshotColor() const
 {
 	return FSlateColor(EStyleColor::Foreground);
diff --git a/Engine/Plugins/Experimental/Avalanche/Source/AvalancheLevelViewport/Private/Widgets/SAvaLevelViewportStatusBarButtons.h b/Engine/Plugins/Experimental/Avalanche/Source/AvalancheLevelViewport/Private/Widgets/SAvaLevelViewportStatusBarButtons.h
--- a/Engine/Plugins/Experimental/Avalanche/Source/AvalancheLevelViewport/Private/Widgets/SAvaLevelViewportStatusBarButtons.h
+++ b/Engine/Plugins/Experimental/Avalanche/Source/AvalancheLevelViewport/Private/Widgets/SAvaLevelViewportStatusBarButtons.h
@@ -62,6 +62,10 @@ protected:
 	void TogglePostProcessMenu(const FToolMenuContext& InContext, EAvaViewportPostProcessType InPostProcessType);
 
 	// Viewport
+	FSlateColor GetToggleViewportToolbarColor() const;
+	bool GetToggleViewportToolbarEnabled() const;
+	FReply ToggleViewportToolbar();
+	
 	FSlateColor GetHighResScreenshotColor() const;
 	bool GetHighResScreenshotEnabled() const;
 	FReply HighResScreenshot();
