MotionDesign : InteractiveTools
- Replace motion design toolbox tab toolkit to use same as modeling tools

diff --git a/Engine/Plugins/Experimental/Avalanche/Source/AvalancheInteractiveTools/Private/EdMode/AvaInteractiveToolsEdModeToolkit.cpp b/Engine/Plugins/Experimental/Avalanche/Source/AvalancheInteractiveTools/Private/EdMode/AvaInteractiveToolsEdModeToolkit.cpp
--- a/Engine/Plugins/Experimental/Avalanche/Source/AvalancheInteractiveTools/Private/EdMode/AvaInteractiveToolsEdModeToolkit.cpp
+++ b/Engine/Plugins/Experimental/Avalanche/Source/AvalancheInteractiveTools/Private/EdMode/AvaInteractiveToolsEdModeToolkit.cpp
@@ -6,6 +6,7 @@
 #include "AvaInteractiveToolsEdMode.h"
 #include "AvaInteractiveToolsSettings.h"
 #include "AvaInteractiveToolsStyle.h"
+#include "AvalancheInteractiveToolsModule.h"
 #include "IAvalancheInteractiveToolsModule.h"
 #include "Toolkits/AssetEditorModeUILayer.h"
 #include "Tools/UEdMode.h"
@@ -44,6 +45,19 @@ FText FAvaInteractiveToolsEdModeToolkit::GetBaseToolkitName() const
 	return LOCTEXT("DisplayName", "AvaInteractiveToolsEdMode Tool");
 }
 
+TSharedPtr<SWidget> FAvaInteractiveToolsEdModeToolkit::GetInlineContent() const
+{
+	checkf(ToolkitWidget.IsValid(), TEXT("Toolkit widget for interactive tools editor mode is invalid"))
+
+	return SNew(SVerticalBox)
+		+ SVerticalBox::Slot()
+		.FillHeight(1.0f)
+		.VAlign(VAlign_Fill)
+		[
+			ToolkitWidget.ToSharedRef()
+		];
+}
+
 void FAvaInteractiveToolsEdModeToolkit::GetToolPaletteNames(TArray<FName>& InPaletteName) const
 {
 	const TMap<FName, TSharedPtr<FUICommandInfo>>& Categories = IAvalancheInteractiveToolsModule::Get().GetCategories();
@@ -75,6 +89,7 @@ void FAvaInteractiveToolsEdModeToolkit::Init(const TSharedPtr<IToolkitHost>& Ini
 	FModeToolkit::Init(InitToolkitHost, InOwningMode);
 
 	MakeViewportOverlayToolbar();
+	MakeToolkitPalettes();
 }
 
 void FAvaInteractiveToolsEdModeToolkit::OnToolStarted(UInteractiveToolManager* InManager, UInteractiveTool* InTool)
@@ -87,15 +102,30 @@ void FAvaInteractiveToolsEdModeToolkit::OnToolEnded(UInteractiveToolManager* InM
 	// Nothing
 }
 
+void FAvaInteractiveToolsEdModeToolkit::InvokeUI()
+{
+	FModeToolkit::InvokeUI();
+
+	const TSharedPtr<SWidget> InlineContentWidget = GetInlineContent();
+	InlineContentHolder->SetContent(InlineContentWidget.ToSharedRef());
+}
+
 void FAvaInteractiveToolsEdModeToolkit::RequestModeUITabs()
 {
-	if (!HasIntegratedToolPalettes() && !HasToolkitBuilder())
+	if (const TSharedPtr<FAssetEditorModeUILayer> ModeUILayerPtr = ModeUILayer.Pin())
 	{
-		const TSharedPtr<FAssetEditorModeUILayer> ModeUILayerPtr = ModeUILayer.Pin();
-		ToolbarInfo.OnSpawnTab = FOnSpawnTab::CreateSP(SharedThis(this), &FAvaInteractiveToolsEdModeToolkit::MakeModeToolbarTab);
-		ToolbarInfo.TabLabel = LOCTEXT("MotionDesignToolboxTab", "Motion Design Toolbox");
-		ToolbarInfo.TabTooltip = LOCTEXT("MotionDesignToolboxTabTooltipText", "Opens the Motion Design Toolbox tab.");
-		ModeUILayerPtr->SetModePanelInfo(UAssetEditorUISubsystem::VerticalToolbarID, ToolbarInfo);
+		PrimaryTabInfo.OnSpawnTab = FOnSpawnTab::CreateSP(SharedThis(this), &FAvaInteractiveToolsEdModeToolkit::CreatePrimaryModePanel);
+		PrimaryTabInfo.TabLabel = LOCTEXT("MotionDesignToolboxTab", "Motion Design");
+		PrimaryTabInfo.TabTooltip = LOCTEXT("MotionDesignToolboxTabTooltipText", "Opens the Motion Design tab.");
+		ModeUILayerPtr->SetModePanelInfo(UAssetEditorUISubsystem::TopLeftTabID, PrimaryTabInfo);
+
+		if (!HasIntegratedToolPalettes() && !HasToolkitBuilder())
+		{
+			ToolbarInfo.OnSpawnTab = FOnSpawnTab::CreateSP(SharedThis(this), &FAvaInteractiveToolsEdModeToolkit::MakeModeToolbarTab);
+			ToolbarInfo.TabLabel = LOCTEXT("MotionDesignToolbarTab", "Motion Design Toolbar");
+			ToolbarInfo.TabTooltip = LOCTEXT("MotionDesignToolbarTabTooltipText", "Opens the toolbar for the Motion Design toolbox.");
+			ModeUILayerPtr->SetModePanelInfo(UAssetEditorUISubsystem::VerticalToolbarID, ToolbarInfo);
+		}
 	}
 }
 
@@ -227,6 +257,96 @@ void FAvaInteractiveToolsEdModeToolkit::MakeViewportOverlayToolbar()
 	GetToolkitHost()->AddViewportOverlayWidget(ViewportToolbarWidget.ToSharedRef());
 }
 
+void FAvaInteractiveToolsEdModeToolkit::MakeToolkitPalettes()
+{
+	const UAvaInteractiveToolsEdMode* AvaInteractiveToolsMode = CastChecked<UAvaInteractiveToolsEdMode>(GetScriptableEditorMode().Get());
+
+	ModeDetailsView->SetIsPropertyVisibleDelegate(
+		FIsPropertyVisible::CreateLambda(
+			[](const FPropertyAndParent& InPropertyAndParent) -> bool
+			{
+				static const FString Material = FString(TEXT("Material"));
+				static const FName Category = FName(TEXT("Category"));
+				const FString CategoryMeta = InPropertyAndParent.Property.GetMetaData(Category);
+				return CategoryMeta != Material;
+			}
+	));
+
+	ToolkitSections = MakeShared<FToolkitSections>();
+
+	// Show warning text when tool is active
+	ToolkitSections->ToolWarningArea = SNew(STextBlock)
+		.Text(this, &FAvaInteractiveToolsEdModeToolkit::GetToolWarningText)
+		.AutoWrapText(true);
+		
+	ToolkitBuilder = MakeShared<FToolkitBuilder>(
+		AvaInteractiveToolsMode->GetModeInfo().ToolbarCustomizationName,
+		GetToolkitCommands(),
+		ToolkitSections);
+
+	const TMap<FName, TSharedPtr<FUICommandInfo>>& Categories = IAvalancheInteractiveToolsModule::Get().GetCategories();
+	TSharedPtr<FUICommandInfo> FirstCategoryCommand = nullptr;
+
+	// Group tools command by category in the palette
+	for (const TPair<FName, TSharedPtr<FUICommandInfo>>& Category : Categories)
+	{
+		if (Category.Value.IsValid())
+		{
+			TArray<TSharedPtr<FUICommandInfo>> CategoryCommands;
+
+			if (const TArray<FAvaInteractiveToolsToolParameters>* ToolList = IAvalancheInteractiveToolsModule::Get().GetTools(Category.Key))
+			{
+				for (const FAvaInteractiveToolsToolParameters& Tool : *ToolList)
+				{
+					if (Tool.UICommand.IsValid())
+					{
+						CategoryCommands.Add(Tool.UICommand);
+					}
+				}
+			}
+
+			if (CategoryCommands.IsEmpty() == false)
+			{
+				ToolkitBuilder->AddPalette(MakeShared<FToolPalette>(Category.Value.ToSharedRef(), CategoryCommands));
+
+				if (FirstCategoryCommand.IsValid() == false)
+				{
+					FirstCategoryCommand = Category.Value;
+				}
+			}
+		}
+	}
+
+	if (FirstCategoryCommand.IsValid())
+	{
+		ToolkitBuilder->SetActivePaletteOnLoad(FirstCategoryCommand.Get());
+	}
+
+	ToolkitBuilder->UpdateWidget();
+
+	const TSharedPtr<SWidget> ToolkitGeneratedWidget = ToolkitBuilder->GenerateWidget();
+	
+	checkf(ToolkitGeneratedWidget.IsValid(), TEXT("Generated widget for interactive tools editor mode is invalid"))
+	
+	SAssignNew(ToolkitWidget, SBorder)
+		.HAlign(HAlign_Fill)
+		.Padding(0)
+		.BorderImage(FAppStyle::GetBrush("ToolPanel.GroupBorder"))
+		[
+			ToolkitGeneratedWidget->AsShared()
+		];
+}
+
+FText FAvaInteractiveToolsEdModeToolkit::GetToolWarningText() const
+{
+	static const FText ToolActive = LOCTEXT("ActiveToolWarning", "Tool Active.\n\nSelect the tool again to perform the default action (if supported).\n\nRight click or press escape to cancel.");
+	static const FText ToolInactive = LOCTEXT("InactiveToolWarning", "Select a tool once to start drawing or double click to perform default action (if supported).");
+
+	return FAvalancheInteractiveToolsModule::Get().HasActiveTool()
+		? ToolActive
+		: ToolInactive;
+}
+
 void FAvaInteractiveToolsEdModeToolkit::OnSettingsChanged(UObject* InSettings, FPropertyChangedEvent& InEvent)
 {
 	MakeViewportOverlayToolbar();
diff --git a/Engine/Plugins/Experimental/Avalanche/Source/AvalancheInteractiveTools/Private/EdMode/AvaInteractiveToolsEdModeToolkit.h b/Engine/Plugins/Experimental/Avalanche/Source/AvalancheInteractiveTools/Private/EdMode/AvaInteractiveToolsEdModeToolkit.h
--- a/Engine/Plugins/Experimental/Avalanche/Source/AvalancheInteractiveTools/Private/EdMode/AvaInteractiveToolsEdModeToolkit.h
+++ b/Engine/Plugins/Experimental/Avalanche/Source/AvalancheInteractiveTools/Private/EdMode/AvaInteractiveToolsEdModeToolkit.h
@@ -4,6 +4,7 @@
 
 #include "Toolkits/BaseToolkit.h"
 
+class SWidget;
 class UObject;
 struct FPropertyChangedEvent;
 
@@ -23,17 +24,19 @@ protected:
 	//~ Begin IToolkit
 	virtual FName GetToolkitFName() const override;
 	virtual FText GetBaseToolkitName() const override;
-	virtual void GetToolPaletteNames(TArray<FName>& PaletteNames) const override;
-	virtual FText GetToolPaletteDisplayName(FName Palette) const override;
-	virtual void OnToolPaletteChanged(FName PaletteName) override;
+	virtual TSharedPtr<SWidget> GetInlineContent() const override;
 	//~ End IToolkit
 
 	//~ Begin FModeToolkit
 	virtual void Init(const TSharedPtr<IToolkitHost>& InitToolkitHost, TWeakObjectPtr<UEdMode> InOwningMode) override;
+	virtual void GetToolPaletteNames(TArray<FName>& PaletteNames) const override;
+	virtual FText GetToolPaletteDisplayName(FName Palette) const override;
+	virtual void OnToolPaletteChanged(FName PaletteName) override;
 	virtual bool HasIntegratedToolPalettes() const override { return false; }
 	virtual bool HasExclusiveToolPalettes() const override { return false; }
 	virtual void OnToolStarted(UInteractiveToolManager* InManager, UInteractiveTool* InTool) override;
 	virtual void OnToolEnded(UInteractiveToolManager* InManager, UInteractiveTool* InTool) override;
+	virtual void InvokeUI() override;
 	//~ End FModeToolkit
 
 	//~ Begin FModeToolkit
@@ -44,9 +47,14 @@ private:
 	/** Creates an viewport overlay toolbar */
 	void MakeViewportOverlayToolbar();
 
+	void MakeToolkitPalettes();
+	
+	FText GetToolWarningText() const;
+
 	void OnSettingsChanged(UObject* InSettings, FPropertyChangedEvent& InEvent);
 	void OnViewportChanged(TSharedPtr<IAssetViewport> InOldViewport, TSharedPtr<IAssetViewport> InNewViewport);
 
+	TSharedPtr<SWidget> ToolkitWidget;
 	TSharedPtr<SWidget> ViewportToolbarWidget;
 	bool bViewportToolbarVisible = true;
 };
