MotionDesign : Modifiers
- Improving pattern modifier to support keyframe on properties and simplifying alignment options
- Modifier tool can be extended to create more patterns
- Assets using this modifier old version will be migrated

diff --git a/Engine/Plugins/Experimental/Avalanche/Source/AvalancheModifiers/Private/Modifiers/AvaPatternModifier.cpp b/Engine/Plugins/Experimental/Avalanche/Source/AvalancheModifiers/Private/Modifiers/AvaPatternModifier.cpp
--- a/Engine/Plugins/Experimental/Avalanche/Source/AvalancheModifiers/Private/Modifiers/AvaPatternModifier.cpp
+++ b/Engine/Plugins/Experimental/Avalanche/Source/AvalancheModifiers/Private/Modifiers/AvaPatternModifier.cpp
@@ -6,9 +6,28 @@
 #include "Components/DynamicMeshComponent.h"
 #include "DynamicMesh/MeshTransforms.h"
 #include "DynamicMeshEditor.h"
+#include "Tools/AvaPatternModifierCircleTool.h"
+#include "Tools/AvaPatternModifierGridTool.h"
+#include "Tools/AvaPatternModifierLineTool.h"
 
 #define LOCTEXT_NAMESPACE "AvaPatternModifier"
 
+struct FAvaPatternModifierVersion
+{
+	enum Type : int32
+	{
+		PreVersioning = 0,
+		/** Moved properties within modifier */
+		MigrateProperties,
+		VersionPlusOne,
+		LatestVersion = VersionPlusOne - 1
+	};
+
+	static constexpr FGuid GUID = FGuid(0x9271D8A4, 0xBF414601, 0xA20FC0A3, 0x9D829565);
+};
+
+FCustomVersionRegistration GRegisterAvaPatternModifierVersion(FAvaPatternModifierVersion::GUID, static_cast<int32>(FAvaPatternModifierVersion::LatestVersion), TEXT("AvaPatternModifierVersion"));
+
 void UAvaPatternModifier::OnModifierCDOSetup(FActorModifierCoreMetadata& InMetadata)
 {
 	Super::OnModifierCDOSetup(InMetadata);
@@ -22,11 +41,9 @@ void UAvaPatternModifier::OnModifierCDOSetup(FActorModifierCoreMetadata& InMetad
 
 void UAvaPatternModifier::Apply()
 {
-	if ((Layout == EAvaPatternModifierLayout::Line && LineLayoutOptions.RepeatCount <= 0)
-		|| (Layout == EAvaPatternModifierLayout::Grid && (GridLayoutOptions.RepeatCount.X <= 0 ||  GridLayoutOptions.RepeatCount.Y <= 0))
-		|| (Layout == EAvaPatternModifierLayout::Circle && CircleLayoutOptions.RepeatCount <= 0))
+	if (!ActiveTool)
 	{
-		Fail(LOCTEXT("InvalidRepeatCount", "Layout repeat count must be greater than 0"));
+		Fail(LOCTEXT("InvalidActiveTool", "Active tool is invalid or not set"));
 		return;
 	}
 
@@ -49,146 +66,23 @@ void UAvaPatternModifier::Apply()
 		OriginalMeshBounds = GetMeshBounds();
 
 		// clear all since we have a copy
-		for (int32 TId : AppendToMesh.TriangleIndicesItr())
+		for (const int32 TId : AppendToMesh.TriangleIndicesItr())
 		{
 			AppendToMesh.RemoveTriangle(TId);
 		}
 
-		FVector CenterAxis = FVector::ZeroVector;
-
-		if (Layout == EAvaPatternModifierLayout::Line)
-		{
-			// edit original mesh to append tmp mesh
-			FMeshIndexMappings TmpMappings;
-			FDynamicMeshEditor Editor(&AppendToMesh);
-			const FTransformSRT3d TransformChange = GetLineLayoutTransformChange();
-			if (!LineLayoutOptions.bAccumulateTransform)
-			{
-				FTransform RotationAndScaleTransform;
-				RotationAndScaleTransform.SetLocation(FVector::ZeroVector);
-				RotationAndScaleTransform.SetRotation(TransformChange.GetRotator().Quaternion());
-				RotationAndScaleTransform.SetScale3D(TransformChange.GetScale());
-				FTransformSRT3d RotateScale(RotationAndScaleTransform);
-				MeshTransforms::ApplyTransform(TmpMesh, RotateScale, true);
-			}
-			for (int32 Idx = 0; Idx < LineLayoutOptions.RepeatCount; Idx++)
-			{
-				if (Idx != 0)
-				{
-					if (LineLayoutOptions.bAccumulateTransform)
-					{
-						MeshTransforms::ApplyTransform(TmpMesh, TransformChange, true);
-					}
-					else
-					{
-						MeshTransforms::Translate(TmpMesh, TransformChange.GetTranslation());
-					}
-				}
-				Editor.AppendMesh(&TmpMesh, TmpMappings);
-				TmpMappings.Reset();
-			}
-			if (LineLayoutOptions.bCentered)
-			{
-				if (LineLayoutOptions.Axis == EAvaPatternModifierAxis::X)
-				{
-					CenterAxis.X = 1;
-				}
-				else if (LineLayoutOptions.Axis == EAvaPatternModifierAxis::Y)
-				{
-					CenterAxis.Y = 1;
-				}
-				else if (LineLayoutOptions.Axis == EAvaPatternModifierAxis::Z)
-				{
-					CenterAxis.Z = 1;
-				}
-			}
-		}
-		else if (Layout == EAvaPatternModifierLayout::Grid)
-		{
-			// edit original mesh to append tmp mesh
-			FMeshIndexMappings TmpMappings;
-			FDynamicMeshEditor Editor(&AppendToMesh);
-			const FTransformSRT3d ColTransformChange = GetGridLayoutColTransformChange();
-			const FTransformSRT3d RowTransformChange = GetGridLayoutRowTransformChange();
-			if (!GridLayoutOptions.bAccumulateTransform)
-			{
-				FTransform RotationAndScaleTransform;
-				RotationAndScaleTransform.SetLocation(FVector::ZeroVector);
-				RotationAndScaleTransform.SetRotation(ColTransformChange.GetRotator().Quaternion());
-				RotationAndScaleTransform.SetScale3D(ColTransformChange.GetScale());
-				FTransformSRT3d RotateScale(RotationAndScaleTransform);
-				MeshTransforms::ApplyTransform(TmpMesh, RotateScale, true);
-			}
-			for (int32 Row = 0; Row < GridLayoutOptions.RepeatCount.X; Row++)
-			{
-				for (int32 Col = 0; Col < GridLayoutOptions.RepeatCount.Y; Col++)
-				{
-					if (!(Row == 0 && Col == 0))
-					{
-						if (GridLayoutOptions.bAccumulateTransform)
-						{
-							MeshTransforms::ApplyTransform(TmpMesh, ColTransformChange, true);
-						}
-						else
-						{
-							MeshTransforms::Translate(TmpMesh, ColTransformChange.GetTranslation());
-						}
-					}
-					Editor.AppendMesh(&TmpMesh, TmpMappings);
-					TmpMappings.Reset();
-				}
-				MeshTransforms::Translate(TmpMesh, RowTransformChange.GetTranslation() - ColTransformChange.GetTranslation() * GridLayoutOptions.RepeatCount.Y);
-			}
-			if (GridLayoutOptions.bCentered)
-			{
-				if (GridLayoutOptions.Plane == EAvaPatternModifierPlane::XY)
-				{
-					CenterAxis.X = 1;
-					CenterAxis.Y = 1;
-				}
-				else if (GridLayoutOptions.Plane == EAvaPatternModifierPlane::YZ)
-				{
-					CenterAxis.Y = 1;
-					CenterAxis.Z = 1;
-				}
-				else if (GridLayoutOptions.Plane == EAvaPatternModifierPlane::ZX)
-				{
-					CenterAxis.X = 1;
-					CenterAxis.Z = 1;
-				}
-			}
-		}
-		else if (Layout == EAvaPatternModifierLayout::Circle)
+		FMeshIndexMappings TmpMappings;
+		FDynamicMeshEditor Editor(&AppendToMesh);
+		for (const FTransform& TransformInstance : ActiveTool->GetTransformInstances(OriginalMeshBounds))
 		{
-			// edit original mesh to append tmp mesh
-			FMeshIndexMappings TmpMappings;
-			FDynamicMeshEditor Editor(&AppendToMesh);
-			if (!CircleLayoutOptions.bAccumulateTransform)
-			{
-				FTransform RotationAndScaleTransform;
-				RotationAndScaleTransform.SetLocation(FVector::ZeroVector);
-				RotationAndScaleTransform.SetRotation(CircleLayoutOptions.Rotation.Quaternion());
-				RotationAndScaleTransform.SetScale3D(CircleLayoutOptions.Scale);
-				FTransformSRT3d RotateScale(RotationAndScaleTransform);
-				MeshTransforms::ApplyTransform(TmpMesh, RotateScale, true);
-			}
-			for (int32 Idx = 0; Idx < CircleLayoutOptions.RepeatCount; Idx++)
-			{
-				FDynamicMesh3 AppendMesh = TmpMesh;
-				FTransformSRT3d TransformChange = GetCircleLayoutTransformChange(Idx);
-				const FVector Translation = TransformChange.GetTranslation();
-
-				if (CircleLayoutOptions.bAccumulateTransform)
-				{
-					TransformChange.SetTranslation(FVector::ZeroVector);
-					MeshTransforms::ApplyTransform(TmpMesh, TransformChange, true);
-				}
-
-				MeshTransforms::Translate(AppendMesh, Translation);
-				Editor.AppendMesh(&AppendMesh, TmpMappings);
-				TmpMappings.Reset();
-			}
+			MeshTransforms::ApplyTransform(TmpMesh, TransformInstance, /** ReverseIfNeeded */true);
+			Editor.AppendMesh(&TmpMesh, TmpMappings);
+			TmpMappings.Reset();
+			MeshTransforms::ApplyTransformInverse(TmpMesh, TransformInstance, /** ReverseIfNeeded */true);
 		}
+
+		const FVector CenterAxis = ActiveTool->GetCenterAlignmentAxis();
+
 		// center this mesh
 		if (!CenterAxis.IsZero())
 		{
@@ -201,202 +95,192 @@ void UAvaPatternModifier::Apply()
 				MeshTransforms::Translate(AppendToMesh, -Translate);
 			}
 		}
-	}, EDynamicMeshChangeType::GeneralEdit, EDynamicMeshAttributeChangeFlags::Unknown);
+	}
+	, EDynamicMeshChangeType::GeneralEdit
+	, EDynamicMeshAttributeChangeFlags::Unknown);
 
 	Next();
 }
 
-#if WITH_EDITOR
-void UAvaPatternModifier::PostEditChangeProperty(FPropertyChangedEvent& PropertyChangedEvent)
+UAvaPatternModifierTool* UAvaPatternModifier::FindOrAddTool(TSubclassOf<UAvaPatternModifierTool> InToolClass)
 {
-	Super::PostEditChangeProperty(PropertyChangedEvent);
-
-	const FName MemberName = PropertyChangedEvent.GetMemberPropertyName();
-
-	static const FName LayoutName = GET_MEMBER_NAME_CHECKED(UAvaPatternModifier, Layout);
-	static const FName LineLayoutOptionsName = GET_MEMBER_NAME_CHECKED(UAvaPatternModifier, LineLayoutOptions);
-	static const FName GridLayoutOptionsName = GET_MEMBER_NAME_CHECKED(UAvaPatternModifier, GridLayoutOptions);
-	static const FName CircleLayoutOptionsName = GET_MEMBER_NAME_CHECKED(UAvaPatternModifier, CircleLayoutOptions);
-
-	if (MemberName == LayoutName)
-	{
-		OnLayoutChanged();
-	}
-	else if (MemberName == LineLayoutOptionsName)
+	if (!InToolClass.Get())
 	{
-		OnLineLayoutOptionsChanged();
+		return nullptr;
 	}
-	else if (MemberName == GridLayoutOptionsName)
+
+	if (const TObjectPtr<UAvaPatternModifierTool>* Tool = Tools.FindByPredicate([InToolClass](const TObjectPtr<UAvaPatternModifierTool>& InTool)
 	{
-		OnGridLayoutOptionsChanged();
-	}
-	else if (MemberName == CircleLayoutOptionsName)
+		return InTool && InTool->GetClass() == InToolClass.Get();
+	}))
 	{
-		OnCircleLayoutOptionsChanged();
+		return *Tool;
 	}
-}
-#endif
 
-void UAvaPatternModifier::SetLayout(EAvaPatternModifierLayout InLayout)
-{
-	if (Layout == InLayout)
+	if (const UAvaPatternModifierTool* DefaultObject = InToolClass.GetDefaultObject())
 	{
-		return;
+		if (UAvaPatternModifierTool* NewTool = NewObject<UAvaPatternModifierTool>(this, InToolClass.Get(), DefaultObject->GetToolName()))
+		{
+			Tools.Add(NewTool);
+			return NewTool;
+		}
 	}
 
-	Layout = InLayout;
-	OnLayoutChanged();
-}
-
-void UAvaPatternModifier::SetLineLayoutOptions(const FAvaPatternModifierLineLayoutOptions& InOptions)
-{
-	LineLayoutOptions = InOptions;
-	OnLineLayoutOptionsChanged();
+	return nullptr;
 }
 
-void UAvaPatternModifier::SetGridLayoutOptions(const FAvaPatternModifierGridLayoutOptions& InOptions)
+void UAvaPatternModifier::Serialize(FArchive& InArchive)
 {
-	GridLayoutOptions = InOptions;
-	OnGridLayoutOptionsChanged();
-}
+	InArchive.UsingCustomVersion(FAvaPatternModifierVersion::GUID);
 
-void UAvaPatternModifier::SetCircleLayoutOptions(const FAvaPatternModifierCircleLayoutOptions& InOptions)
-{
-	CircleLayoutOptions = InOptions;
-	OnCircleLayoutOptionsChanged();
-}
+	Super::Serialize(InArchive);
 
-void UAvaPatternModifier::OnLayoutChanged()
-{
-	MarkModifierDirty();
-}
+	const int32 Version = InArchive.CustomVer(FAvaPatternModifierVersion::GUID);
 
-void UAvaPatternModifier::OnLineLayoutOptionsChanged()
-{
-	if (Layout == EAvaPatternModifierLayout::Line)
+	if (Version < FAvaPatternModifierVersion::LatestVersion)
 	{
-		MarkModifierDirty();
+		MigrateVersion(Version);
 	}
 }
 
-void UAvaPatternModifier::OnGridLayoutOptionsChanged()
+#if WITH_EDITOR
+void UAvaPatternModifier::PostEditChangeProperty(FPropertyChangedEvent& PropertyChangedEvent)
 {
-	if (Layout == EAvaPatternModifierLayout::Grid)
+	Super::PostEditChangeProperty(PropertyChangedEvent);
+
+	const FName MemberName = PropertyChangedEvent.GetMemberPropertyName();
+
+	static const FName ActiveToolClassName = GET_MEMBER_NAME_CHECKED(UAvaPatternModifier, ActiveToolClass);
+
+	if (MemberName.IsEqual(ActiveToolClassName))
 	{
-		MarkModifierDirty();
+		OnActiveToolClassChanged();
 	}
 }
+#endif
 
-void UAvaPatternModifier::OnCircleLayoutOptionsChanged()
+UAvaPatternModifier::UAvaPatternModifier()
 {
-	if (Layout == EAvaPatternModifierLayout::Circle)
-	{
-		MarkModifierDirty();
-	}
+	ActiveToolClass = UAvaPatternModifierLineTool::StaticClass();
 }
 
-UE::Geometry::FTransformSRT3d UAvaPatternModifier::GetLineLayoutTransformChange() const
+void UAvaPatternModifier::OnModifierEnabled(EActorModifierCoreEnableReason InReason)
 {
-	const FVector Axis = (LineLayoutOptions.bAxisInverted ? -1 : 1 ) *
-		(LineLayoutOptions.Axis == EAvaPatternModifierAxis::X ? FVector::XAxisVector : (LineLayoutOptions.Axis == EAvaPatternModifierAxis::Y ? FVector::YAxisVector : FVector::ZAxisVector));
-	const FVector Size3D = OriginalMeshBounds.GetSize();
-	const FVector Translation = Axis * Size3D + Axis * LineLayoutOptions.Spacing;
-	FTransform TransformChange;
-	TransformChange.SetRotation(LineLayoutOptions.Rotation.Quaternion());
-	TransformChange.SetLocation(Translation);
-	TransformChange.SetScale3D(LineLayoutOptions.Scale);
-	return UE::Geometry::FTransformSRT3d(TransformChange);
+	Super::OnModifierEnabled(InReason);
+	OnActiveToolClassChanged();
 }
 
-UE::Geometry::FTransformSRT3d UAvaPatternModifier::GetGridLayoutColTransformChange() const
+void UAvaPatternModifier::SetActiveToolClass(const TSubclassOf<UAvaPatternModifierTool>& InClass)
 {
-	FVector ColAxis;
-	if (GridLayoutOptions.Plane == EAvaPatternModifierPlane::XY)
-	{
-		ColAxis = FVector(0, 1, 0);
-	}
-	else if (GridLayoutOptions.Plane == EAvaPatternModifierPlane::YZ)
-	{
-		ColAxis = FVector(0, 0, 1);
-	}
-	else if (GridLayoutOptions.Plane == EAvaPatternModifierPlane::ZX)
+	if (!InClass.Get() || ActiveToolClass == InClass)
 	{
-		ColAxis = FVector(1, 0, 0);
+		return;
 	}
-	ColAxis *= (GridLayoutOptions.AxisInverted.bY ? -1 : 1);
-	const FVector Size3D = OriginalMeshBounds.GetSize();
-	const FVector Translation = ColAxis * Size3D + ColAxis * GridLayoutOptions.Spacing.Y;
-	FTransform TransformChange;
-	TransformChange.SetRotation(GridLayoutOptions.Rotation.Quaternion());
-	TransformChange.SetLocation(Translation);
-	TransformChange.SetScale3D(GridLayoutOptions.Scale);
-	return UE::Geometry::FTransformSRT3d(TransformChange);
+
+	ActiveToolClass = InClass;
+	OnActiveToolClassChanged();
 }
 
-UE::Geometry::FTransformSRT3d UAvaPatternModifier::GetGridLayoutRowTransformChange() const
+void UAvaPatternModifier::OnActiveToolClassChanged()
 {
-	FVector RowAxis;
-	if (GridLayoutOptions.Plane == EAvaPatternModifierPlane::XY)
-	{
-		RowAxis = FVector(1, 0, 0);
-	}
-	else if (GridLayoutOptions.Plane == EAvaPatternModifierPlane::YZ)
+	ActiveTool = FindOrAddTool(ActiveToolClass);
+	MarkModifierDirty();
+}
+
+void UAvaPatternModifier::MigrateVersion(int32 InCurrentVersion)
+{
+	if (InCurrentVersion >= FAvaPatternModifierVersion::LatestVersion)
 	{
-		RowAxis = FVector(0, 1, 0);
+		return;
 	}
-	else if (GridLayoutOptions.Plane == EAvaPatternModifierPlane::ZX)
+
+	if (InCurrentVersion < FAvaPatternModifierVersion::MigrateProperties)
 	{
-		RowAxis = FVector(0, 0, 1);
-	}
-	RowAxis *= (GridLayoutOptions.AxisInverted.bX ? -1 : 1);
-	const FVector Size3D = OriginalMeshBounds.GetSize();
-	const FVector Translation = RowAxis * Size3D + RowAxis * GridLayoutOptions.Spacing.X;
-	FTransform TransformChange;
-	TransformChange.SetRotation(GridLayoutOptions.Rotation.Quaternion());
-	TransformChange.SetLocation(Translation);
-	TransformChange.SetScale3D(GridLayoutOptions.Scale);
-	return UE::Geometry::FTransformSRT3d(TransformChange);
-}
+		PRAGMA_DISABLE_DEPRECATION_WARNINGS
 
-UE::Geometry::FTransformSRT3d UAvaPatternModifier::GetCircleLayoutTransformChange(int32 Idx) const
-{
-	const int32 Count = CircleLayoutOptions.RepeatCount;
-	const float Radius = CircleLayoutOptions.Radius;
-	const float StartAngle = FMath::DegreesToRadians(CircleLayoutOptions.StartAngle);
-	const float FullAngle = FMath::DegreesToRadians(CircleLayoutOptions.FullAngle);
-	const float AngleStep = FullAngle / Count;
+		if (UAvaPatternModifierLineTool* LineTool = FindOrAddTool<UAvaPatternModifierLineTool>())
+		{
+			LineTool->LineAxis = LineLayoutOptions.Axis;
+			LineTool->LineCount = LineLayoutOptions.RepeatCount;
+			LineTool->LineSpacing = LineLayoutOptions.Spacing;
+			LineTool->bLineAccumulateTransform = LineLayoutOptions.bAccumulateTransform;
+			LineTool->LineRotation = LineLayoutOptions.Rotation;
+			LineTool->LineScale = LineLayoutOptions.Scale;
+
+			if (LineLayoutOptions.bCentered)
+			{
+				LineTool->LineAlignment = EAvaPatternModifierLineAlignment::Center;
+			}
+			else if (LineLayoutOptions.bAxisInverted)
+			{
+				LineTool->LineAlignment = EAvaPatternModifierLineAlignment::End;
+			}
+			else
+			{
+				LineTool->LineAlignment = EAvaPatternModifierLineAlignment::Start;
+			}
+		}
+
+		if (UAvaPatternModifierGridTool* GridTool = FindOrAddTool<UAvaPatternModifierGridTool>())
+		{
+			GridTool->GridPlane = GridLayoutOptions.Plane;
+			GridTool->GridCountX = GridLayoutOptions.RepeatCount.X;
+			GridTool->GridCountY = GridLayoutOptions.RepeatCount.Y;
+			GridTool->GridSpacingX = GridLayoutOptions.Spacing.X;
+			GridTool->GridSpacingY = GridLayoutOptions.Spacing.Y;
+			GridTool->bGridAccumulateTransform = GridLayoutOptions.bAccumulateTransform;
+			GridTool->GridRotation = GridLayoutOptions.Rotation;
+			GridTool->GridScale = GridLayoutOptions.Scale;
 
-	const float ZeroAngle = StartAngle;
-	const float ZeroX = Radius * FMath::Cos(ZeroAngle);
-	const float ZeroY = Radius * FMath::Sin(ZeroAngle);
+			if (GridLayoutOptions.bCentered)
+			{
+				GridTool->GridAlignment = EAvaPatternModifierGridAlignment::Center;
+			}
+			else if (GridLayoutOptions.AxisInverted.bX && GridLayoutOptions.AxisInverted.bY)
+			{
+				GridTool->GridAlignment = EAvaPatternModifierGridAlignment::TopRight;
+			}
+			else if (GridLayoutOptions.AxisInverted.bX && !GridLayoutOptions.AxisInverted.bY)
+			{
+				GridTool->GridAlignment = EAvaPatternModifierGridAlignment::BottomRight;
+			}
+			else if (!GridLayoutOptions.AxisInverted.bX && GridLayoutOptions.AxisInverted.bY)
+			{
+				GridTool->GridAlignment = EAvaPatternModifierGridAlignment::TopLeft;
+			}
+			else if (!GridLayoutOptions.AxisInverted.bX && !GridLayoutOptions.AxisInverted.bY)
+			{
+				GridTool->GridAlignment = EAvaPatternModifierGridAlignment::BottomLeft;
+			}
+		}
 
-	const float CurrentAngle = StartAngle + Idx * AngleStep;
-	const float X = Radius * FMath::Cos(CurrentAngle);
-	const float Y = Radius * FMath::Sin(CurrentAngle);
+		if (UAvaPatternModifierCircleTool* CircleTool = FindOrAddTool<UAvaPatternModifierCircleTool>())
+		{
+			CircleTool->CirclePlane = CircleLayoutOptions.Plane;
+			CircleTool->CircleRadius = CircleLayoutOptions.Radius;
+			CircleTool->CircleStartAngle = CircleLayoutOptions.StartAngle;
+			CircleTool->CircleFullAngle = CircleLayoutOptions.FullAngle;
+			CircleTool->CircleCount = CircleLayoutOptions.RepeatCount;
+			CircleTool->bCircleAccumulateTransform = CircleLayoutOptions.bAccumulateTransform;
+			CircleTool->CircleRotation = CircleLayoutOptions.Rotation;
+			CircleTool->CircleScale = CircleLayoutOptions.Scale;
+		}
 
-	const bool bCenterPlane = CircleLayoutOptions.bCentered;
+		if (Layout == EAvaPatternModifierLayout::Line)
+		{
+			ActiveToolClass = UAvaPatternModifierLineTool::StaticClass();
+		}
+		else if (Layout == EAvaPatternModifierLayout::Grid)
+		{
+			ActiveToolClass = UAvaPatternModifierGridTool::StaticClass();
+		}
+		else if (Layout == EAvaPatternModifierLayout::Circle)
+		{
+			ActiveToolClass = UAvaPatternModifierCircleTool::StaticClass();
+		}
 
-	FVector Translation = FVector::ZeroVector;
-	if (CircleLayoutOptions.Plane == EAvaPatternModifierPlane::XY)
-	{
-		const FVector TranslationZero = bCenterPlane ? FVector::ZeroVector : FVector(ZeroX, ZeroY, 0);
-		Translation = FVector(X, Y, 0) - TranslationZero;
-	}
-	else if (CircleLayoutOptions.Plane == EAvaPatternModifierPlane::YZ)
-	{
-		const FVector TranslationZero = bCenterPlane ? FVector::ZeroVector : FVector(0, ZeroX, ZeroY);
-		Translation = FVector(0, X, Y) - TranslationZero;
-	}
-	else if (CircleLayoutOptions.Plane == EAvaPatternModifierPlane::ZX)
-	{
-		const FVector TranslationZero = bCenterPlane ? FVector::ZeroVector : FVector(ZeroX, 0, ZeroY);
-		Translation = FVector(X, 0, Y) - TranslationZero;
+		PRAGMA_ENABLE_DEPRECATION_WARNINGS
 	}
-	FTransform TransformChange;
-	TransformChange.SetLocation(Translation);
-	TransformChange.SetRotation(CircleLayoutOptions.Rotation.Quaternion());
-	TransformChange.SetScale3D(CircleLayoutOptions.Scale);
-	return UE::Geometry::FTransformSRT3d(TransformChange);
 }
 
 #undef LOCTEXT_NAMESPACE
diff --git a/Engine/Plugins/Experimental/Avalanche/Source/AvalancheModifiers/Private/Tools/AvaPatternModifierCircleTool.cpp b/Engine/Plugins/Experimental/Avalanche/Source/AvalancheModifiers/Private/Tools/AvaPatternModifierCircleTool.cpp
new file mode 100644
--- /dev/null
+++ b/Engine/Plugins/Experimental/Avalanche/Source/AvalancheModifiers/Private/Tools/AvaPatternModifierCircleTool.cpp
@@ -0,0 +1,169 @@
+// Copyright Epic Games, Inc. All Rights Reserved.
+
+#include "Tools/AvaPatternModifierCircleTool.h"
+
+void UAvaPatternModifierCircleTool::SetCirclePlane(EAvaPatternModifierPlane InCirclePlane)
+{
+	if (CirclePlane == InCirclePlane)
+	{
+		return;
+	}
+
+	CirclePlane = InCirclePlane;
+	OnToolPropertiesChanged();
+}
+
+void UAvaPatternModifierCircleTool::SetCircleRadius(float InCircleRadius)
+{
+	if (FMath::IsNearlyEqual(CircleRadius, InCircleRadius))
+	{
+		return;
+	}
+
+	CircleRadius = InCircleRadius;
+	OnToolPropertiesChanged();
+}
+
+void UAvaPatternModifierCircleTool::SetCircleStartAngle(float InCircleStartAngle)
+{
+	if (FMath::IsNearlyEqual(CircleStartAngle, InCircleStartAngle))
+	{
+		return;
+	}
+
+	CircleStartAngle = InCircleStartAngle;
+	OnToolPropertiesChanged();
+}
+
+void UAvaPatternModifierCircleTool::SetCircleFullAngle(float InCircleFullAngle)
+{
+	if (FMath::IsNearlyEqual(CircleFullAngle, InCircleFullAngle))
+	{
+		return;
+	}
+
+	CircleFullAngle = InCircleFullAngle;
+	OnToolPropertiesChanged();
+}
+
+void UAvaPatternModifierCircleTool::SetCircleCount(int32 InCircleCount)
+{
+	InCircleCount = FMath::Max(InCircleCount, 1);
+	if (CircleCount == InCircleCount)
+	{
+		return;
+	}
+
+	CircleCount = InCircleCount;
+	OnToolPropertiesChanged();
+}
+
+void UAvaPatternModifierCircleTool::SetCircleAccumulateTransform(bool bInCircleAccumulateTransform)
+{
+	if (bCircleAccumulateTransform == bInCircleAccumulateTransform)
+	{
+		return;
+	}
+
+	bCircleAccumulateTransform = bInCircleAccumulateTransform;
+	OnToolPropertiesChanged();
+}
+
+void UAvaPatternModifierCircleTool::SetCircleRotation(const FRotator& InCircleRotation)
+{
+	if (CircleRotation.Equals(InCircleRotation))
+	{
+		return;
+	}
+
+	CircleRotation = InCircleRotation;
+	OnToolPropertiesChanged();
+}
+
+void UAvaPatternModifierCircleTool::SetCircleScale(const FVector& InCircleScale)
+{
+	if (CircleScale.Equals(InCircleScale))
+	{
+		return;
+	}
+
+	CircleScale = InCircleScale;
+	OnToolPropertiesChanged();
+}
+
+#if WITH_EDITOR
+void UAvaPatternModifierCircleTool::PostEditChangeProperty(FPropertyChangedEvent& InEvent)
+{
+	Super::PostEditChangeProperty(InEvent);
+
+	static const TSet<FName> CirclePropertyNames
+	{
+		GET_MEMBER_NAME_CHECKED(UAvaPatternModifierCircleTool, CirclePlane),
+		GET_MEMBER_NAME_CHECKED(UAvaPatternModifierCircleTool, CircleRadius),
+		GET_MEMBER_NAME_CHECKED(UAvaPatternModifierCircleTool, CircleStartAngle),
+		GET_MEMBER_NAME_CHECKED(UAvaPatternModifierCircleTool, CircleFullAngle),
+		GET_MEMBER_NAME_CHECKED(UAvaPatternModifierCircleTool, CircleCount),
+		GET_MEMBER_NAME_CHECKED(UAvaPatternModifierCircleTool, bCircleAccumulateTransform),
+		GET_MEMBER_NAME_CHECKED(UAvaPatternModifierCircleTool, CircleRotation),
+		GET_MEMBER_NAME_CHECKED(UAvaPatternModifierCircleTool, CircleScale),
+	};
+
+	if (CirclePropertyNames.Contains(InEvent.GetMemberPropertyName()))
+	{
+		OnToolPropertiesChanged();
+	}
+}
+#endif
+
+TArray<FTransform> UAvaPatternModifierCircleTool::GetTransformInstances(const FBox& InOriginalBounds) const
+{
+	TArray<FTransform> Transforms;
+
+	const int32 Count = CircleCount;
+	const float Radius = CircleRadius;
+	const float StartAngle = FMath::DegreesToRadians(CircleStartAngle);
+	const float FullAngle = FMath::DegreesToRadians(CircleFullAngle);
+	const float AngleStep = FullAngle / Count;
+
+	const FTransform BaseTransform(CircleRotation.Quaternion(), FVector::ZeroVector, CircleScale);
+	FTransform AccumulatedTransform = BaseTransform;
+	for (int32 Idx = 0; Idx < CircleCount; Idx++)
+	{
+		const float CurrentAngle = StartAngle + Idx * AngleStep;
+		const float X = Radius * FMath::Cos(CurrentAngle);
+		const float Y = Radius * FMath::Sin(CurrentAngle);
+        
+		FVector Translation = FVector::ZeroVector;
+		if (CirclePlane == EAvaPatternModifierPlane::XY)
+		{
+			Translation = FVector(X, Y, 0);
+		}
+		else if (CirclePlane == EAvaPatternModifierPlane::YZ)
+		{
+			Translation = FVector(0, X, Y);
+		}
+		else if (CirclePlane == EAvaPatternModifierPlane::ZX)
+		{
+			Translation = FVector(X, 0, Y);
+		}
+
+		Transforms.Add(FTransform(AccumulatedTransform.GetRotation(), AccumulatedTransform.GetTranslation() + Translation, AccumulatedTransform.GetScale3D()));
+
+		if (bCircleAccumulateTransform)
+		{
+			AccumulatedTransform *= BaseTransform;
+		}
+	}
+
+	return Transforms;
+}
+
+FVector UAvaPatternModifierCircleTool::GetCenterAlignmentAxis() const
+{
+	return FVector::ZeroVector;
+}
+
+FName UAvaPatternModifierCircleTool::GetToolName() const
+{
+	return TEXT("Circle");
+}
diff --git a/Engine/Plugins/Experimental/Avalanche/Source/AvalancheModifiers/Private/Tools/AvaPatternModifierGridTool.cpp b/Engine/Plugins/Experimental/Avalanche/Source/AvalancheModifiers/Private/Tools/AvaPatternModifierGridTool.cpp
new file mode 100644
--- /dev/null
+++ b/Engine/Plugins/Experimental/Avalanche/Source/AvalancheModifiers/Private/Tools/AvaPatternModifierGridTool.cpp
@@ -0,0 +1,216 @@
+// Copyright Epic Games, Inc. All Rights Reserved.
+
+#include "Tools/AvaPatternModifierGridTool.h"
+
+void UAvaPatternModifierGridTool::SetGridPlane(EAvaPatternModifierPlane InGridPlane)
+{
+	if (GridPlane == InGridPlane)
+	{
+		return;
+	}
+
+	GridPlane = InGridPlane;
+	OnToolPropertiesChanged();
+}
+
+void UAvaPatternModifierGridTool::SetGridAlignment(EAvaPatternModifierGridAlignment InGridAlignment)
+{
+	if (GridAlignment == InGridAlignment)
+	{
+		return;
+	}
+
+	GridAlignment = InGridAlignment;
+	OnToolPropertiesChanged();
+}
+
+void UAvaPatternModifierGridTool::SetGridCountX(int32 InGridCountX)
+{
+	InGridCountX = FMath::Max(InGridCountX, 1);
+	if (GridCountX == InGridCountX)
+	{
+		return;
+	}
+
+	GridCountX = InGridCountX;
+	OnToolPropertiesChanged();
+}
+
+void UAvaPatternModifierGridTool::SetGridCountY(int32 InGridCountY)
+{
+	InGridCountY = FMath::Max(InGridCountY, 1);
+	if (GridCountY == InGridCountY)
+	{
+		return;
+	}
+
+	GridCountY = InGridCountY;
+	OnToolPropertiesChanged();
+}
+
+void UAvaPatternModifierGridTool::SetGridSpacingX(float InGridSpacingX)
+{
+	if (FMath::IsNearlyEqual(GridSpacingX, InGridSpacingX))
+	{
+		return;
+	}
+
+	GridSpacingX = InGridSpacingX;
+	OnToolPropertiesChanged();
+}
+
+void UAvaPatternModifierGridTool::SetGridSpacingY(float InGridSpacingY)
+{
+	if (FMath::IsNearlyEqual(GridSpacingY, InGridSpacingY))
+	{
+		return;
+	}
+
+	GridSpacingY = InGridSpacingY;
+	OnToolPropertiesChanged();
+}
+
+void UAvaPatternModifierGridTool::SetGridAccumulateTransform(bool bInGridAccumulateTransform)
+{
+	if (bGridAccumulateTransform == bInGridAccumulateTransform)
+	{
+		return;
+	}
+
+	bGridAccumulateTransform = bInGridAccumulateTransform;
+	OnToolPropertiesChanged();
+}
+
+void UAvaPatternModifierGridTool::SetGridRotation(const FRotator& InGridRotation)
+{
+	if (GridRotation.Equals(InGridRotation))
+	{
+		return;
+	}
+	
+	GridRotation = InGridRotation;
+	OnToolPropertiesChanged();
+}
+
+void UAvaPatternModifierGridTool::SetGridScale(const FVector& InGridScale)
+{
+	if (GridScale.Equals(InGridScale))
+	{
+		return;
+	}
+
+	GridScale = InGridScale;
+	OnToolPropertiesChanged();
+}
+
+#if WITH_EDITOR
+void UAvaPatternModifierGridTool::PostEditChangeProperty(FPropertyChangedEvent& InEvent)
+{
+	Super::PostEditChangeProperty(InEvent);
+	
+	static const TSet<FName> GridPropertyNames
+	{
+		GET_MEMBER_NAME_CHECKED(UAvaPatternModifierGridTool, GridPlane),
+		GET_MEMBER_NAME_CHECKED(UAvaPatternModifierGridTool, GridAlignment),
+		GET_MEMBER_NAME_CHECKED(UAvaPatternModifierGridTool, GridCountX),
+		GET_MEMBER_NAME_CHECKED(UAvaPatternModifierGridTool, GridCountY),
+		GET_MEMBER_NAME_CHECKED(UAvaPatternModifierGridTool, GridSpacingX),
+		GET_MEMBER_NAME_CHECKED(UAvaPatternModifierGridTool, GridSpacingY),
+		GET_MEMBER_NAME_CHECKED(UAvaPatternModifierGridTool, bGridAccumulateTransform),
+		GET_MEMBER_NAME_CHECKED(UAvaPatternModifierGridTool, GridRotation),
+		GET_MEMBER_NAME_CHECKED(UAvaPatternModifierGridTool, GridScale),
+	};
+	
+	if (GridPropertyNames.Contains(InEvent.GetMemberPropertyName()))
+	{
+		OnToolPropertiesChanged();
+	}
+}
+#endif
+
+TArray<FTransform> UAvaPatternModifierGridTool::GetTransformInstances(const FBox& InOriginalBounds) const
+{
+	TArray<FTransform> Transforms;
+
+	FVector ColAxis = (GridPlane == EAvaPatternModifierPlane::XY) ? FVector::YAxisVector :
+					  (GridPlane == EAvaPatternModifierPlane::YZ) ? FVector::ZAxisVector :
+					  FVector::XAxisVector;
+
+	FVector RowAxis = (GridPlane == EAvaPatternModifierPlane::XY) ? FVector::XAxisVector :
+					  (GridPlane == EAvaPatternModifierPlane::YZ) ? FVector::YAxisVector :
+					  FVector::ZAxisVector;
+
+	if (GridAlignment == EAvaPatternModifierGridAlignment::TopLeft || GridAlignment == EAvaPatternModifierGridAlignment::TopRight)
+	{
+		ColAxis *= -1;
+	}
+
+	if (GridAlignment == EAvaPatternModifierGridAlignment::BottomRight || GridAlignment == EAvaPatternModifierGridAlignment::TopRight)
+	{
+		RowAxis *= -1;
+	}
+
+	const FVector Size3D = InOriginalBounds.GetSize();
+	const FVector ColTranslation = ColAxis * Size3D + ColAxis * GridSpacingY;
+	const FVector RowTranslation = RowAxis * Size3D + RowAxis * GridSpacingX;
+
+	const FQuat GridQuat = GridRotation.Quaternion();
+	const FTransform BaseTransform(GridQuat, FVector::ZeroVector, GridScale);
+	FTransform AccumulatedTransform = BaseTransform;
+	for (int32 Row = 0; Row < GridCountX; Row++)
+	{
+		FTransform RowTransform = AccumulatedTransform;
+		for (int32 Col = 0; Col < GridCountY; Col++)
+		{
+			Transforms.Add(RowTransform);
+			if (bGridAccumulateTransform)
+			{
+				RowTransform *= FTransform(GridQuat, ColTranslation, GridScale);
+			}
+			else
+			{
+				RowTransform.SetLocation(RowTransform.GetLocation() + ColTranslation);
+			}
+		}
+
+		if (bGridAccumulateTransform)
+		{
+			AccumulatedTransform *= FTransform(GridQuat, RowTranslation, GridScale);
+		}
+		else
+		{
+			AccumulatedTransform.SetLocation(AccumulatedTransform.GetLocation() + RowTranslation);
+		}
+	}
+
+	return Transforms;
+}
+
+FVector UAvaPatternModifierGridTool::GetCenterAlignmentAxis() const
+{
+	FVector CenterAxis = FVector::ZeroVector;
+	if (GridAlignment == EAvaPatternModifierGridAlignment::Center)
+	{
+		if (GridPlane == EAvaPatternModifierPlane::XY)
+		{
+			CenterAxis.X = 1;
+			CenterAxis.Y = 1;
+		}
+		else if (GridPlane == EAvaPatternModifierPlane::YZ)
+		{
+			CenterAxis.Y = 1;
+			CenterAxis.Z = 1;
+		}
+		else if (GridPlane == EAvaPatternModifierPlane::ZX)
+		{
+			CenterAxis.X = 1;
+			CenterAxis.Z = 1;
+		}
+	}
+	return CenterAxis;
+}
+
+FName UAvaPatternModifierGridTool::GetToolName() const
+{
+	return TEXT("Grid");
+}
diff --git a/Engine/Plugins/Experimental/Avalanche/Source/AvalancheModifiers/Private/Tools/AvaPatternModifierLineTool.cpp b/Engine/Plugins/Experimental/Avalanche/Source/AvalancheModifiers/Private/Tools/AvaPatternModifierLineTool.cpp
new file mode 100644
--- /dev/null
+++ b/Engine/Plugins/Experimental/Avalanche/Source/AvalancheModifiers/Private/Tools/AvaPatternModifierLineTool.cpp
@@ -0,0 +1,172 @@
+// Copyright Epic Games, Inc. All Rights Reserved.
+
+#include "Tools/AvaPatternModifierLineTool.h"
+
+void UAvaPatternModifierLineTool::SetLineAxis(EAvaPatternModifierAxis InLineAxis)
+{
+	if (LineAxis == InLineAxis)
+	{
+		return;
+	}
+
+	LineAxis = InLineAxis;
+	OnToolPropertiesChanged();
+}
+
+void UAvaPatternModifierLineTool::SetLineAlignment(EAvaPatternModifierLineAlignment InLineAlignment)
+{
+	if (LineAlignment == InLineAlignment)
+	{
+		return;
+	}
+
+	LineAlignment = InLineAlignment;
+	OnToolPropertiesChanged();
+}
+
+void UAvaPatternModifierLineTool::SetLineCount(int32 InLineCount)
+{
+	InLineCount = FMath::Max(InLineCount, 1);
+	if (LineCount == InLineCount)
+	{
+		return;
+	}
+
+	LineCount = InLineCount;
+	OnToolPropertiesChanged();
+}
+
+void UAvaPatternModifierLineTool::SetLineSpacing(float InLineSpacing)
+{
+	if (FMath::IsNearlyEqual(LineSpacing, InLineSpacing))
+	{
+		return;
+	}
+
+	LineSpacing = InLineSpacing;
+	OnToolPropertiesChanged();
+}
+
+void UAvaPatternModifierLineTool::SetLineAccumulateTransform(bool bInLineAccumulateTransform)
+{
+	if (bLineAccumulateTransform == bInLineAccumulateTransform)
+	{
+		return;
+	}
+
+	bLineAccumulateTransform = bInLineAccumulateTransform;
+	OnToolPropertiesChanged();
+}
+
+void UAvaPatternModifierLineTool::SetLineRotation(const FRotator& InLineRotation)
+{
+	if (LineRotation.Equals(InLineRotation))
+	{
+		return;
+	}
+
+	LineRotation = InLineRotation;
+	OnToolPropertiesChanged();
+}
+
+void UAvaPatternModifierLineTool::SetLineScale(const FVector& InLineScale)
+{
+	if (LineScale.Equals(InLineScale))
+	{
+		return;
+	}
+
+	LineScale = InLineScale;
+	OnToolPropertiesChanged();
+}
+
+#if WITH_EDITOR
+void UAvaPatternModifierLineTool::PostEditChangeProperty(FPropertyChangedEvent& InEvent)
+{
+	Super::PostEditChangeProperty(InEvent);
+	
+	static const TSet<FName> LinePropertyNames
+	{
+		GET_MEMBER_NAME_CHECKED(UAvaPatternModifierLineTool, LineAxis),
+		GET_MEMBER_NAME_CHECKED(UAvaPatternModifierLineTool, LineAlignment),
+		GET_MEMBER_NAME_CHECKED(UAvaPatternModifierLineTool, LineCount),
+		GET_MEMBER_NAME_CHECKED(UAvaPatternModifierLineTool, LineSpacing),
+		GET_MEMBER_NAME_CHECKED(UAvaPatternModifierLineTool, bLineAccumulateTransform),
+		GET_MEMBER_NAME_CHECKED(UAvaPatternModifierLineTool, LineRotation),
+		GET_MEMBER_NAME_CHECKED(UAvaPatternModifierLineTool, LineScale),
+	};
+
+	if (LinePropertyNames.Contains(InEvent.GetMemberPropertyName()))
+	{
+		OnToolPropertiesChanged();
+	}
+}
+#endif
+
+TArray<FTransform> UAvaPatternModifierLineTool::GetTransformInstances(const FBox& InOriginalBounds) const
+{
+	TArray<FTransform> Transforms;
+
+	FVector Axis = FVector::ZeroVector;
+	switch (LineAxis)
+	{
+		case EAvaPatternModifierAxis::X:
+			Axis = FVector::XAxisVector;
+			break;
+		case EAvaPatternModifierAxis::Y:
+			Axis = FVector::YAxisVector;
+			break;
+		case EAvaPatternModifierAxis::Z:
+			Axis = FVector::ZAxisVector;
+			break;
+	};
+
+	Axis *= (LineAlignment == EAvaPatternModifierLineAlignment::End ? -1 : 1);
+
+	const FVector Size3D = InOriginalBounds.GetSize();
+	const FVector Translation = Axis * Size3D + Axis * LineSpacing;
+
+	const FQuat LineQuat = LineRotation.Quaternion();
+	FTransform AccumulatedTransform(LineQuat, FVector::ZeroVector, LineScale);
+	for (int32 Idx = 0; Idx < LineCount; Idx++)
+	{
+		Transforms.Add(AccumulatedTransform);
+        
+		if (bLineAccumulateTransform)
+		{
+			AccumulatedTransform *= FTransform(LineQuat, Translation, LineScale);
+		}
+		else
+		{
+			AccumulatedTransform.SetLocation(AccumulatedTransform.GetLocation() + Translation);
+		}
+	}
+
+	return Transforms;
+}
+
+FVector UAvaPatternModifierLineTool::GetCenterAlignmentAxis() const
+{
+	FVector CenterAxis = FVector::ZeroVector;
+	if (LineAlignment == EAvaPatternModifierLineAlignment::Center)
+	{
+		if (LineAxis == EAvaPatternModifierAxis::X)
+		{
+			CenterAxis.X = 1;
+		}
+		else if (LineAxis == EAvaPatternModifierAxis::Y)
+		{
+			CenterAxis.Y = 1;
+		}
+		else if (LineAxis == EAvaPatternModifierAxis::Z)
+		{
+			CenterAxis.Z = 1;
+		}
+	}
+	return CenterAxis;
+}
+
+FName UAvaPatternModifierLineTool::GetToolName() const
+{
+	return TEXT("Line");
+}
diff --git a/Engine/Plugins/Experimental/Avalanche/Source/AvalancheModifiers/Private/Tools/AvaPatternModifierTool.cpp b/Engine/Plugins/Experimental/Avalanche/Source/AvalancheModifiers/Private/Tools/AvaPatternModifierTool.cpp
new file mode 100644
--- /dev/null
+++ b/Engine/Plugins/Experimental/Avalanche/Source/AvalancheModifiers/Private/Tools/AvaPatternModifierTool.cpp
@@ -0,0 +1,16 @@
+// Copyright Epic Games, Inc. All Rights Reserved.
+
+#include "Tools/AvaPatternModifierTool.h"
+
+#include "Modifiers/AvaPatternModifier.h"
+
+void UAvaPatternModifierTool::OnToolPropertiesChanged() const
+{
+	if (UAvaPatternModifier* PatternModifier = GetTypedOuter<UAvaPatternModifier>())
+	{
+		if (PatternModifier->GetActiveToolClass() == GetClass())
+		{
+			PatternModifier->MarkModifierDirty();
+		}
+	}
+}
diff --git a/Engine/Plugins/Experimental/Avalanche/Source/AvalancheModifiers/Private/Tools/AvaPatternModifierTool.h b/Engine/Plugins/Experimental/Avalanche/Source/AvalancheModifiers/Private/Tools/AvaPatternModifierTool.h
new file mode 100644
--- /dev/null
+++ b/Engine/Plugins/Experimental/Avalanche/Source/AvalancheModifiers/Private/Tools/AvaPatternModifierTool.h
@@ -0,0 +1,20 @@
+// Copyright Epic Games, Inc. All Rights Reserved.
+
+#pragma once
+
+#include "AvaPatternModifierTool.generated.h"
+
+/** Represent a tool supported by the pattern modifier */
+UCLASS(Abstract)
+class UAvaPatternModifierTool : public UObject
+{
+	GENERATED_BODY()
+
+public:
+	virtual TArray<FTransform> GetTransformInstances(const FBox& InOriginalBounds) const PURE_VIRTUAL(UAvaPatternModifierTool::GetTransformInstances, return {}; );
+	virtual FVector GetCenterAlignmentAxis() const PURE_VIRTUAL(UAvaPatternModifierTool::GetCenterAlignmentAxis, return FVector::ZeroVector; );
+	virtual FName GetToolName() const PURE_VIRTUAL(UAvaPatternModifierTool::GetToolName, return NAME_None; );
+
+protected:
+	void OnToolPropertiesChanged() const;
+};
diff --git a/Engine/Plugins/Experimental/Avalanche/Source/AvalancheModifiers/Public/Modifiers/AvaPatternModifier.h b/Engine/Plugins/Experimental/Avalanche/Source/AvalancheModifiers/Public/Modifiers/AvaPatternModifier.h
--- a/Engine/Plugins/Experimental/Avalanche/Source/AvalancheModifiers/Public/Modifiers/AvaPatternModifier.h
+++ b/Engine/Plugins/Experimental/Avalanche/Source/AvalancheModifiers/Public/Modifiers/AvaPatternModifier.h
@@ -3,6 +3,7 @@
 #pragma once
 
 #include "AvaGeometryBaseModifier.h"
+#include "Tools/AvaPatternModifierTool.h"
 #include "TransformTypes.h"
 #include "AvaPatternModifier.generated.h"
 
@@ -30,20 +31,26 @@ enum class EAvaPatternModifierPlane : uint8
 	YZ = 2
 };
 
-USTRUCT(BlueprintType)
-struct FVector2b
+UENUM(BlueprintType)
+enum class EAvaPatternModifierLineAlignment : uint8
 {
-	GENERATED_BODY()
-
-	UPROPERTY(EditInstanceOnly, BlueprintReadWrite, Category = "Values")
-	bool bX = false;
+	Start,
+	Center,
+	End
+};
 
-	UPROPERTY(EditInstanceOnly, BlueprintReadWrite, Category = "Values")
-	bool bY = false;
+UENUM(BlueprintType)
+enum class EAvaPatternModifierGridAlignment : uint8
+{
+	TopLeft,
+	TopRight,
+	Center,
+	BottomLeft,
+	BottomRight
 };
 
 USTRUCT(BlueprintType)
-struct FVector3b
+struct FVector2b
 {
 	GENERATED_BODY()
 
@@ -52,9 +59,6 @@ struct FVector3b
 
 	UPROPERTY(EditInstanceOnly, BlueprintReadWrite, Category = "Values")
 	bool bY = false;
-
-	UPROPERTY(EditInstanceOnly, BlueprintReadWrite, Category = "Values")
-	bool bZ = false;
 };
 
 USTRUCT(BlueprintType)
@@ -160,77 +164,68 @@ class UAvaPatternModifier : public UAvaGeometryBaseModifier
 	GENERATED_BODY()
 
 public:
-	UFUNCTION(BlueprintCallable, Category="Motion Design|Modifiers|Pattern")
-	AVALANCHEMODIFIERS_API void SetLayout(EAvaPatternModifierLayout InLayout);
-
-	UFUNCTION(BlueprintPure, Category="Motion Design|Modifiers|Pattern")
-	EAvaPatternModifierLayout GetLayout() const
-	{
-		return Layout;
-	}
+	UAvaPatternModifier();
 
 	UFUNCTION(BlueprintCallable, Category="Motion Design|Modifiers|Pattern")
-	AVALANCHEMODIFIERS_API void SetLineLayoutOptions(const FAvaPatternModifierLineLayoutOptions& InOptions);
+	void SetActiveToolClass(const TSubclassOf<UAvaPatternModifierTool>& InClass);
 
 	UFUNCTION(BlueprintPure, Category="Motion Design|Modifiers|Pattern")
-	const FAvaPatternModifierLineLayoutOptions& GetLineLayoutOptions() const
+	TSubclassOf<UAvaPatternModifierTool> GetActiveToolClass() const
 	{
-		return LineLayoutOptions;
-	}
-
-	UFUNCTION(BlueprintCallable, Category="Motion Design|Modifiers|Pattern")
-	AVALANCHEMODIFIERS_API void SetGridLayoutOptions(const FAvaPatternModifierGridLayoutOptions& InOptions);
-
-	UFUNCTION(BlueprintPure, Category="Motion Design|Modifiers|Pattern")
-	const FAvaPatternModifierGridLayoutOptions& GetGridLayoutOptions() const
-	{
-		return GridLayoutOptions;
-	}
-
-	UFUNCTION(BlueprintCallable, Category="Motion Design|Modifiers|Pattern")
-	AVALANCHEMODIFIERS_API void SetCircleLayoutOptions(const FAvaPatternModifierCircleLayoutOptions& InOptions);
-
-	UFUNCTION(BlueprintPure, Category="Motion Design|Modifiers|Pattern")
-	const FAvaPatternModifierCircleLayoutOptions& GetCircleLayoutOptions() const
-	{
-		return CircleLayoutOptions;
+		return ActiveToolClass;
 	}
 
 protected:
 	//~ Begin UObject
+	virtual void Serialize(FArchive& InArchive) override;
 #if WITH_EDITOR
 	virtual void PostEditChangeProperty(FPropertyChangedEvent& PropertyChangedEvent) override;
 #endif
 	//~ End UObject
 
 	//~ Begin UActorModifierCoreBase
+	virtual void OnModifierEnabled(EActorModifierCoreEnableReason InReason) override;
 	virtual void OnModifierCDOSetup(FActorModifierCoreMetadata& InMetadata) override;
 	virtual void Apply() override;
 	//~ End UActorModifierCoreBase
 
-	void OnLayoutChanged();
-	void OnLineLayoutOptionsChanged();
-	void OnGridLayoutOptionsChanged();
-	void OnCircleLayoutOptionsChanged();
+	/** Finds saved tool or create and saves it */
+	UAvaPatternModifierTool* FindOrAddTool(TSubclassOf<UAvaPatternModifierTool> InToolClass);
+
+	template<typename InToolClass>
+	InToolClass* FindOrAddTool()
+	{
+		return Cast<InToolClass>(FindOrAddTool(InToolClass::StaticClass()));
+	}
+
+	void OnActiveToolClassChanged();
+
+	UPROPERTY(EditInstanceOnly, Setter, Getter, Category="Pattern", NoClear, meta=(ShowDisplayNames, AllowPrivateAccess="true"))
+	TSubclassOf<UAvaPatternModifierTool> ActiveToolClass;
+
+	UPROPERTY(VisibleInstanceOnly, Instanced, Category="Pattern", Transient, TextExportTransient, DuplicateTransient)
+	TObjectPtr<UAvaPatternModifierTool> ActiveTool;
+
+	UPROPERTY()
+	TArray<TObjectPtr<UAvaPatternModifierTool>> Tools;
 
-	UE::Geometry::FTransformSRT3d GetLineLayoutTransformChange() const;
-	UE::Geometry::FTransformSRT3d GetGridLayoutColTransformChange() const;
-	UE::Geometry::FTransformSRT3d GetGridLayoutRowTransformChange() const;
-	UE::Geometry::FTransformSRT3d GetCircleLayoutTransformChange(int32 Idx) const;
+private:
+	void MigrateVersion(int32 InCurrentVersion);
 
-	UPROPERTY(EditInstanceOnly, Setter="SetLayout", Getter="GetLayout", Category="Pattern", meta=(AllowPrivateAccess="true"))
+	UE_DEPRECATED(5.6, "Moved to ToolName")
+	UPROPERTY()
 	EAvaPatternModifierLayout Layout = EAvaPatternModifierLayout::Line;
 
-	/** Line layout options */
-	UPROPERTY(EditInstanceOnly, Setter="SetLineLayoutOptions", Getter="GetLineLayoutOptions", Category="Pattern", meta=(EditCondition="Layout == EAvaPatternModifierLayout::Line", EditConditionHides, AllowPrivateAccess="true"))
+	UE_DEPRECATED(5.6, "Moved to Line Tool")
+	UPROPERTY()
 	FAvaPatternModifierLineLayoutOptions LineLayoutOptions;
 
-	/** Grid layout options */
-	UPROPERTY(EditInstanceOnly, Setter="SetGridLayoutOptions", Getter="GetGridLayoutOptions", Category="Pattern", meta=(EditCondition="Layout == EAvaPatternModifierLayout::Grid", EditConditionHides, AllowPrivateAccess="true"))
+	UE_DEPRECATED(5.6, "Moved to Grid Tool")
+	UPROPERTY()
 	FAvaPatternModifierGridLayoutOptions GridLayoutOptions;
 
-	/** Circle layout options */
-	UPROPERTY(EditInstanceOnly, Setter="SetCircleLayoutOptions", Getter="GetCircleLayoutOptions", Category="Pattern", meta=(EditCondition="Layout == EAvaPatternModifierLayout::Circle", EditConditionHides, AllowPrivateAccess="true"))
+	UE_DEPRECATED(5.6, "Moved to Circle Tool")
+	UPROPERTY()
 	FAvaPatternModifierCircleLayoutOptions CircleLayoutOptions;
 
 	UPROPERTY(Transient, DuplicateTransient, TextExportTransient)
diff --git a/Engine/Plugins/Experimental/Avalanche/Source/AvalancheModifiers/Public/Tools/AvaPatternModifierCircleTool.h b/Engine/Plugins/Experimental/Avalanche/Source/AvalancheModifiers/Public/Tools/AvaPatternModifierCircleTool.h
new file mode 100644
--- /dev/null
+++ b/Engine/Plugins/Experimental/Avalanche/Source/AvalancheModifiers/Public/Tools/AvaPatternModifierCircleTool.h
@@ -0,0 +1,126 @@
+// Copyright Epic Games, Inc. All Rights Reserved.
+
+#pragma once
+
+#include "Tools/AvaPatternModifierTool.h"
+#include "Modifiers/AvaPatternModifier.h"
+#include "AvaPatternModifierCircleTool.generated.h"
+
+/** Circle tool for the pattern modifier */
+UCLASS(DisplayName="Circle", ClassGroup="Modifiers")
+class UAvaPatternModifierCircleTool : public UAvaPatternModifierTool
+{
+	GENERATED_BODY()
+
+	friend class UAvaPatternModifier;
+
+public:
+	UFUNCTION(BlueprintPure, Category="Motion Design|Modifiers|Pattern")
+	EAvaPatternModifierPlane GetCirclePlane() const
+	{
+		return CirclePlane;
+	}
+
+	UFUNCTION(BlueprintCallable, Category="Motion Design|Modifiers|Pattern")
+	void SetCirclePlane(EAvaPatternModifierPlane InCirclePlane);
+
+	UFUNCTION(BlueprintPure, Category="Motion Design|Modifiers|Pattern")
+	float GetCircleRadius() const
+	{
+		return CircleRadius;
+	}
+
+	UFUNCTION(BlueprintCallable, Category="Motion Design|Modifiers|Pattern")
+	void SetCircleRadius(float InCircleRadius);
+
+	UFUNCTION(BlueprintPure, Category="Motion Design|Modifiers|Pattern")
+	float GetCircleStartAngle() const
+	{
+		return CircleStartAngle;
+	}
+
+	UFUNCTION(BlueprintCallable, Category="Motion Design|Modifiers|Pattern")
+	void SetCircleStartAngle(float InCircleStartAngle);
+
+	UFUNCTION(BlueprintPure, Category="Motion Design|Modifiers|Pattern")
+	float GetCircleFullAngle() const
+	{
+		return CircleFullAngle;
+	}
+
+	UFUNCTION(BlueprintCallable, Category="Motion Design|Modifiers|Pattern")
+	void SetCircleFullAngle(float InCircleFullAngle);
+
+	UFUNCTION(BlueprintPure, Category="Motion Design|Modifiers|Pattern")
+	int32 GetCircleCount() const
+	{
+		return CircleCount;
+	}
+
+	UFUNCTION(BlueprintCallable, Category="Motion Design|Modifiers|Pattern")
+	void SetCircleCount(int32 InCircleCount);
+
+	UFUNCTION(BlueprintPure, Category="Motion Design|Modifiers|Pattern")
+	bool GetCircleAccumulateTransform() const
+	{
+		return bCircleAccumulateTransform;
+	}
+
+	UFUNCTION(BlueprintCallable, Category="Motion Design|Modifiers|Pattern")
+	void SetCircleAccumulateTransform(bool bInCircleAccumulateTransform);
+
+	UFUNCTION(BlueprintPure, Category="Motion Design|Modifiers|Pattern")
+	FRotator GetCircleRotation() const
+	{
+		return CircleRotation;
+	}
+
+	UFUNCTION(BlueprintCallable, Category="Motion Design|Modifiers|Pattern")
+	void SetCircleRotation(const FRotator& InCircleRotation);
+
+	UFUNCTION(BlueprintPure, Category="Motion Design|Modifiers|Pattern")
+	FVector GetCircleScale() const
+	{
+		return CircleScale;
+	}
+
+	UFUNCTION(BlueprintCallable, Category="Motion Design|Modifiers|Pattern")
+	void SetCircleScale(const FVector& InCircleScale);
+
+protected:
+	//~ Begin UObject
+#if WITH_EDITOR
+	virtual void PostEditChangeProperty(FPropertyChangedEvent& InEvent) override;
+#endif
+	//~ End UObject
+
+	//~ Begin UAvaPatternModifierTool
+	virtual TArray<FTransform> GetTransformInstances(const FBox& InOriginalBounds) const;
+	virtual FVector GetCenterAlignmentAxis() const override;
+	virtual FName GetToolName() const override;
+	//~ End UAvaPatternModifierTool
+
+	UPROPERTY(EditInstanceOnly, Setter, Getter, Category="Pattern", meta=(EditCondition="Layout == EAvaPatternModifierLayout::Circle", EditConditionHides))
+	EAvaPatternModifierPlane CirclePlane = EAvaPatternModifierPlane::YZ;
+
+	UPROPERTY(EditInstanceOnly, Setter, Getter, Category="Pattern", meta=(EditCondition="Layout == EAvaPatternModifierLayout::Circle", EditConditionHides))
+	float CircleRadius = 100.f;
+
+	UPROPERTY(EditInstanceOnly, Setter, Getter, Category="Pattern", meta=(EditCondition="Layout == EAvaPatternModifierLayout::Circle", EditConditionHides))
+	float CircleStartAngle = 180.f;
+
+	UPROPERTY(EditInstanceOnly, Setter, Getter, Category="Pattern", meta=(EditCondition="Layout == EAvaPatternModifierLayout::Circle", EditConditionHides))
+	float CircleFullAngle = 360.f;
+
+	UPROPERTY(EditInstanceOnly, Setter, Getter, Category="Pattern", meta=(ClampMin="1", EditCondition="Layout == EAvaPatternModifierLayout::Circle", EditConditionHides))
+	int32 CircleCount = 4;
+
+	UPROPERTY(EditInstanceOnly, Setter="SetCircleAccumulateTransform", Getter="GetCircleAccumulateTransform", Category="Pattern", meta=(EditCondition="Layout == EAvaPatternModifierLayout::Circle", EditConditionHides))
+	bool bCircleAccumulateTransform = false;
+
+	UPROPERTY(EditInstanceOnly, Setter, Getter, Category="Pattern", meta=(EditCondition="Layout == EAvaPatternModifierLayout::Circle", EditConditionHides))
+	FRotator CircleRotation = FRotator::ZeroRotator;
+
+	UPROPERTY(EditInstanceOnly, Setter, Getter, Category="Pattern", meta=(EditCondition="Layout == EAvaPatternModifierLayout::Circle", EditConditionHides, MotionDesignVectorWidget, AllowPreserveRatio="XYZ"))
+	FVector CircleScale = FVector::OneVector;
+};
diff --git a/Engine/Plugins/Experimental/Avalanche/Source/AvalancheModifiers/Public/Tools/AvaPatternModifierGridTool.h b/Engine/Plugins/Experimental/Avalanche/Source/AvalancheModifiers/Public/Tools/AvaPatternModifierGridTool.h
new file mode 100644
--- /dev/null
+++ b/Engine/Plugins/Experimental/Avalanche/Source/AvalancheModifiers/Public/Tools/AvaPatternModifierGridTool.h
@@ -0,0 +1,138 @@
+// Copyright Epic Games, Inc. All Rights Reserved.
+
+#pragma once
+
+#include "Tools/AvaPatternModifierTool.h"
+#include "Modifiers/AvaPatternModifier.h"
+#include "AvaPatternModifierGridTool.generated.h"
+
+/** Grid tool for the pattern modifier */
+UCLASS(DisplayName="Grid", ClassGroup="Modifiers")
+class UAvaPatternModifierGridTool : public UAvaPatternModifierTool
+{
+	GENERATED_BODY()
+
+	friend class UAvaPatternModifier;
+
+public:
+	UFUNCTION(BlueprintPure, Category="Motion Design|Modifiers|Pattern")
+	EAvaPatternModifierPlane GetGridPlane() const
+	{
+		return GridPlane;
+	}
+
+	UFUNCTION(BlueprintCallable, Category="Motion Design|Modifiers|Pattern")
+	void SetGridPlane(EAvaPatternModifierPlane InGridPlane);
+
+	UFUNCTION(BlueprintPure, Category="Motion Design|Modifiers|Pattern")
+	EAvaPatternModifierGridAlignment GetGridAlignment() const
+	{
+		return GridAlignment;
+	}
+
+	UFUNCTION(BlueprintCallable, Category="Motion Design|Modifiers|Pattern")
+	void SetGridAlignment(EAvaPatternModifierGridAlignment InGridAlignment);
+
+	UFUNCTION(BlueprintPure, Category="Motion Design|Modifiers|Pattern")
+	int32 GetGridCountX() const
+	{
+		return GridCountX;
+	}
+
+	UFUNCTION(BlueprintCallable, Category="Motion Design|Modifiers|Pattern")
+	void SetGridCountX(int32 InGridCountX);
+
+	UFUNCTION(BlueprintPure, Category="Motion Design|Modifiers|Pattern")
+	int32 GetGridCountY() const
+	{
+		return GridCountY;
+	}
+
+	UFUNCTION(BlueprintCallable, Category="Motion Design|Modifiers|Pattern")
+	void SetGridCountY(int32 InGridCountY);
+
+	UFUNCTION(BlueprintPure, Category="Motion Design|Modifiers|Pattern")
+	float GetGridSpacingX() const
+	{
+		return GridSpacingX;
+	}
+
+	UFUNCTION(BlueprintCallable, Category="Motion Design|Modifiers|Pattern")
+	void SetGridSpacingX(float InGridSpacingX);
+	
+	UFUNCTION(BlueprintPure, Category="Motion Design|Modifiers|Pattern")
+	float GetGridSpacingY() const
+	{
+		return GridSpacingY;
+	}
+
+	UFUNCTION(BlueprintCallable, Category="Motion Design|Modifiers|Pattern")
+	void SetGridSpacingY(float InGridSpacingY);
+
+	UFUNCTION(BlueprintPure, Category="Motion Design|Modifiers|Pattern")
+	bool GetGridAccumulateTransform() const
+	{
+		return bGridAccumulateTransform;
+	}
+
+	UFUNCTION(BlueprintCallable, Category="Motion Design|Modifiers|Pattern")
+	void SetGridAccumulateTransform(bool bInGridAccumulateTransform);
+
+	UFUNCTION(BlueprintPure, Category="Motion Design|Modifiers|Pattern")
+	FRotator GetGridRotation() const
+	{
+		return GridRotation;
+	}
+
+	UFUNCTION(BlueprintCallable, Category="Motion Design|Modifiers|Pattern")
+	void SetGridRotation(const FRotator& InGridRotation);
+
+	UFUNCTION(BlueprintPure, Category="Motion Design|Modifiers|Pattern")
+	FVector GetGridScale() const
+	{
+		return GridScale;
+	}
+
+	UFUNCTION(BlueprintCallable, Category="Motion Design|Modifiers|Pattern")
+	void SetGridScale(const FVector& InGridScale);
+
+protected:
+	//~ Begin UObject
+#if WITH_EDITOR
+	virtual void PostEditChangeProperty(FPropertyChangedEvent& InEvent) override;
+#endif
+	//~ End UObject
+
+	//~ Begin UAvaPatternModifierTool
+	virtual TArray<FTransform> GetTransformInstances(const FBox& InOriginalBounds) const;
+	virtual FVector GetCenterAlignmentAxis() const override;
+	virtual FName GetToolName() const override;
+	//~ End UAvaPatternModifierTool
+
+	UPROPERTY(EditInstanceOnly, Setter, Getter, Category="Pattern", meta=(EditCondition="Layout == EAvaPatternModifierLayout::Grid", EditConditionHides))
+	EAvaPatternModifierPlane GridPlane = EAvaPatternModifierPlane::YZ;
+
+	UPROPERTY(EditInstanceOnly, Setter, Getter, Category="Pattern", meta=(EditCondition="Layout == EAvaPatternModifierLayout::Grid", EditConditionHides))
+	EAvaPatternModifierGridAlignment GridAlignment = EAvaPatternModifierGridAlignment::Center;
+
+	UPROPERTY(EditInstanceOnly, Setter, Getter, Category="Pattern", meta=(ClampMin="1", EditCondition="Layout == EAvaPatternModifierLayout::Grid", EditConditionHides))
+	int32 GridCountX = 2;
+
+	UPROPERTY(EditInstanceOnly, Setter, Getter, Category="Pattern", meta=(ClampMin="1", EditCondition="Layout == EAvaPatternModifierLayout::Grid", EditConditionHides))
+	int32 GridCountY = 2;
+
+	UPROPERTY(EditInstanceOnly, Setter, Getter, Category="Pattern", meta=(EditCondition="Layout == EAvaPatternModifierLayout::Grid", EditConditionHides))
+	float GridSpacingX = 0.f;
+
+	UPROPERTY(EditInstanceOnly, Setter, Getter, Category="Pattern", meta=(EditCondition="Layout == EAvaPatternModifierLayout::Grid", EditConditionHides))
+	float GridSpacingY = 0.f;
+
+	UPROPERTY(EditInstanceOnly, Setter="SetGridAccumulateTransform", Getter="GetGridAccumulateTransform", Category="Pattern", meta=(EditCondition="Layout == EAvaPatternModifierLayout::Grid", EditConditionHides))
+	bool bGridAccumulateTransform = false;
+
+	UPROPERTY(EditInstanceOnly, Setter, Getter, Category="Pattern", meta=(EditCondition="Layout == EAvaPatternModifierLayout::Grid", EditConditionHides))
+	FRotator GridRotation = FRotator::ZeroRotator;
+
+	UPROPERTY(EditInstanceOnly, Setter, Getter, Category="Pattern", meta=(EditCondition="Layout == EAvaPatternModifierLayout::Grid", EditConditionHides, MotionDesignVectorWidget, AllowPreserveRatio="XYZ"))
+	FVector GridScale = FVector::OneVector;
+};
\ No newline at end of file
diff --git a/Engine/Plugins/Experimental/Avalanche/Source/AvalancheModifiers/Public/Tools/AvaPatternModifierLineTool.h b/Engine/Plugins/Experimental/Avalanche/Source/AvalancheModifiers/Public/Tools/AvaPatternModifierLineTool.h
new file mode 100644
--- /dev/null
+++ b/Engine/Plugins/Experimental/Avalanche/Source/AvalancheModifiers/Public/Tools/AvaPatternModifierLineTool.h
@@ -0,0 +1,114 @@
+// Copyright Epic Games, Inc. All Rights Reserved.
+
+#pragma once
+
+#include "Tools/AvaPatternModifierTool.h"
+#include "Modifiers/AvaPatternModifier.h"
+#include "AvaPatternModifierLineTool.generated.h"
+
+/** Line tool for the pattern modifier */
+UCLASS(DisplayName="Line", ClassGroup="Modifiers")
+class UAvaPatternModifierLineTool : public UAvaPatternModifierTool
+{
+	GENERATED_BODY()
+
+	friend class UAvaPatternModifier;
+
+public:
+	UFUNCTION(BlueprintPure, Category="Motion Design|Modifiers|Pattern")
+	EAvaPatternModifierAxis GetLineAxis() const
+	{
+		return LineAxis;
+	}
+
+	UFUNCTION(BlueprintCallable, Category="Motion Design|Modifiers|Pattern")
+	void SetLineAxis(EAvaPatternModifierAxis InLineAxis);
+
+	UFUNCTION(BlueprintPure, Category="Motion Design|Modifiers|Pattern")
+	EAvaPatternModifierLineAlignment GetLineAlignment() const
+	{
+		return LineAlignment;
+	}
+
+	UFUNCTION(BlueprintCallable, Category="Motion Design|Modifiers|Pattern")
+	void SetLineAlignment(EAvaPatternModifierLineAlignment InLineAlignment);
+
+	UFUNCTION(BlueprintPure, Category="Motion Design|Modifiers|Pattern")
+	int32 GetLineCount() const
+	{
+		return LineCount;
+	}
+
+	UFUNCTION(BlueprintCallable, Category="Motion Design|Modifiers|Pattern")
+	void SetLineCount(int32 InLineCount);
+
+	UFUNCTION(BlueprintPure, Category="Motion Design|Modifiers|Pattern")
+	float GetLineSpacing() const
+	{
+		return LineSpacing;
+	}
+
+	UFUNCTION(BlueprintCallable, Category="Motion Design|Modifiers|Pattern")
+	void SetLineSpacing(float InLineSpacing);
+
+	UFUNCTION(BlueprintPure, Category="Motion Design|Modifiers|Pattern")
+	bool GetLineAccumulateTransform() const
+	{
+		return bLineAccumulateTransform;
+	}
+
+	UFUNCTION(BlueprintCallable, Category="Motion Design|Modifiers|Pattern")
+	void SetLineAccumulateTransform(bool bInLineAccumulateTransform);
+
+	UFUNCTION(BlueprintPure, Category="Motion Design|Modifiers|Pattern")
+	FRotator GetLineRotation() const
+	{
+		return LineRotation;
+	}
+
+	UFUNCTION(BlueprintCallable, Category="Motion Design|Modifiers|Pattern")
+	void SetLineRotation(const FRotator& InLineRotation);
+
+	UFUNCTION(BlueprintPure, Category="Motion Design|Modifiers|Pattern")
+	FVector GetLineScale() const
+	{
+		return LineScale;
+	}
+
+	UFUNCTION(BlueprintCallable, Category="Motion Design|Modifiers|Pattern")
+	void SetLineScale(const FVector& InLineScale);
+
+protected:
+	//~ Begin UObject
+#if WITH_EDITOR
+	virtual void PostEditChangeProperty(FPropertyChangedEvent& InEvent) override;
+#endif
+	//~ End UObject
+
+	//~ Begin UAvaPatternModifierTool
+	virtual TArray<FTransform> GetTransformInstances(const FBox& InOriginalBounds) const;
+	virtual FVector GetCenterAlignmentAxis() const override;
+	virtual FName GetToolName() const override;
+	//~ End UAvaPatternModifierTool
+
+	UPROPERTY(EditInstanceOnly, Setter, Getter, Category="Pattern", meta=(EditCondition="Layout == EAvaPatternModifierLayout::Line", EditConditionHides))
+	EAvaPatternModifierAxis LineAxis = EAvaPatternModifierAxis::Y;
+
+	UPROPERTY(EditInstanceOnly, Setter, Getter, Category="Pattern", meta=(EditCondition="Layout == EAvaPatternModifierLayout::Line", EditConditionHides))
+	EAvaPatternModifierLineAlignment LineAlignment = EAvaPatternModifierLineAlignment::Center;
+
+	UPROPERTY(EditInstanceOnly, Setter, Getter, Category="Pattern", meta=(ClampMin="1", EditCondition="Layout == EAvaPatternModifierLayout::Line", EditConditionHides))
+	int32 LineCount = 4;
+
+	UPROPERTY(EditInstanceOnly, Setter, Getter, Category="Pattern", meta=(EditCondition="Layout == EAvaPatternModifierLayout::Line", EditConditionHides))
+	float LineSpacing = 0.f;
+
+	UPROPERTY(EditInstanceOnly, Setter="SetLineAccumulateTransform", Getter="GetLineAccumulateTransform", Category="Pattern", meta=(EditCondition="Layout == EAvaPatternModifierLayout::Line", EditConditionHides))
+	bool bLineAccumulateTransform = false;
+
+	UPROPERTY(EditInstanceOnly, Setter, Getter, Category="Pattern", meta=(EditCondition="Layout == EAvaPatternModifierLayout::Line", EditConditionHides))
+	FRotator LineRotation = FRotator::ZeroRotator;
+
+	UPROPERTY(EditInstanceOnly, Setter, Getter, Category="Pattern", meta=(EditCondition="Layout == EAvaPatternModifierLayout::Line", EditConditionHides, MotionDesignVectorWidget, AllowPreserveRatio="XYZ"))
+	FVector LineScale = FVector::OneVector;
+};
