[AvalancheMaskEditor] Adding MediaPlate Material templates for mask modifier.

#virtualized

diff --git a/Engine/Plugins/Experimental/Avalanche/Source/AvalancheMaskEditor/AvalancheMaskEditor.Build.cs b/Engine/Plugins/Experimental/Avalanche/Source/AvalancheMaskEditor/AvalancheMaskEditor.Build.cs
--- a/Engine/Plugins/Experimental/Avalanche/Source/AvalancheMaskEditor/AvalancheMaskEditor.Build.cs
+++ b/Engine/Plugins/Experimental/Avalanche/Source/AvalancheMaskEditor/AvalancheMaskEditor.Build.cs
@@ -30,6 +30,7 @@ public class AvalancheMaskEditor : ModuleRules
 				"GeometryMaskEditor",
 				"InputCore",
 				"LevelEditor",
+				"MediaPlateEditor",
 				"Projects",
 				"Slate",
 				"SlateCore",
diff --git a/Engine/Plugins/Experimental/Avalanche/Source/AvalancheMaskEditor/Private/AvaMaskEditorModule.cpp b/Engine/Plugins/Experimental/Avalanche/Source/AvalancheMaskEditor/Private/AvaMaskEditorModule.cpp
--- a/Engine/Plugins/Experimental/Avalanche/Source/AvalancheMaskEditor/Private/AvaMaskEditorModule.cpp
+++ b/Engine/Plugins/Experimental/Avalanche/Source/AvalancheMaskEditor/Private/AvaMaskEditorModule.cpp
@@ -9,6 +9,7 @@
 #include "EditorModeManager.h"
 #include "Mask2D/AvaMask2DReadModifier.h"
 #include "Mask2D/AvaMask2DWriteModifier.h"
+#include "MediaPlateEditorModule.h"
 #include "PropertyEditorModule.h"
 #include "Templates/SharedPointer.h"
 #include "ToolMenuEntry.h"
@@ -32,6 +33,11 @@ void FAvalancheMaskEditorModule::StartupModule()
 		PropertyModule.NotifyCustomizationModuleChanged();
 	}
 
+	if (FMediaPlateEditorModule* EditorModule = FModuleManager::LoadModulePtr<FMediaPlateEditorModule>("MediaPlateEditor"))
+	{
+		EditorModule->OnGetMediaPlateMaterialAssetPaths().AddRaw(this, &FAvalancheMaskEditorModule::OnGetMediaPlateMaterialAssetPaths);
+	}
+
 	FAvaMaskEditorCommands::Register();
 
 	CommandList = MakeShared<FUICommandList>();
@@ -51,6 +57,11 @@ void FAvalancheMaskEditorModule::ShutdownModule()
 			PropertyModule->NotifyCustomizationModuleChanged();
 		}
 	}
+
+	if (FMediaPlateEditorModule* EditorModule = FModuleManager::GetModulePtr<FMediaPlateEditorModule>("MediaPlateEditor"))
+	{
+		EditorModule->OnGetMediaPlateMaterialAssetPaths().RemoveAll(this);
+	}
 	
 	UToolMenus::UnRegisterStartupCallback(this);
 	UToolMenus::UnregisterOwner(this);
@@ -101,6 +112,12 @@ FReply FAvalancheMaskEditorModule::OnToggleMaskModeClicked()
 		: FReply::Unhandled();
 }
 
+void FAvalancheMaskEditorModule::OnGetMediaPlateMaterialAssetPaths(TArray<FName>& OutAssetPaths)
+{
+	// Add package path(s) to media plate material variations for mask modifier. 
+	OutAssetPaths.AddUnique(FName("/Avalanche/MaskResources/MediaPlate"));
+}
+
 void FAvalancheMaskEditorModule::RegisterMenus()
 {
 	// Owner will be used for cleanup in call to UToolMenus::UnregisterOwner
diff --git a/Engine/Plugins/Experimental/Avalanche/Source/AvalancheMaskEditor/Private/AvaMaskEditorModule.h b/Engine/Plugins/Experimental/Avalanche/Source/AvalancheMaskEditor/Private/AvaMaskEditorModule.h
--- a/Engine/Plugins/Experimental/Avalanche/Source/AvalancheMaskEditor/Private/AvaMaskEditorModule.h
+++ b/Engine/Plugins/Experimental/Avalanche/Source/AvalancheMaskEditor/Private/AvaMaskEditorModule.h
@@ -27,6 +27,7 @@ private:
 
 	TSharedRef<SWidget> GetStatusBarWidgetMenuContent();
 	FReply OnToggleMaskModeClicked();
+	void OnGetMediaPlateMaterialAssetPaths(TArray<FName>& OutAssetPaths);
 
 private:
 	TSharedPtr<FUICommandList> CommandList;
