[AvalancheMask] Adding Media Plate material handle.

If a MID needs to be created, it will be created as a sub-object of the Media Plate's Static Mesh Component. That way, the MID can be saved and reloaded correctly, and prevents the Media Plate from resettting to default material.

diff --git a/Engine/Plugins/Experimental/Avalanche/Source/AvalancheMask/AvalancheMask.Build.cs b/Engine/Plugins/Experimental/Avalanche/Source/AvalancheMask/AvalancheMask.Build.cs
--- a/Engine/Plugins/Experimental/Avalanche/Source/AvalancheMask/AvalancheMask.Build.cs
+++ b/Engine/Plugins/Experimental/Avalanche/Source/AvalancheMask/AvalancheMask.Build.cs
@@ -29,6 +29,7 @@ public class AvalancheMask : ModuleRules
 				"DynamicMaterial",
 				"Engine",
 				"GeometryFramework",
+				"MediaPlate",
 				"Slate",
 				"SlateCore",
 				"Text3D",
diff --git a/Engine/Plugins/Experimental/Avalanche/Source/AvalancheMask/Private/Handling/AvaMaskActorMaterialCollectionHandle.cpp b/Engine/Plugins/Experimental/Avalanche/Source/AvalancheMask/Private/Handling/AvaMaskActorMaterialCollectionHandle.cpp
--- a/Engine/Plugins/Experimental/Avalanche/Source/AvalancheMask/Private/Handling/AvaMaskActorMaterialCollectionHandle.cpp
+++ b/Engine/Plugins/Experimental/Avalanche/Source/AvalancheMask/Private/Handling/AvaMaskActorMaterialCollectionHandle.cpp
@@ -414,8 +414,12 @@ bool FAvaMaskActorMaterialCollectionHandle::IsSlotOccupied(
 void FAvaMaskActorMaterialCollectionHandle::OnMaterialAssignmentChanged(UPrimitiveComponent* InPrimitiveComponent, const int32 InPrevSlotNum, const int32 InNewSlotNum, const TArray<int32>& InChangedMaterialSlots)
 {
 	check(InPrimitiveComponent);
-	
-	RefreshMaterials();
+
+	// todo: we could either ignore or delay material assignment observer during ApplyModifiedState.
+	if (!bIsIteratingMaterials)
+	{
+		RefreshMaterials();
+	}
 
 	if (const AActor* Actor = GetActor())
 	{
diff --git a/Engine/Plugins/Experimental/Avalanche/Source/AvalancheMask/Private/Handling/AvaMaskMediaPlateMaterialHandle.cpp b/Engine/Plugins/Experimental/Avalanche/Source/AvalancheMask/Private/Handling/AvaMaskMediaPlateMaterialHandle.cpp
new file mode 100644
--- /dev/null
+++ b/Engine/Plugins/Experimental/Avalanche/Source/AvalancheMask/Private/Handling/AvaMaskMediaPlateMaterialHandle.cpp
@@ -0,0 +1,89 @@
+// Copyright Epic Games, Inc. All Rights Reserved.
+
+#include "AvaMaskMediaPlateMaterialHandle.h"
+
+#include "AvaMaterialInstanceHandle.h"
+#include "Components/StaticMeshComponent.h"
+#include "Materials/MaterialInstance.h"
+#include "MediaPlate.h"
+
+FAvaMaskMediaPlateMaterialHandle::FAvaMaskMediaPlateMaterialHandle(const TWeakObjectPtr<UMaterialInterface>& InWeakParentMaterial)
+	: FAvaMaskMaterialInstanceHandle(InWeakParentMaterial)
+{
+	
+}
+
+void FAvaMaskMediaPlateMaterialHandle::SetBlendMode(const EBlendMode InBlendMode)
+{
+	if (GetParentMaterial() && GetParentMaterial() == GetMaterialInstance())
+	{
+		return;	// Don't set blend mode on the parent material.
+	}
+
+	Super::SetBlendMode(InBlendMode);
+}
+
+bool FAvaMaskMediaPlateMaterialHandle::IsSupported(const UStruct* InStruct, const TVariant<UObject*, FStructView>& InInstance, FName InTag)
+{
+	// Make sure this is generally supported as a material instance.
+	if (!FAvaMaterialInstanceHandle::IsSupported(InStruct, InInstance, InTag))
+	{
+		return false;
+	}
+
+	// MediaPlate makes either a MIC or MID.
+	// Check if the material is embedded under the static mesh component.
+	if (const UMaterialInstance* MaterialInstance = Cast<UMaterialInstance>(InInstance.Get<UObject*>()))
+	{
+		if (const UStaticMeshComponent* StaticMeshComponent = Cast<UStaticMeshComponent>(MaterialInstance->GetOuter()))
+		{
+			// Need to check that it is a media plate to be sure
+			if (const AMediaPlate* MediaPlate = Cast<AMediaPlate>(StaticMeshComponent->GetOuter()))
+			{
+				// only consider the mesh component that the media plate uses
+				return MediaPlate->StaticMeshComponent == StaticMeshComponent;
+			}
+		}
+	}
+	return false;
+}
+
+UMaterialInstanceDynamic* FAvaMaskMediaPlateMaterialHandle::GetOrCreateMaterialInstance(const FName InMaskName, const EBlendMode InBlendMode)
+{
+	if (UMaterialInstanceDynamic* MaterialInstance = GetOrCreateMaterialInstanceImpl())
+	{
+		return MaterialInstance;
+	}
+
+	// Fallback to material instance handling.
+	return Super::GetOrCreateMaterialInstance(InMaskName, InBlendMode);
+}
+
+UMaterialInstanceDynamic* FAvaMaskMediaPlateMaterialHandle::GetOrCreateMaterialInstance()
+{
+	if (UMaterialInstanceDynamic* MaterialInstance = GetOrCreateMaterialInstanceImpl())
+	{
+		return MaterialInstance;
+	}
+	return Super::GetOrCreateMaterialInstance();
+}
+
+UMaterialInstanceDynamic* FAvaMaskMediaPlateMaterialHandle::GetOrCreateMaterialInstanceImpl()
+{
+	if (UMaterialInstanceDynamic* MaterialInstance = GetMaterialInstance())
+	{
+		return MaterialInstance;
+	}
+
+	if (UMaterialInterface* ParentMaterial = GetParentMaterial())
+	{
+		// Need to be embedded in the level.
+		if (UObject* Outer = Cast<UStaticMeshComponent>(ParentMaterial->GetOuter()))
+		{
+			UMaterialInstanceDynamic* MaterialInstance = UMaterialInstanceDynamic::Create(ParentMaterial, Outer);
+			WeakMaterialInstance = MaterialInstance;
+			return MaterialInstance;
+		}
+	}
+	return nullptr;
+}
diff --git a/Engine/Plugins/Experimental/Avalanche/Source/AvalancheMask/Private/Handling/AvaMaskMediaPlateMaterialHandle.h b/Engine/Plugins/Experimental/Avalanche/Source/AvalancheMask/Private/Handling/AvaMaskMediaPlateMaterialHandle.h
new file mode 100644
--- /dev/null
+++ b/Engine/Plugins/Experimental/Avalanche/Source/AvalancheMask/Private/Handling/AvaMaskMediaPlateMaterialHandle.h
@@ -0,0 +1,31 @@
+// Copyright Epic Games, Inc. All Rights Reserved.
+
+#pragma once
+
+#include "Handling/AvaMaskMaterialInstanceHandle.h"
+
+class FAvaMaskMediaPlateMaterialHandle : public FAvaMaskMaterialInstanceHandle
+{
+public:
+	UE_AVA_INHERITS_WITH_SUPER(FAvaMaskMediaPlateMaterialHandle, FAvaMaskMaterialInstanceHandle);
+	
+	explicit FAvaMaskMediaPlateMaterialHandle(const TWeakObjectPtr<UMaterialInterface>& InWeakParentMaterial);
+	virtual ~FAvaMaskMediaPlateMaterialHandle() override = default;
+	
+	//~ Begin IAvaMaskMaterialHandle
+	virtual void SetBlendMode(const EBlendMode InBlendMode) override;
+	//~ End IAvaMaskMaterialHandle
+
+	static bool IsSupported(const UStruct* InStruct, const TVariant<UObject*, FStructView>& InInstance, FName InTag = NAME_None);
+
+protected:
+	//~ Begin FAvaMaskMaterialInstanceHandle
+	virtual UMaterialInstanceDynamic* GetOrCreateMaterialInstance(const FName InMaskName, const EBlendMode InBlendMode) override;
+	//~ End FAvaMaskMaterialInstanceHandle
+
+	//~ Begin FAvaMaterialInstanceHandle
+	virtual UMaterialInstanceDynamic* GetOrCreateMaterialInstance() override;
+	//~ End FAvaMaterialInstanceHandle
+
+	UMaterialInstanceDynamic* GetOrCreateMaterialInstanceImpl();
+};
diff --git a/Engine/Plugins/Experimental/Avalanche/Source/AvalancheMask/Private/Handling/AvaObjectHandleSubsystem.cpp b/Engine/Plugins/Experimental/Avalanche/Source/AvalancheMask/Private/Handling/AvaObjectHandleSubsystem.cpp
--- a/Engine/Plugins/Experimental/Avalanche/Source/AvalancheMask/Private/Handling/AvaObjectHandleSubsystem.cpp
+++ b/Engine/Plugins/Experimental/Avalanche/Source/AvalancheMask/Private/Handling/AvaObjectHandleSubsystem.cpp
@@ -6,6 +6,7 @@
 #include "AvaMaskActorMaterialCollectionHandle.h"
 #include "AvaMaskAvaShapeMaterialCollectionHandle.h"
 #include "AvaMaskLog.h"
+#include "AvaMaskMediaPlateMaterialHandle.h"
 #include "AvaMaskText3DActorMaterialCollectionHandle.h"
 #include "Engine/Engine.h"
 #include "GameFramework/Actor.h"
@@ -104,7 +105,14 @@ void UAvaObjectHandleSubsystem::FindObjectHandleFactories()
 			{
 				return MakeShared<FAvaMaskParametricMaterialHandle>(InMaterial.Get<FStructView>());
 			}});
-	
+
+		ObjectHandleFactories.Add({
+			FAvaMaskMediaPlateMaterialHandle::IsSupported
+			, [](const TVariant<UObject*, FStructView>& InMaterial) -> TSharedPtr<IAvaMaskMaterialHandle>
+			{
+				return MakeShared<FAvaMaskMediaPlateMaterialHandle>(Cast<UMaterialInterface>(InMaterial.Get<UObject*>()));
+			}});
+
 		ObjectHandleFactories.Add({
 			FAvaMaskMaterialInstanceHandle::IsSupported
 			, [](const TVariant<UObject*, FStructView>& InMaterial) -> TSharedPtr<IAvaMaskMaterialHandle>
