[AvalancheMedia] Broadcast Channel: fixing error state management when a media output fails to initialize (early failure).

diff --git a/Engine/Plugins/Experimental/Avalanche/Source/AvalancheMedia/Private/Broadcast/Channel/AvaBroadcastOutputChannel.cpp b/Engine/Plugins/Experimental/Avalanche/Source/AvalancheMedia/Private/Broadcast/Channel/AvaBroadcastOutputChannel.cpp
--- a/Engine/Plugins/Experimental/Avalanche/Source/AvalancheMedia/Private/Broadcast/Channel/AvaBroadcastOutputChannel.cpp
+++ b/Engine/Plugins/Experimental/Avalanche/Source/AvalancheMedia/Private/Broadcast/Channel/AvaBroadcastOutputChannel.cpp
@@ -427,10 +427,10 @@ EAvaBroadcastOutputState FAvaBroadcastOutputChannel::GetMediaOutputState(const U
 			return EAvaBroadcastOutputState::Idle;
 		}
 	}
-	
-	// If no media capture is found for the given output and we are marked as broadcasting,
-	// this is an error state.
-	return bInternalStateBroadcasting ?  EAvaBroadcastOutputState::Error : EAvaBroadcastOutputState::Idle;
+
+	// Check if this output has reported errors.	
+	return GetLocalSeverityForOutput(InMediaOutput) == EAvaBroadcastIssueSeverity::Errors ?
+		EAvaBroadcastOutputState::Error : EAvaBroadcastOutputState::Idle;
 }
 
 const FAvaBroadcastMediaOutputInfo& FAvaBroadcastOutputChannel::GetMediaOutputInfo(const UMediaOutput* InMediaOutput) const
@@ -662,8 +662,6 @@ bool FAvaBroadcastOutputChannel::IsValidChannel() const
 
 bool FAvaBroadcastOutputChannel::StartChannelBroadcast()
 {
-	bInternalStateBroadcasting = true;
-	
 	IAvaMediaModule& Module = IAvaMediaModule::Get();
 	if (Module.IsPlaybackClientStarted())
 	{
@@ -719,12 +717,21 @@ bool FAvaBroadcastOutputChannel::StartChannelBroadcast()
 		
 		if (!IsValid(MediaCapture))
 		{
-			FString& Message = EmplaceMessageForOutput(MediaOutput, EAvaBroadcastIssueSeverity::Warnings);
-			Message = FString::Printf(
-				TEXT("Output Channel (%s) has a Media Output (%s) that could not create a valid Media Capture.")
-				, *ChannelName, *MediaOutput->GetName());
+			FString& Message = EmplaceMessageForOutput(MediaOutput, EAvaBroadcastIssueSeverity::Errors);
+			
+			// Try to get a failure reason.
+			FString FailureReason;
+			if (!MediaOutput->Validate(FailureReason))
+			{
+				Message = FString::Printf(TEXT("Couldn't create media capture. %s"), *FailureReason);
+			}
+			else
+			{
+				Message = FString::Printf(TEXT("Couldn't create media capture."));
+			}
 
-			UE_LOG(LogAvaOutputChannel, Warning, TEXT("%s"), *Message);
+			UE_LOG(LogAvaOutputChannel, Error, TEXT("Output Channel (%s) has a Media Output (%s): %s"),
+				*ChannelName, *MediaOutput->GetName(), *Message);
 			
 			StopCaptureForOutput(MediaOutput);	// Ensures no stale capture is running.
 			continue;
@@ -746,6 +753,11 @@ bool FAvaBroadcastOutputChannel::StartChannelBroadcast()
 			AvaMediaCapture.Reset(MediaCapture, MediaOutput, RenderTarget);
 			AvaMediaCapture.OnMediaCaptureStateChanged.AddRaw(this, &FAvaBroadcastOutputChannel::OnAvaMediaCaptureStateChanged); 
 		}
+		else
+		{
+			FString& Message = EmplaceMessageForOutput(MediaOutput, EAvaBroadcastIssueSeverity::Errors);
+			Message = TEXT("Failed to initialize. See logs for more info.");
+		}
 	}
 
 	RefreshState();
@@ -754,8 +766,6 @@ bool FAvaBroadcastOutputChannel::StartChannelBroadcast()
 
 void FAvaBroadcastOutputChannel::StopChannelBroadcast()
 {
-	bInternalStateBroadcasting = false;
-
 	StopPlaceholderTick();
 	
 	IAvaMediaModule& Module = IAvaMediaModule::Get();
@@ -871,6 +881,7 @@ int32 FAvaBroadcastOutputChannel::RemoveMediaOutput(UMediaOutput* InMediaOutput)
 		RemoveCount = 1;
 		MediaOutputs.RemoveAt(IndexToRemove);
 		MediaOutputInfos.RemoveAt(IndexToRemove);
+		LocalMediaOutputStatuses.Remove(InMediaOutput);
 		OnChannelChanged.Broadcast(*this, EAvaBroadcastChannelChange::MediaOutputs);
 		UpdateChannelResources(/*bInIsProfileActive*/true);
 		RefreshState();
diff --git a/Engine/Plugins/Experimental/Avalanche/Source/AvalancheMedia/Public/Broadcast/Channel/AvaBroadcastOutputChannel.h b/Engine/Plugins/Experimental/Avalanche/Source/AvalancheMedia/Public/Broadcast/Channel/AvaBroadcastOutputChannel.h
--- a/Engine/Plugins/Experimental/Avalanche/Source/AvalancheMedia/Public/Broadcast/Channel/AvaBroadcastOutputChannel.h
+++ b/Engine/Plugins/Experimental/Avalanche/Source/AvalancheMedia/Public/Broadcast/Channel/AvaBroadcastOutputChannel.h
@@ -288,9 +288,4 @@ protected:
 	FString& EmplaceMessageForOutput(const UMediaOutput* InMediaOutput, EAvaBroadcastIssueSeverity InSeverity) const;
 	EAvaBroadcastIssueSeverity GetLocalSeverityForOutput(const UMediaOutput* InMediaOutput) const;
 	void ResetLocalMediaOutputStatus(const UMediaOutput* InMediaOutput);
-	
-	// This internal state is used to mark the channel as "should be broadcasting".
-	// This changes the interpretation of the output statuses with regard to the
-	// overall channel state.
-	bool bInternalStateBroadcasting = false;
 };
