Motion Design: fix issue with jump marks not working outside sequencer

diff --git a/Engine/Plugins/Experimental/Avalanche/Source/AvalancheSequence/Private/AvaSequencePlayer.cpp b/Engine/Plugins/Experimental/Avalanche/Source/AvalancheSequence/Private/AvaSequencePlayer.cpp
--- a/Engine/Plugins/Experimental/Avalanche/Source/AvalancheSequence/Private/AvaSequencePlayer.cpp
+++ b/Engine/Plugins/Experimental/Avalanche/Source/AvalancheSequence/Private/AvaSequencePlayer.cpp
@@ -150,10 +150,9 @@ void UAvaSequencePlayer::JumpTo(FFrameTime InJumpToFrame, bool bInEvaluate)
 		constexpr bool bHasJumped = true;
 		UpdateTimeCursorPosition(InJumpToFrame, EUpdatePositionMethod::Jump, bHasJumped);
 	}
-	else
-	{
-		PlayPosition.JumpTo(InJumpToFrame);
-	}
+
+	PlayPosition.JumpTo(InJumpToFrame);
+	TimeController->Reset(GetCurrentTime());
 }
 
 FFrameRate UAvaSequencePlayer::GetTickResolution() const
