[AvalancheMedia] Fixing Rundown Preview Frame functionality.

- Fixing the transition path to include Preview Frame transitions.
- Getting rid of no transition path, pages are always played with a transition now.

diff --git a/Engine/Plugins/Experimental/Avalanche/Source/AvalancheMedia/Private/Playable/AvaPlayable.cpp b/Engine/Plugins/Experimental/Avalanche/Source/AvalancheMedia/Private/Playable/AvaPlayable.cpp
--- a/Engine/Plugins/Experimental/Avalanche/Source/AvalancheMedia/Private/Playable/AvaPlayable.cpp
+++ b/Engine/Plugins/Experimental/Avalanche/Source/AvalancheMedia/Private/Playable/AvaPlayable.cpp
@@ -202,6 +202,8 @@ EAvaPlayableCommandResult UAvaPlayable::ExecuteAnimationCommand(EAvaPlaybackAnim
 	UE_LOG(LogAvaPlayable, Verbose,
 		TEXT("%s Playable {%s} -> Executing Sequence Command: {%s}."),
 		*GetBriefFrameInfo(), *GetPrettyPlayableInfo(this), *GetPrettySequenceCommandInfo(InAnimAction, InAnimPlaySettings));
+
+	bool bFoundPreviewMark = false;
 	
 	for (const TObjectPtr<UAvaSequence>& Sequence : SequenceProvider->GetSequences())
 	{
@@ -223,11 +225,26 @@ EAvaPlayableCommandResult UAvaPlayable::ExecuteAnimationCommand(EAvaPlaybackAnim
 				break;
 
 			case EAvaPlaybackAnimAction::PreviewFrame:
-				PlaybackObject->PreviewFrame(Sequence);
+				// If no animation is specified, delay the missing preview mark warning.
+				if (Sequence->GetPreviewMark() || Sequence->GetFName() == InAnimPlaySettings.AnimationName)
+				{
+					bFoundPreviewMark = true;
+					PlaybackObject->PreviewFrame(Sequence);
+				}
 				break;
 			}
 		}
 	}
+
+	// Log a warning if a PreviewFrame was requested but no marks where found in any of the sequences.
+	if (InAnimPlaySettings.AnimationName.IsNone() && InAnimAction == EAvaPlaybackAnimAction::PreviewFrame && !bFoundPreviewMark)
+	{
+		const FString SequenceList = FString::JoinBy(SequenceProvider->GetSequences(), TEXT(", "), [](const TObjectPtr<UAvaSequence>& InSequence)
+		{
+			return FString::Printf(TEXT("%s' ('%s')"), *InSequence->GetLabel().ToString(), *InSequence->GetName());
+		});		
+		UE_LOG(LogAvaPlayable, Warning , TEXT("Failed to Preview Sequence. Preview Mark was not found in any sequences: %s.") , *SequenceList);
+	}
 	
 	return EAvaPlayableCommandResult::Executed;
 }
diff --git a/Engine/Plugins/Experimental/Avalanche/Source/AvalancheMedia/Private/Playable/Transition/AvaPlayableLocalTransition.cpp b/Engine/Plugins/Experimental/Avalanche/Source/AvalancheMedia/Private/Playable/Transition/AvaPlayableLocalTransition.cpp
--- a/Engine/Plugins/Experimental/Avalanche/Source/AvalancheMedia/Private/Playable/Transition/AvaPlayableLocalTransition.cpp
+++ b/Engine/Plugins/Experimental/Avalanche/Source/AvalancheMedia/Private/Playable/Transition/AvaPlayableLocalTransition.cpp
@@ -61,19 +61,29 @@ namespace UE::AvaPlayableTransition::Private
 		return true;
 	}
 
-	FAvaTagHandle GetTransitionLayer(const FAvaTransitionBehaviorInstance& InBehaviorInstance)
+	FAvaTagHandle GetTransitionLayer(const IAvaTransitionBehavior* InBehavior)
 	{
-		// Remark: InBehaviorInstance.GetTransitionLayer() doesn't return a valid layer.
-		if (const IAvaTransitionBehavior* Behavior = InBehaviorInstance.GetBehavior())
+		if (InBehavior)
 		{
-			if (const UAvaTransitionTree* TransitionTree = Behavior->GetTransitionTree())
+			if (const UAvaTransitionTree* TransitionTree = InBehavior->GetTransitionTree())
 			{
 				return TransitionTree->GetTransitionLayer();
 			}
 		}
 		return FAvaTagHandle();
 	}
-	
+
+	FAvaTagHandle GetTransitionLayer(const FAvaTransitionBehaviorInstance& InBehaviorInstance)
+	{
+		// Remark: InBehaviorInstance.GetTransitionLayer() doesn't return a valid layer.
+		return GetTransitionLayer(InBehaviorInstance.GetBehavior());
+	}
+
+	FAvaTagHandle GetTransitionLayer(const UAvaPlayable& InPlayable)
+	{
+		return GetTransitionLayer(GetTransitionBehavior(InPlayable));
+	}
+
 	class FBuilderHelper
 	{
 	public:
@@ -252,6 +262,30 @@ namespace UE::AvaPlayableTransition::Private
 		}
 		return false;
 	}
+
+	/**
+	 * Determines if a playing playable should be kicked out based on the transition layer.
+	 */
+	bool ShouldKickOutPlayable(bool bInIsPreviewFrameLogic, const UAvaPlayable* InPlayable, const TArray<FAvaTagHandle>& InTransitionLayersToKickOut)
+	{
+		if (bInIsPreviewFrameLogic && InPlayable)
+		{
+			const FAvaTagHandle TransitionLayer = GetTransitionLayer(*InPlayable);
+
+			auto IsTransitionLayerOverlap = [&TransitionLayer](const FAvaTagHandle& InOtherLayer)
+			{
+				// Note: Invalid layer (No TL) means everything is kicked out.
+				return !InOtherLayer.IsValid() || InOtherLayer.Overlaps(TransitionLayer);
+			};
+			
+			if (!TransitionLayer.IsValid()	// No TL gets kicked out regardless.
+				|| InTransitionLayersToKickOut.ContainsByPredicate(IsTransitionLayerOverlap))
+			{
+				return true;
+			}
+		}
+		return false;
+	}
 }
 
 bool UAvaPlayableLocalTransition::Start()
@@ -291,6 +325,10 @@ bool UAvaPlayableLocalTransition::Start()
 		OtherPlayableValues.Add(Playable, Playable->GetLatestRemoteControlValues());
 	}
 
+	// PreviewFrame will bypass the transition behaviors.
+	const bool bIsPreviewFrameTransition = EnumHasAnyFlags(TransitionFlags, EAvaPlayableTransitionFlags::PlayEnterPlayablesAtPreviewFrame);
+	TArray<FAvaTagHandle> TransitionLayersToKickOut;
+
 	int32 ArrayIndex = 0;
 	for (UAvaPlayable* Playable : EnterPlayables)
 	{
@@ -298,12 +336,20 @@ bool UAvaPlayableLocalTransition::Start()
 		{
 			Playable->UpdateRemoteControlCommand(EnterPlayableValues[ArrayIndex].ToSharedRef());
 		}
-		
-		if (!Helper.AddTransitionBehaviorInstance(Playable, EAvaPlayableTransitionEntryRole::Enter))
+
+		if (bIsPreviewFrameTransition || !Helper.AddTransitionBehaviorInstance(Playable, EAvaPlayableTransitionEntryRole::Enter))
 		{
-			// Handle Playables with no transition tree.
+			// Handle Playables with no transition tree (or PreviewFrame override).
 			InitializeSequences(Playable);
 			PostExecutorSequencePlayablesWeak.Add(Playable);
+
+			// For preview frame, we keep track of the transition layer to kick out all other playables on that layer.
+			// Normally, the transition tree would do that, but we are bypassing it in this mode.
+			if (bIsPreviewFrameTransition)
+			{
+				// Note: adding potentially invalid layer on purpose, indicates a no TL enter page and that will kick everything else.
+				TransitionLayersToKickOut.Add(GetTransitionLayer(*Playable));
+			}
 		}
 
 		// Signal that this playable can be shown now that sequences are initialized and RC is injected.
@@ -316,9 +362,9 @@ bool UAvaPlayableLocalTransition::Start()
 
 	for (UAvaPlayable* Playable : PlayingPlayables)
 	{
-		if (!Helper.AddTransitionBehaviorInstance(Playable, PlayingPlayablesRole))
+		if (ShouldKickOutPlayable(bIsPreviewFrameTransition, Playable, TransitionLayersToKickOut) || !Helper.AddTransitionBehaviorInstance(Playable, PlayingPlayablesRole))
 		{
-			// No Transition Tree:
+			// No Transition Tree (or kick out logic):
 			// Stop exit/playing playable without any transition.
 			// We can't integrate that with a transition tree for enter pages atm.
 			UAvaPlayable::OnTransitionEvent().Broadcast(Playable, this, EAvaPlayableTransitionEventFlags::StopPlayable);
@@ -327,9 +373,9 @@ bool UAvaPlayableLocalTransition::Start()
 
 	for (UAvaPlayable* Playable : ExitPlayables)
 	{
-		if (!Helper.AddTransitionBehaviorInstance(Playable, EAvaPlayableTransitionEntryRole::Exit))
+		if (ShouldKickOutPlayable(bIsPreviewFrameTransition, Playable, TransitionLayersToKickOut) || !Helper.AddTransitionBehaviorInstance(Playable, EAvaPlayableTransitionEntryRole::Exit))
 		{
-			// No Transition Tree:
+			// No Transition Tree (or kick out logic):
 			// Stop exit/playing playable without any transition.
 			// We can't integrate that with a transition tree for enter pages atm.
 			UAvaPlayable::OnTransitionEvent().Broadcast(Playable, this, EAvaPlayableTransitionEventFlags::StopPlayable);
@@ -442,8 +488,18 @@ void UAvaPlayableLocalTransition::PostTransitionExecutorPhase()
 	{
 		if (UAvaPlayable* Playable = EnterPlayableWeak.Get())
 		{
+			FAvaPlaybackAnimPlaySettings AnimSettings;
+			if (EnumHasAnyFlags(TransitionFlags, EAvaPlayableTransitionFlags::PlayEnterPlayablesAtPreviewFrame))
+			{
+				AnimSettings.Action = EAvaPlaybackAnimAction::PreviewFrame;	
+			}
+			else
+			{
+				AnimSettings.Action = EAvaPlaybackAnimAction::Play;
+			}
+
 			// Default settings will play all sequences.
-			if (Playable->ExecuteAnimationCommand(EAvaPlaybackAnimAction::Play, FAvaPlaybackAnimPlaySettings()) == EAvaPlayableCommandResult::Executed)
+			if (Playable->ExecuteAnimationCommand(AnimSettings.Action, AnimSettings) == EAvaPlayableCommandResult::Executed)
 			{
 				bSequenceStarted = true;
 			}
diff --git a/Engine/Plugins/Experimental/Avalanche/Source/AvalancheMedia/Private/Rundown/AvaRundown.cpp b/Engine/Plugins/Experimental/Avalanche/Source/AvalancheMedia/Private/Rundown/AvaRundown.cpp
--- a/Engine/Plugins/Experimental/Avalanche/Source/AvalancheMedia/Private/Rundown/AvaRundown.cpp
+++ b/Engine/Plugins/Experimental/Avalanche/Source/AvalancheMedia/Private/Rundown/AvaRundown.cpp
@@ -1216,13 +1216,7 @@ TArray<int32> UAvaRundown::PlayPages(const TArray<int32>& InPageIds, EAvaRundown
 				continue;
 			}
 
-			// Only preview from frame still uses the no transition path.
-			// Everything else uses transitions, even when no transition tree.
-			const bool bPagePlayed = (InPlayType != EAvaRundownPagePlayType::PreviewFromFrame)
-				? PlayPageWithTransition(TransitionBuilder, SelectedPage, InPlayType, bIsPreview, InPreviewChannelName)
-				: PlayPageNoTransition(SelectedPage, InPlayType, bIsPreview, InPreviewChannelName);
-		
-			if (bPagePlayed)
+			if (PlayPageWithTransition(TransitionBuilder, SelectedPage, InPlayType, bIsPreview, InPreviewChannelName))
 			{
 				GetOrCreatePageListPlaybackContextCollection().GetOrCreateContext(bIsPreview, InPreviewChannelName).PlayHeadPageId = PageId;
 				PlayedPageIds.Add(PageId);
@@ -1996,42 +1990,6 @@ IAvaRundownPageLoadingManager& UAvaRundown::MakePageLoadingManager()
 	return *PageLoadingManager;
 }
 
-bool UAvaRundown::PlayPageNoTransition(const FAvaRundownPage& InPage, EAvaRundownPagePlayType InPlayType, bool bInIsPreview, const FName& InPreviewChannelName)
-{
-	UAvaRundownPagePlayer* PagePlayer = FindPlayerForPage(InPage.GetPageId(), bInIsPreview, InPreviewChannelName);
-	if (!PagePlayer)
-	{
-		PagePlayer = NewObject<UAvaRundownPagePlayer>(this);
-		if (PagePlayer->InitializeAndLoad(this, InPage, bInIsPreview, InPreviewChannelName))
-		{
-			AddPagePlayer(PagePlayer);
-		}
-	}
-
-	if (PagePlayer && PagePlayer->IsLoaded())
-	{
-		PagePlayer->Play(InPlayType, false);	// not using transition logic: animations commands are going to be pushed.
-			
-		// Don't let something else play on this channel.
-		// FIXME: this is not perfect, on the remote render node, it may cause a frame of "place holder"
-		// between the end of the previous page and start of the next one.
-
-		for (UAvaRundownPagePlayer* OtherPagePlayer : PagePlayers)
-		{
-			// Don't let something else play on this channel.
-			if (OtherPagePlayer != PagePlayer && OtherPagePlayer->ChannelFName == PagePlayer->ChannelFName)
-			{
-				OtherPagePlayer->Stop();
-			}
-		}
-
-		RemoveStoppedPagePlayers();
-				
-		return true;
-	}
-	return false;
-}
-
 namespace UE::AvaMedia::Rundown::Private
 {
 	bool ArePageRCValuesEqualForSubTemplate(const FAvaRundownPage& InSubTemplate, const FAvaRundownPage& InPage, const FAvaRundownPage& InOtherPage)
@@ -2250,12 +2208,13 @@ bool UAvaRundown::PlayPageWithTransition(FAvaRundownPageTransitionBuilder& InBui
 			{
 				PageTransition->InstancesBypassingTransition.Append(InstancesBypassingTransition);
 				PageTransition->ReusedInstances.Append(ReusedInstances);
+				PageTransition->bIsPreviewFrameTransition = (InPlayType == EAvaRundownPagePlayType::PreviewFromFrame);
 
 				AddPagePlayer(NewPagePlayer);
 
 				// Start the playback, will only actually start on next tick.
 				// Animation command will not be pushed, relying on TL to start the appropriate animations.
-				NewPagePlayer->Play(InPlayType, true);
+				NewPagePlayer->Play(InPlayType);
 				return true;
 			}
 		}
diff --git a/Engine/Plugins/Experimental/Avalanche/Source/AvalancheMedia/Private/Rundown/AvaRundownPagePlayer.cpp b/Engine/Plugins/Experimental/Avalanche/Source/AvalancheMedia/Private/Rundown/AvaRundownPagePlayer.cpp
--- a/Engine/Plugins/Experimental/Avalanche/Source/AvalancheMedia/Private/Rundown/AvaRundownPagePlayer.cpp
+++ b/Engine/Plugins/Experimental/Avalanche/Source/AvalancheMedia/Private/Rundown/AvaRundownPagePlayer.cpp
@@ -10,18 +10,6 @@
 
 namespace UE::AvaMedia::Rundown::PagePlayer::Private
 {
-	static void PushAllAnimations(UAvaPlaybackGraph* InPlaybackObject, const UAvaRundown* InRundown, const FAvaRundownPage& InPage
-		, const FString& InChannelName, EAvaPlaybackAnimAction InAnimAction)
-	{
-		const int32 NumTemplates = InPage.GetNumTemplates(InRundown);
-		for (int32 TemplateIndex = 0; TemplateIndex < NumTemplates; ++TemplateIndex)
-		{
-			// Note: default FAvaPlaybackAnimPlaySettings (with NAME_None) will play all animations.
-			InPlaybackObject->PushAnimationCommand(InPage.GetAssetPath(InRundown, TemplateIndex), InChannelName,
-				InAnimAction, FAvaPlaybackAnimPlaySettings());
-		}
-	}
-
 	static void PushCameraCut(UAvaPlaybackGraph* InPlaybackObject, const UAvaRundown* InRundown, const FAvaRundownPage& InPage
 		, const FString& InChannelName)
 	{
@@ -67,7 +55,7 @@ bool UAvaRundownPlaybackInstancePlayer::IsLoaded() const
 	return Playback != nullptr;
 }
 
-void UAvaRundownPlaybackInstancePlayer::Play(const UAvaRundownPagePlayer& InPagePlayer, const UAvaRundown* InRundown, EAvaRundownPagePlayType InPlayType, bool bInIsUsingTransitionLogic)
+void UAvaRundownPlaybackInstancePlayer::Play(const UAvaRundownPagePlayer& InPagePlayer, const UAvaRundown* InRundown, EAvaRundownPagePlayType InPlayType)
 {
 	using namespace UE::AvaMedia::Rundown::PagePlayer::Private;
 	
@@ -83,25 +71,13 @@ void UAvaRundownPlaybackInstancePlayer::Play(const UAvaRundownPagePlayer& InPage
 		Playback->Play();
 	}
 
-	const FAvaRundownPage& Page = InRundown->GetPage(InPagePlayer.PageId);
-
-	// When not using transition logic, we start all animations manually.
-	if (!bInIsUsingTransitionLogic)
-	{
-		Playback->PushRemoteControlValues(SourceAssetPath, InPagePlayer.ChannelName, MakeShared<FAvaPlayableRemoteControlValues>(Page.GetRemoteControlValues()));
-
-		// TODO: play type should be transmitted to the transition behavior.
-		const EAvaPlaybackAnimAction AnimAction = (InPlayType == EAvaRundownPagePlayType::PreviewFromFrame)
-			? EAvaPlaybackAnimAction::PreviewFrame
-			: EAvaPlaybackAnimAction::Play;
-	
-		// Play all the animations with the page's animation settings.
-		PushAllAnimations(Playback, InRundown, Page, InPagePlayer.ChannelName, AnimAction);
-	}
-	
 	if (bPlaybackObjectWasPlaying)
 	{
-		PushCameraCut(Playback, InRundown, Page, InPagePlayer.ChannelName);
+		const FAvaRundownPage& Page = InRundown->GetPage(InPagePlayer.PageId);
+		if (Page.IsValidPage())
+		{
+			PushCameraCut(Playback, InRundown, Page, InPagePlayer.ChannelName);
+		}
 	}
 }
 
@@ -266,7 +242,7 @@ bool UAvaRundownPagePlayer::IsLoaded() const
 	return false;
 }
 
-bool UAvaRundownPagePlayer::Play(EAvaRundownPagePlayType InPlayType, bool bInIsUsingTransitionLogic)
+bool UAvaRundownPagePlayer::Play(EAvaRundownPagePlayType InPlayType)
 {
 	const UAvaRundown* Rundown = RundownWeak.Get();
 	if (!Rundown)
@@ -277,7 +253,7 @@ bool UAvaRundownPagePlayer::Play(EAvaRundownPagePlayType InPlayType, bool bInIsU
 	bool bIsPlaying = false;
 	for (const TObjectPtr<UAvaRundownPlaybackInstancePlayer>& InstancePlayer : InstancePlayers)
 	{
-		InstancePlayer->Play(*this, Rundown, InPlayType, bInIsUsingTransitionLogic);
+		InstancePlayer->Play(*this, Rundown, InPlayType);
 		bIsPlaying |= InstancePlayer->IsPlaying();
 	}
 
diff --git a/Engine/Plugins/Experimental/Avalanche/Source/AvalancheMedia/Private/Rundown/Transition/AvaRundownPageTransition.cpp b/Engine/Plugins/Experimental/Avalanche/Source/AvalancheMedia/Private/Rundown/Transition/AvaRundownPageTransition.cpp
--- a/Engine/Plugins/Experimental/Avalanche/Source/AvalancheMedia/Private/Rundown/Transition/AvaRundownPageTransition.cpp
+++ b/Engine/Plugins/Experimental/Avalanche/Source/AvalancheMedia/Private/Rundown/Transition/AvaRundownPageTransition.cpp
@@ -540,6 +540,12 @@ void UAvaRundownPageTransition::MakePlayableTransition()
 		{
 			TransitionFlags |= EAvaPlayableTransitionFlags::HasReusedPlayables;
 		}
+
+		if (bIsPreviewFrameTransition)
+		{
+			TransitionFlags |= EAvaPlayableTransitionFlags::PlayEnterPlayablesAtPreviewFrame;
+		}
+
 		PlayableTransition->SetTransitionFlags(TransitionFlags);
 	}
 }
diff --git a/Engine/Plugins/Experimental/Avalanche/Source/AvalancheMedia/Private/Rundown/Transition/AvaRundownPageTransition.h b/Engine/Plugins/Experimental/Avalanche/Source/AvalancheMedia/Private/Rundown/Transition/AvaRundownPageTransition.h
--- a/Engine/Plugins/Experimental/Avalanche/Source/AvalancheMedia/Private/Rundown/Transition/AvaRundownPageTransition.h
+++ b/Engine/Plugins/Experimental/Avalanche/Source/AvalancheMedia/Private/Rundown/Transition/AvaRundownPageTransition.h
@@ -69,6 +69,12 @@ public:
 	UPROPERTY(Transient)
 	TSet<FGuid> ReusedInstances;
 
+	/**
+	 * Special mark if any of the enter playables where issued with a PreviewFrame play type.
+	 * Current logic will mark the whole transition.
+	 */
+	bool  bIsPreviewFrameTransition = false;
+
 protected:
 	UAvaRundownPlaybackInstancePlayer* FindInstancePlayerForPlayable(const UAvaPlayable* InPlayable) const;
 	UAvaRundownPagePlayer* FindPagePlayerForPlayable(const UAvaPlayable* InPlayable) const;
diff --git a/Engine/Plugins/Experimental/Avalanche/Source/AvalancheMedia/Public/AvaMediaDefines.h b/Engine/Plugins/Experimental/Avalanche/Source/AvalancheMedia/Public/AvaMediaDefines.h
--- a/Engine/Plugins/Experimental/Avalanche/Source/AvalancheMedia/Public/AvaMediaDefines.h
+++ b/Engine/Plugins/Experimental/Avalanche/Source/AvalancheMedia/Public/AvaMediaDefines.h
@@ -151,6 +151,8 @@ enum class EAvaPlayableTransitionFlags : uint8
 	TreatPlayingAsExiting = 1 << 0,
 	/** Transition contains some reused playables (i.e. both enter and playing). */
 	HasReusedPlayables = 1 << 1,
+	/** Special logic to override the enter playables behaviors for the "PreviewFrame". */
+	PlayEnterPlayablesAtPreviewFrame = 1 << 2,
 };
 ENUM_CLASS_FLAGS(EAvaPlayableTransitionFlags);
 
diff --git a/Engine/Plugins/Experimental/Avalanche/Source/AvalancheMedia/Public/Rundown/AvaRundown.h b/Engine/Plugins/Experimental/Avalanche/Source/AvalancheMedia/Public/Rundown/AvaRundown.h
--- a/Engine/Plugins/Experimental/Avalanche/Source/AvalancheMedia/Public/Rundown/AvaRundown.h
+++ b/Engine/Plugins/Experimental/Avalanche/Source/AvalancheMedia/Public/Rundown/AvaRundown.h
@@ -666,7 +666,6 @@ protected:
 private:
 	IAvaRundownPageLoadingManager& MakePageLoadingManager();
 
-	bool PlayPageNoTransition(const FAvaRundownPage& InPage, EAvaRundownPagePlayType InPlayType, bool bInIsPreview, const FName& InPreviewChannelName);
 	bool PlayPageWithTransition(FAvaRundownPageTransitionBuilder& InBuilder, const FAvaRundownPage& InPage, EAvaRundownPagePlayType InPlayType, bool bInIsPreview, const FName& InPreviewChannelName);
 	bool StopPageNoTransition(const FAvaRundownPage& InPage, bool bInPreview, const FName& InPreviewChannelName);
 	bool StopPageWithTransition(FAvaRundownPageTransitionBuilder& InBuilder, const FAvaRundownPage& InPage, bool bInPreview, const FName& InPreviewChannelName);
diff --git a/Engine/Plugins/Experimental/Avalanche/Source/AvalancheMedia/Public/Rundown/AvaRundownPagePlayer.h b/Engine/Plugins/Experimental/Avalanche/Source/AvalancheMedia/Public/Rundown/AvaRundownPagePlayer.h
--- a/Engine/Plugins/Experimental/Avalanche/Source/AvalancheMedia/Public/Rundown/AvaRundownPagePlayer.h
+++ b/Engine/Plugins/Experimental/Avalanche/Source/AvalancheMedia/Public/Rundown/AvaRundownPagePlayer.h
@@ -31,7 +31,13 @@ public:
 	bool Load(const UAvaRundownPagePlayer& InPagePlayer, const UAvaRundown* InRundown, const FAvaRundownPage& InPage, int32 InSubPageIndex, const FGuid& InInstanceId);
 	bool IsLoaded() const;
 
-	void Play(const UAvaRundownPagePlayer& InPagePlayer, const UAvaRundown* InRundown, EAvaRundownPagePlayType InPlayType, bool bInIsUsingTransitionLogic);
+	UE_DEPRECATED(5.6, "Playing a page without a transition is no longuer supported.")
+	void Play(const UAvaRundownPagePlayer& InPagePlayer, const UAvaRundown* InRundown, EAvaRundownPagePlayType InPlayType, bool bInIsUsingTransitionLogic)
+	{
+		Play(InPagePlayer, InRundown, InPlayType);
+	}
+	
+	void Play(const UAvaRundownPagePlayer& InPagePlayer, const UAvaRundown* InRundown, EAvaRundownPagePlayType InPlayType);
 	bool IsPlaying() const;
 
 	bool Continue(const FString& InChannelName);
@@ -90,8 +96,18 @@ public:
 
 	/** Returns true if at least one of the instance player is loaded, false otherwise. */
 	bool IsLoaded() const;
+
+	UE_DEPRECATED(5.6, "Playing a page without a transition is no longuer supported.")
+	bool Play(EAvaRundownPagePlayType InPlayType, bool bInIsUsingTransitionLogic)
+	{
+		return Play(InPlayType);
+	}
 	
-	bool Play(EAvaRundownPagePlayType InPlayType, bool bInIsUsingTransitionLogic);
+	/**
+	 * Ensure all the playback instances are playing.
+	 * Returns true if at least one playback instance is playing, false otherwise.
+	 */
+	bool Play(EAvaRundownPagePlayType InPlayType);
 
 	/** Returns true if at least one of the instance player is playing, false otherwise. */
 	bool IsPlaying() const;
