Motion Design: refactored mask material handles to not use struct views as the use case with parametric material is unsafe/misused (no control over the lifetime of the parametric material struct)

diff --git a/Engine/Plugins/Experimental/Avalanche/Source/AvalancheMask/Private/Handling/AvaDesignedMaterialHandle.cpp b/Engine/Plugins/Experimental/Avalanche/Source/AvalancheMask/Private/Handling/AvaDesignedMaterialHandle.cpp
--- a/Engine/Plugins/Experimental/Avalanche/Source/AvalancheMask/Private/Handling/AvaDesignedMaterialHandle.cpp
+++ b/Engine/Plugins/Experimental/Avalanche/Source/AvalancheMask/Private/Handling/AvaDesignedMaterialHandle.cpp
@@ -2,6 +2,7 @@
 
 #include "Handling/AvaDesignedMaterialHandle.h"
 
+#include "AvaMaskMaterialReference.h"
 #include "AvaMaskSettings.h"
 #include "Material/DynamicMaterialInstance.h"
 
@@ -30,12 +31,7 @@ bool FAvaDesignedMaterialHandle::IsValid() const
 	return Super::IsValid() && WeakDesignedMaterial.IsValid();
 }
 
-bool FAvaDesignedMaterialHandle::IsSupported(
-	const UStruct* InStruct
-	, const TVariant<UObject*, FStructView>& InInstance
-	, FName InTag)
+bool FAvaDesignedMaterialHandle::IsSupported(const FAvaMaskMaterialReference& InInstance, FName InTag)
 {
-	return InStruct
-		&& ::IsValid(InStruct)
-		&& InStruct->IsChildOf<UDynamicMaterialInstance>();
+	return !!InInstance.GetTypedObject<UDynamicMaterialInstance>();
 }
diff --git a/Engine/Plugins/Experimental/Avalanche/Source/AvalancheMask/Private/Handling/AvaDesignedMaterialHandle.h b/Engine/Plugins/Experimental/Avalanche/Source/AvalancheMask/Private/Handling/AvaDesignedMaterialHandle.h
--- a/Engine/Plugins/Experimental/Avalanche/Source/AvalancheMask/Private/Handling/AvaDesignedMaterialHandle.h
+++ b/Engine/Plugins/Experimental/Avalanche/Source/AvalancheMask/Private/Handling/AvaDesignedMaterialHandle.h
@@ -6,13 +6,13 @@
 #include "IAvaObjectHandle.h"
 
 class UDynamicMaterialInstance;
+struct FAvaMaskMaterialReference;
 
-class FAvaDesignedMaterialHandle
-	: public FAvaMaterialInstanceHandle
+class FAvaDesignedMaterialHandle: public FAvaMaterialInstanceHandle
 {
 public:
 	UE_AVA_INHERITS_WITH_SUPER(FAvaDesignedMaterialHandle, FAvaMaterialInstanceHandle);
-	
+
 	explicit FAvaDesignedMaterialHandle(const TWeakObjectPtr<UDynamicMaterialInstance>& InWeakDesignedMaterial);
 
 	// ~Begin IAvaMaterialHandle
@@ -23,8 +23,8 @@ public:
 	virtual bool IsValid() const override;
 	// ~End IAvaObjectHandle
 
-	static bool IsSupported(const UStruct* InStruct, const TVariant<UObject*, FStructView>& InInstance, FName InTag = NAME_None);
-	
+	static bool IsSupported(const FAvaMaskMaterialReference& InInstance, FName InTag = NAME_None);
+
 protected:
 	TWeakObjectPtr<UDynamicMaterialInstance> WeakDesignedMaterial;
 };
diff --git a/Engine/Plugins/Experimental/Avalanche/Source/AvalancheMask/Private/Handling/AvaMaskActorMaterialCollectionHandle.cpp b/Engine/Plugins/Experimental/Avalanche/Source/AvalancheMask/Private/Handling/AvaMaskActorMaterialCollectionHandle.cpp
--- a/Engine/Plugins/Experimental/Avalanche/Source/AvalancheMask/Private/Handling/AvaMaskActorMaterialCollectionHandle.cpp
+++ b/Engine/Plugins/Experimental/Avalanche/Source/AvalancheMask/Private/Handling/AvaMaskActorMaterialCollectionHandle.cpp
@@ -4,6 +4,7 @@
 
 #include "AvaHandleUtilities.h"
 #include "AvaMaskMaterialAssignmentObserver.h"
+#include "AvaMaskMaterialReference.h"
 #include "AvaMaskUtilities.h"
 #include "AvaObjectHandleSubsystem.h"
 #include "Components/ActorComponent.h"
@@ -352,16 +353,9 @@ bool FAvaMaskActorMaterialCollectionHandle::IsValid() const
 	return WeakActor.IsValid();
 }
 
-bool FAvaMaskActorMaterialCollectionHandle::IsSupported(
-	const UStruct* InStruct
-	, const TVariant<UObject*, FStructView>& InInstance
-	, FName InTag)
+bool FAvaMaskActorMaterialCollectionHandle::IsSupported(const FAvaMaskMaterialReference& InInstance, FName InTag)
 {
-	return InTag == UE::AvaMask::Internal::HandleTag
-		&& InStruct
-		&& ::IsValid(InStruct)
-		&& InStruct->IsChildOf<AActor>()
-		&& InInstance.TryGet<UObject*>() != nullptr;
+	return InTag == UE::AvaMask::Internal::HandleTag && InInstance.GetTypedObject<AActor>();
 }
 
 FStructView FAvaMaskActorMaterialCollectionHandle::GetMaterialHandleData(
diff --git a/Engine/Plugins/Experimental/Avalanche/Source/AvalancheMask/Private/Handling/AvaMaskActorMaterialCollectionHandle.h b/Engine/Plugins/Experimental/Avalanche/Source/AvalancheMask/Private/Handling/AvaMaskActorMaterialCollectionHandle.h
--- a/Engine/Plugins/Experimental/Avalanche/Source/AvalancheMask/Private/Handling/AvaMaskActorMaterialCollectionHandle.h
+++ b/Engine/Plugins/Experimental/Avalanche/Source/AvalancheMask/Private/Handling/AvaMaskActorMaterialCollectionHandle.h
@@ -20,7 +20,7 @@ class IAvaMaskMaterialHandle;
 class UAvaText3DComponent;
 class UMaterialInterface;
 class UText3DComponent;
-struct FStructView;
+struct FAvaMaskMaterialReference;
 
 USTRUCT()
 struct FAvaMaskActorMaterialCollectionHandleData
@@ -81,7 +81,7 @@ public:
 	virtual bool IsValid() const override;
 	// ~End IAvaMaskMaterialCollectionProvider
 
-	static bool IsSupported(const UStruct* InStruct, const TVariant<UObject*, FStructView>& InInstance, FName InTag = NAME_None);
+	static bool IsSupported(const FAvaMaskMaterialReference& InInstance, FName InTag = NAME_None);
 
 protected:
 	// ~Begin TAvaMaskMaterialCollectionHandle
diff --git a/Engine/Plugins/Experimental/Avalanche/Source/AvalancheMask/Private/Handling/AvaMaskAvaShapeMaterialCollectionHandle.cpp b/Engine/Plugins/Experimental/Avalanche/Source/AvalancheMask/Private/Handling/AvaMaskAvaShapeMaterialCollectionHandle.cpp
--- a/Engine/Plugins/Experimental/Avalanche/Source/AvalancheMask/Private/Handling/AvaMaskAvaShapeMaterialCollectionHandle.cpp
+++ b/Engine/Plugins/Experimental/Avalanche/Source/AvalancheMask/Private/Handling/AvaMaskAvaShapeMaterialCollectionHandle.cpp
@@ -3,6 +3,7 @@
 #include "Handling/AvaMaskAvaShapeMaterialCollectionHandle.h"
 
 #include "AvaMaskMaterialAssignmentObserver.h"
+#include "AvaMaskMaterialReference.h"
 #include "AvaMaskUtilities.h"
 #include "AvaText3DComponent.h"
 #include "Components/DynamicMeshComponent.h"
@@ -350,7 +351,6 @@ bool FAvaMaskAvaShapeMaterialCollectionHandle::RefreshMaterials()
 	{
 		OccupiedMaterialSlots.Reset();
 		WeakMaterials.Reset();
-		ParametricMaterialStructs.Reset();
 		MaterialHandles.Reset();
 	};
 
@@ -360,7 +360,6 @@ bool FAvaMaskAvaShapeMaterialCollectionHandle::RefreshMaterials()
 	OccupiedMaterialSlots.SetNum(ProbableNumMaterials, false);
 
 	WeakMaterials.Reserve(ProbableNumMaterials);
-	ParametricMaterialStructs.Reserve(ProbableNumMaterials);
 	MaterialHandles.Reserve(ProbableNumMaterials);
 
 	if (UAvaShapeDynamicMeshBase* ShapeComponent = GetComponent())
@@ -388,8 +387,9 @@ bool FAvaMaskAvaShapeMaterialCollectionHandle::RefreshMaterials()
 				{
 					if (ShapeComponent->GetParametricMaterial(InMeshIdx))
 					{
-						const FStructView& ParametricMaterialView = ParametricMaterialStructs.Emplace(InMeshIdx, FStructView::Make<FAvaShapeParametricMaterial>(*ShapeComponent->GetParametricMaterial(InMeshIdx)));
-						MaterialHandle = MaterialHandles.Emplace(InMeshIdx, HandleSubsystem->MakeHandle<IAvaMaskMaterialHandle, FAvaShapeParametricMaterial>(ParametricMaterialView, UE::AvaMask::Internal::HandleTag));
+						MaterialHandle = HandleSubsystem->MakeHandle<IAvaMaskMaterialHandle>(FAvaMaskMaterialReference(ShapeComponent, InMeshIdx), UE::AvaMask::Internal::HandleTag);
+
+						MaterialHandles.Emplace(InMeshIdx, MaterialHandle);
 						if (!ensureAlways(UE::Ava::Internal::IsHandleValid(MaterialHandle)))
 						{
 							return false;
@@ -449,7 +449,6 @@ bool FAvaMaskAvaShapeMaterialCollectionHandle::RefreshMaterials()
 		}
 
 		UE::AvaMask::Private::RemoveUnused(WeakMaterials, UsedMaterialSlots);
-		UE::AvaMask::Private::RemoveUnused(ParametricMaterialStructs, UsedMaterialSlots);
 		UE::AvaMask::Private::RemoveUnused(MaterialHandles, UsedMaterialSlots);
 
 		return true;
@@ -458,18 +457,11 @@ bool FAvaMaskAvaShapeMaterialCollectionHandle::RefreshMaterials()
 	return false;
 }
 
-bool FAvaMaskAvaShapeMaterialCollectionHandle::IsSupported(
-	const UStruct* InStruct
-	, const TVariant<UObject*, FStructView>& InInstance
-	, FName InTag)
+bool FAvaMaskAvaShapeMaterialCollectionHandle::IsSupported(const FAvaMaskMaterialReference& InInstance, FName InTag)
 {
-	if (InTag == UE::AvaMask::Internal::HandleTag
-		&& InStruct
-		&& ::IsValid(InStruct)
-		&& InStruct->IsChildOf<AActor>()
-		&& InInstance.TryGet<UObject*>() != nullptr)
+	if (InTag == UE::AvaMask::Internal::HandleTag)
 	{
-		if (const AActor* Actor = Cast<AActor>(InInstance.Get<UObject*>()))
+		if (const AActor* Actor = InInstance.GetTypedObject<AActor>())
 		{
 			return Actor->GetComponentByClass<UAvaShapeDynamicMeshBase>() != nullptr;
 		}
diff --git a/Engine/Plugins/Experimental/Avalanche/Source/AvalancheMask/Private/Handling/AvaMaskAvaShapeMaterialCollectionHandle.h b/Engine/Plugins/Experimental/Avalanche/Source/AvalancheMask/Private/Handling/AvaMaskAvaShapeMaterialCollectionHandle.h
--- a/Engine/Plugins/Experimental/Avalanche/Source/AvalancheMask/Private/Handling/AvaMaskAvaShapeMaterialCollectionHandle.h
+++ b/Engine/Plugins/Experimental/Avalanche/Source/AvalancheMask/Private/Handling/AvaMaskAvaShapeMaterialCollectionHandle.h
@@ -22,7 +22,7 @@ class UAvaShapeDynamicMeshBase;
 class UAvaText3DComponent;
 class UMaterialInterface;
 class UText3DComponent;
-struct FStructView;
+struct FAvaMaskMaterialReference;
 
 USTRUCT()
 struct FAvaMaskAvaShapeMaterialCollectionHandleData
@@ -82,7 +82,7 @@ public:
 	virtual bool IsValid() const override;
 	// ~End IAvaMaskMaterialCollectionProvider
 
-	static bool IsSupported(const UStruct* InStruct, const TVariant<UObject*, FStructView>& InInstance, FName InTag = NAME_None);
+	static bool IsSupported(const FAvaMaskMaterialReference& InInstance, FName InTag = NAME_None);
 
 protected:
 	// ~Begin TAvaMaskMaterialCollectionHandle
@@ -112,7 +112,6 @@ protected:
 	TWeakObjectPtr<UAvaShapeDynamicMeshBase> WeakComponent;
 	TBitArray<> OccupiedMaterialSlots;
 	TMap<int32, TWeakObjectPtr<UMaterialInterface>> WeakMaterials;
-	TMap<int32, FStructView> ParametricMaterialStructs;
 	TMap<int32, TSharedPtr<IAvaMaskMaterialHandle>> MaterialHandles;
 	TUniquePtr<FAvaMaskMaterialAssignmentObserver> MaterialAssignmentObserver;
 };
diff --git a/Engine/Plugins/Experimental/Avalanche/Source/AvalancheMask/Private/Handling/AvaMaskDesignedMaterialHandle.cpp b/Engine/Plugins/Experimental/Avalanche/Source/AvalancheMask/Private/Handling/AvaMaskDesignedMaterialHandle.cpp
--- a/Engine/Plugins/Experimental/Avalanche/Source/AvalancheMask/Private/Handling/AvaMaskDesignedMaterialHandle.cpp
+++ b/Engine/Plugins/Experimental/Avalanche/Source/AvalancheMask/Private/Handling/AvaMaskDesignedMaterialHandle.cpp
@@ -255,13 +255,9 @@ bool FAvaMaskDesignedMaterialHandle::IsValid() const
 	return FAvaDesignedMaterialHandle::IsValid();
 }
 
-bool FAvaMaskDesignedMaterialHandle::IsSupported(
-	const UStruct* InStruct
-	, const TVariant<UObject*, FStructView>& InInstance
-	, FName InTag)
+bool FAvaMaskDesignedMaterialHandle::IsSupported(const FAvaMaskMaterialReference& InInstance, FName InTag)
 {
-	return InTag == UE::AvaMask::Internal::HandleTag
-		&& FAvaDesignedMaterialHandle::IsSupported(InStruct, InInstance, InTag);
+	return InTag == UE::AvaMask::Internal::HandleTag && FAvaDesignedMaterialHandle::IsSupported(InInstance, InTag);
 }
 
 UMaterialInstanceDynamic* FAvaMaskDesignedMaterialHandle::GetMaterialInstance()
diff --git a/Engine/Plugins/Experimental/Avalanche/Source/AvalancheMask/Private/Handling/AvaMaskDesignedMaterialHandle.h b/Engine/Plugins/Experimental/Avalanche/Source/AvalancheMask/Private/Handling/AvaMaskDesignedMaterialHandle.h
--- a/Engine/Plugins/Experimental/Avalanche/Source/AvalancheMask/Private/Handling/AvaMaskDesignedMaterialHandle.h
+++ b/Engine/Plugins/Experimental/Avalanche/Source/AvalancheMask/Private/Handling/AvaMaskDesignedMaterialHandle.h
@@ -64,7 +64,7 @@ public:
     virtual bool IsValid() const override;
     // ~End IAvaObjectHandle
 
-	static bool IsSupported(const UStruct* InStruct, const TVariant<UObject*, FStructView>& InInstance, FName InTag = NAME_None);
+	static bool IsSupported(const FAvaMaskMaterialReference& InInstance, FName InTag = NAME_None);
 
 protected:
 	virtual UMaterialInstanceDynamic* GetMaterialInstance() override;
diff --git a/Engine/Plugins/Experimental/Avalanche/Source/AvalancheMask/Private/Handling/AvaMaskMaterialInstanceHandle.cpp b/Engine/Plugins/Experimental/Avalanche/Source/AvalancheMask/Private/Handling/AvaMaskMaterialInstanceHandle.cpp
--- a/Engine/Plugins/Experimental/Avalanche/Source/AvalancheMask/Private/Handling/AvaMaskMaterialInstanceHandle.cpp
+++ b/Engine/Plugins/Experimental/Avalanche/Source/AvalancheMask/Private/Handling/AvaMaskMaterialInstanceHandle.cpp
@@ -125,13 +125,9 @@ bool FAvaMaskMaterialInstanceHandle::IsValid() const
 	return FAvaMaterialInstanceHandle::IsValid();
 }
 
-bool FAvaMaskMaterialInstanceHandle::IsSupported(
-	const UStruct* InStruct
-	, const TVariant<UObject*, FStructView>& InInstance
-	, FName InTag)
+bool FAvaMaskMaterialInstanceHandle::IsSupported(const FAvaMaskMaterialReference& InInstance, FName InTag)
 {
-	return InTag == UE::AvaMask::Internal::HandleTag
-		&& Super::IsSupported(InStruct, InInstance, InTag);
+	return InTag == UE::AvaMask::Internal::HandleTag && Super::IsSupported(InInstance, InTag);
 }
 
 UMaterialInstanceDynamic* FAvaMaskMaterialInstanceHandle::GetMaterialInstance()
diff --git a/Engine/Plugins/Experimental/Avalanche/Source/AvalancheMask/Private/Handling/AvaMaskMaterialInstanceHandle.h b/Engine/Plugins/Experimental/Avalanche/Source/AvalancheMask/Private/Handling/AvaMaskMaterialInstanceHandle.h
--- a/Engine/Plugins/Experimental/Avalanche/Source/AvalancheMask/Private/Handling/AvaMaskMaterialInstanceHandle.h
+++ b/Engine/Plugins/Experimental/Avalanche/Source/AvalancheMask/Private/Handling/AvaMaskMaterialInstanceHandle.h
@@ -11,6 +11,7 @@
 
 class UMaterialInstanceDynamic;
 class UMaterialInterface;
+struct FAvaMaskMaterialReference;
 
 USTRUCT()
 struct FAvaMaskMaterialInstanceHandleData
@@ -62,7 +63,7 @@ public:
 	virtual bool IsValid() const override;
 	// ~End IAvaObjectHandle
 
-	static bool IsSupported(const UStruct* InStruct, const TVariant<UObject*, FStructView>& InInstance, FName InTag = NAME_None);
+	static bool IsSupported(const FAvaMaskMaterialReference& InInstance, FName InTag = NAME_None);
 
 protected:
 	virtual UMaterialInstanceDynamic* GetMaterialInstance() override;
diff --git a/Engine/Plugins/Experimental/Avalanche/Source/AvalancheMask/Private/Handling/AvaMaskMaterialReference.cpp b/Engine/Plugins/Experimental/Avalanche/Source/AvalancheMask/Private/Handling/AvaMaskMaterialReference.cpp
new file mode 100644
--- /dev/null
+++ b/Engine/Plugins/Experimental/Avalanche/Source/AvalancheMask/Private/Handling/AvaMaskMaterialReference.cpp
@@ -0,0 +1,24 @@
+// Copyright Epic Games, Inc. All Rights Reserved.
+
+#include "AvaMaskMaterialReference.h"
+
+FString FAvaMaskMaterialReference::ToString() const
+{
+	const UObject* const Object = GetObject();
+	if (!Object)
+	{
+		return TEXT("(Invalid Reference)");
+	}
+
+	if (Index == INDEX_NONE)
+	{
+		return Object->GetName();
+	}
+
+	return FString::Printf(TEXT("%s [%d]"), *Object->GetName(), Index);
+}
+
+UObject* FAvaMaskMaterialReference::GetObject() const
+{
+	return ObjectWeak.Get();
+}
diff --git a/Engine/Plugins/Experimental/Avalanche/Source/AvalancheMask/Private/Handling/AvaMaskMaterialReference.h b/Engine/Plugins/Experimental/Avalanche/Source/AvalancheMask/Private/Handling/AvaMaskMaterialReference.h
new file mode 100644
--- /dev/null
+++ b/Engine/Plugins/Experimental/Avalanche/Source/AvalancheMask/Private/Handling/AvaMaskMaterialReference.h
@@ -0,0 +1,34 @@
+// Copyright Epic Games, Inc. All Rights Reserved.
+
+#pragma once
+
+#include "UObject/WeakObjectPtr.h"
+
+/** Describes a reference to an object containing a material, or the material itself */
+struct FAvaMaskMaterialReference
+{
+	FAvaMaskMaterialReference() = default;
+
+	explicit FAvaMaskMaterialReference(UObject* InObject, int32 InIndex = INDEX_NONE)
+		: ObjectWeak(InObject)
+		, Index(InIndex)
+	{
+	}
+
+	FString ToString() const;
+
+	/** Resolves the weak object reference */
+	UObject* GetObject() const;
+
+	template<typename T>
+	T* GetTypedObject() const
+	{
+		return Cast<T>(GetObject());
+	}
+
+	/** The material itself or the object holding the material */
+	FWeakObjectPtr ObjectWeak;
+
+	/** Optional index to further identify a material in the object */
+	int32 Index = INDEX_NONE;
+};
diff --git a/Engine/Plugins/Experimental/Avalanche/Source/AvalancheMask/Private/Handling/AvaMaskMediaPlateMaterialHandle.cpp b/Engine/Plugins/Experimental/Avalanche/Source/AvalancheMask/Private/Handling/AvaMaskMediaPlateMaterialHandle.cpp
--- a/Engine/Plugins/Experimental/Avalanche/Source/AvalancheMask/Private/Handling/AvaMaskMediaPlateMaterialHandle.cpp
+++ b/Engine/Plugins/Experimental/Avalanche/Source/AvalancheMask/Private/Handling/AvaMaskMediaPlateMaterialHandle.cpp
@@ -2,6 +2,7 @@
 
 #include "AvaMaskMediaPlateMaterialHandle.h"
 
+#include "AvaMaskMaterialReference.h"
 #include "AvaMaterialInstanceHandle.h"
 #include "Components/StaticMeshComponent.h"
 #include "Materials/MaterialInstance.h"
@@ -23,17 +24,17 @@ void FAvaMaskMediaPlateMaterialHandle::SetBlendMode(const EBlendMode InBlendMode
 	Super::SetBlendMode(InBlendMode);
 }
 
-bool FAvaMaskMediaPlateMaterialHandle::IsSupported(const UStruct* InStruct, const TVariant<UObject*, FStructView>& InInstance, FName InTag)
+bool FAvaMaskMediaPlateMaterialHandle::IsSupported(const FAvaMaskMaterialReference& InInstance, FName InTag)
 {
 	// Make sure this is generally supported as a material instance.
-	if (!FAvaMaterialInstanceHandle::IsSupported(InStruct, InInstance, InTag))
+	if (!FAvaMaterialInstanceHandle::IsSupported(InInstance, InTag))
 	{
 		return false;
 	}
 
 	// MediaPlate makes either a MIC or MID.
 	// Check if the material is embedded under the static mesh component.
-	if (const UMaterialInstance* MaterialInstance = Cast<UMaterialInstance>(InInstance.Get<UObject*>()))
+	if (const UMaterialInstance* MaterialInstance = InInstance.GetTypedObject<UMaterialInstance>())
 	{
 		if (const UStaticMeshComponent* StaticMeshComponent = Cast<UStaticMeshComponent>(MaterialInstance->GetOuter()))
 		{
diff --git a/Engine/Plugins/Experimental/Avalanche/Source/AvalancheMask/Private/Handling/AvaMaskMediaPlateMaterialHandle.h b/Engine/Plugins/Experimental/Avalanche/Source/AvalancheMask/Private/Handling/AvaMaskMediaPlateMaterialHandle.h
--- a/Engine/Plugins/Experimental/Avalanche/Source/AvalancheMask/Private/Handling/AvaMaskMediaPlateMaterialHandle.h
+++ b/Engine/Plugins/Experimental/Avalanche/Source/AvalancheMask/Private/Handling/AvaMaskMediaPlateMaterialHandle.h
@@ -16,7 +16,7 @@ public:
 	virtual void SetBlendMode(const EBlendMode InBlendMode) override;
 	//~ End IAvaMaskMaterialHandle
 
-	static bool IsSupported(const UStruct* InStruct, const TVariant<UObject*, FStructView>& InInstance, FName InTag = NAME_None);
+	static bool IsSupported(const FAvaMaskMaterialReference& InInstance, FName InTag = NAME_None);
 
 protected:
 	//~ Begin FAvaMaskMaterialInstanceHandle
diff --git a/Engine/Plugins/Experimental/Avalanche/Source/AvalancheMask/Private/Handling/AvaMaskParametricMaterialHandle.cpp b/Engine/Plugins/Experimental/Avalanche/Source/AvalancheMask/Private/Handling/AvaMaskParametricMaterialHandle.cpp
--- a/Engine/Plugins/Experimental/Avalanche/Source/AvalancheMask/Private/Handling/AvaMaskParametricMaterialHandle.cpp
+++ b/Engine/Plugins/Experimental/Avalanche/Source/AvalancheMask/Private/Handling/AvaMaskParametricMaterialHandle.cpp
@@ -8,7 +8,7 @@
 #include "Handling/AvaParametricMaterialHandle.h"
 #include "Materials/MaterialInstanceDynamic.h"
 
-FAvaMaskParametricMaterialHandle::FAvaMaskParametricMaterialHandle(const FStructView& InParametricMaterial)
+FAvaMaskParametricMaterialHandle::FAvaMaskParametricMaterialHandle(const FAvaMaskMaterialReference& InParametricMaterial)
 	: FAvaParametricMaterialHandle(InParametricMaterial)
 {
 	if (FAvaShapeParametricMaterial* ParametricMtl = GetParametricMaterial())
@@ -20,20 +20,15 @@ FAvaMaskParametricMaterialHandle::FAvaMaskParametricMaterialHandle(const FStruct
 			LastMaterialInstance = MID;
 		}
 
-		if (!FAvaShapeParametricMaterial::OnMaterialParameterChanged().IsBoundToObject(this))
-		{
-			FAvaShapeParametricMaterial::OnMaterialParameterChanged().AddRaw(this, &FAvaMaskParametricMaterialHandle::OnMaterialChanged);
-		}
+		OnMaterialParameterChangedHandle = FAvaShapeParametricMaterial::OnMaterialParameterChanged().AddRaw(this
+			, &FAvaMaskParametricMaterialHandle::OnMaterialChanged);
 	}
 }
 
 FAvaMaskParametricMaterialHandle::~FAvaMaskParametricMaterialHandle()
 {
-	// @note: if this doesn't cover all situations, the handle needs to be aware of the struct's outer to check validity
-	if (!GIsGarbageCollecting)
-	{
-		FAvaShapeParametricMaterial::OnMaterialParameterChanged().RemoveAll(this);
-	}
+	FAvaShapeParametricMaterial::OnMaterialParameterChanged().Remove(OnMaterialParameterChangedHandle);
+	OnMaterialParameterChangedHandle.Reset();
 }
 
 bool FAvaMaskParametricMaterialHandle::GetMaskParameters(FAvaMask2DMaterialParameters& OutParameters)
@@ -239,13 +234,9 @@ bool FAvaMaskParametricMaterialHandle::IsValid() const
 	return FAvaParametricMaterialHandle::IsValid();
 }
 
-bool FAvaMaskParametricMaterialHandle::IsSupported(
-	const UStruct* InStruct
-	, const TVariant<UObject*, FStructView>& InInstance
-	, FName InTag)
+bool FAvaMaskParametricMaterialHandle::IsSupported(const FAvaMaskMaterialReference& InInstance, FName InTag)
 {
-	return InTag == UE::AvaMask::Internal::HandleTag
-		&& FAvaParametricMaterialHandle::IsSupported(InStruct, InInstance, InTag);
+	return InTag == UE::AvaMask::Internal::HandleTag && FAvaParametricMaterialHandle::IsSupported(InInstance, InTag);
 }
 
 UMaterialInstanceDynamic* FAvaMaskParametricMaterialHandle::GetMaterialInstance()
diff --git a/Engine/Plugins/Experimental/Avalanche/Source/AvalancheMask/Private/Handling/AvaMaskParametricMaterialHandle.h b/Engine/Plugins/Experimental/Avalanche/Source/AvalancheMask/Private/Handling/AvaMaskParametricMaterialHandle.h
--- a/Engine/Plugins/Experimental/Avalanche/Source/AvalancheMask/Private/Handling/AvaMaskParametricMaterialHandle.h
+++ b/Engine/Plugins/Experimental/Avalanche/Source/AvalancheMask/Private/Handling/AvaMaskParametricMaterialHandle.h
@@ -32,7 +32,7 @@ class FAvaMaskParametricMaterialHandle
 public:
 	UE_AVA_INHERITS_WITH_SUPER(FAvaMaskParametricMaterialHandle, FAvaParametricMaterialHandle);
 	
-	explicit FAvaMaskParametricMaterialHandle(const FStructView& InParametricMaterial);
+	explicit FAvaMaskParametricMaterialHandle(const FAvaMaskMaterialReference& InParametricMaterial);
 	virtual ~FAvaMaskParametricMaterialHandle() override;
 
 	// ~Begin IAvaMaskMaterialHandle
@@ -59,7 +59,7 @@ public:
 	virtual bool IsValid() const override;
 	// ~End IAvaObjectHandle
 
-	static bool IsSupported(const UStruct* InStruct, const TVariant<UObject*, FStructView>& InInstance, FName InTag = NAME_None);
+	static bool IsSupported(const FAvaMaskMaterialReference& InInstance, FName InTag = NAME_None);
 
 protected:
 	virtual UMaterialInstanceDynamic* GetMaterialInstance() override;
@@ -71,4 +71,6 @@ private:
 	EBlendMode LastBlendMode;
 	TWeakObjectPtr<UMaterialInstance> LastMaterialInstance;
 	TWeakObjectPtr<UMaterialInstanceDynamic> ChildMaterial = nullptr;
+
+	FDelegateHandle OnMaterialParameterChangedHandle;
 };
diff --git a/Engine/Plugins/Experimental/Avalanche/Source/AvalancheMask/Private/Handling/AvaMaskText3DActorMaterialCollectionHandle.cpp b/Engine/Plugins/Experimental/Avalanche/Source/AvalancheMask/Private/Handling/AvaMaskText3DActorMaterialCollectionHandle.cpp
--- a/Engine/Plugins/Experimental/Avalanche/Source/AvalancheMask/Private/Handling/AvaMaskText3DActorMaterialCollectionHandle.cpp
+++ b/Engine/Plugins/Experimental/Avalanche/Source/AvalancheMask/Private/Handling/AvaMaskText3DActorMaterialCollectionHandle.cpp
@@ -3,6 +3,7 @@
 #include "Handling/AvaMaskText3DActorMaterialCollectionHandle.h"
 
 #include "AvaMaskLog.h"
+#include "AvaMaskMaterialReference.h"
 #include "AvaMaskUtilities.h"
 #include "AvaObjectHandleSubsystem.h"
 #include "AvaText3DComponent.h"
@@ -305,18 +306,11 @@ bool FAvaMaskText3DActorMaterialCollectionHandle::IsValid() const
 	return WeakActor.IsValid() && WeakComponent.IsValid();
 }
 
-bool FAvaMaskText3DActorMaterialCollectionHandle::IsSupported(
-	const UStruct* InStruct
-	, const TVariant<UObject*, FStructView>& InInstance
-	, FName InTag)
+bool FAvaMaskText3DActorMaterialCollectionHandle::IsSupported(const FAvaMaskMaterialReference& InInstance, FName InTag)
 {
-	if (InTag == UE::AvaMask::Internal::HandleTag
-		&& InStruct
-		&& ::IsValid(InStruct)
-		&& InStruct->IsChildOf<AActor>()
-		&& InInstance.TryGet<UObject*>() != nullptr)
+	if (InTag == UE::AvaMask::Internal::HandleTag)
 	{
-		if (const AActor* Actor = Cast<AActor>(InInstance.Get<UObject*>()))
+		if (const AActor* Actor = InInstance.GetTypedObject<AActor>())
 		{
 			return Actor->GetComponentByClass<UText3DComponent>() != nullptr;
 		}
@@ -750,18 +744,11 @@ bool FAvaMaskAvaTextActorMaterialCollectionHandle::IsValid() const
 	return WeakActor.IsValid() && WeakComponent.IsValid();
 }
 
-bool FAvaMaskAvaTextActorMaterialCollectionHandle::IsSupported(
-	const UStruct* InStruct
-	, const TVariant<UObject*, FStructView>& InInstance
-	, FName InTag)
+bool FAvaMaskAvaTextActorMaterialCollectionHandle::IsSupported(const FAvaMaskMaterialReference& InInstance, FName InTag)
 {
-	if (InTag == UE::AvaMask::Internal::HandleTag
-		&& InStruct
-		&& ::IsValid(InStruct)
-		&& InStruct->IsChildOf<AActor>()
-		&& InInstance.TryGet<UObject*>() != nullptr)
+	if (InTag == UE::AvaMask::Internal::HandleTag)
 	{
-		if (const AActor* Actor = Cast<AActor>(InInstance.Get<UObject*>()))
+		if (const AActor* Actor = InInstance.GetTypedObject<AActor>())
 		{
 			return Actor->GetComponentByClass<UAvaText3DComponent>() != nullptr;
 		}
diff --git a/Engine/Plugins/Experimental/Avalanche/Source/AvalancheMask/Private/Handling/AvaMaskText3DActorMaterialCollectionHandle.h b/Engine/Plugins/Experimental/Avalanche/Source/AvalancheMask/Private/Handling/AvaMaskText3DActorMaterialCollectionHandle.h
--- a/Engine/Plugins/Experimental/Avalanche/Source/AvalancheMask/Private/Handling/AvaMaskText3DActorMaterialCollectionHandle.h
+++ b/Engine/Plugins/Experimental/Avalanche/Source/AvalancheMask/Private/Handling/AvaMaskText3DActorMaterialCollectionHandle.h
@@ -20,7 +20,7 @@ class UText3DComponent;
 class AText3DActor;
 class AActor;
 class UMaterialInterface;
-struct FStructView;
+struct FAvaMaskMaterialReference;
 
 USTRUCT()
 struct FAvaMaskText3DActorMaterialCollectionHandleData
@@ -83,7 +83,7 @@ public:
 	virtual bool IsValid() const override;
 	// ~End IAvaMaskMaterialCollectionProvider
 
-	static bool IsSupported(const UStruct* InStruct, const TVariant<UObject*, FStructView>& InInstance, FName InTag = NAME_None);
+	static bool IsSupported(const FAvaMaskMaterialReference& InInstance, FName InTag = NAME_None);
 
 protected:
 	// ~Begin TAvaMaskMaterialCollectionHandle
@@ -158,7 +158,7 @@ public:
 	virtual bool IsValid() const override;
 	// ~End IAvaMaskMaterialCollectionHandle
 
-	static bool IsSupported(const UStruct* InStruct, const TVariant<UObject*, FStructView>& InInstance, FName InTag = NAME_None);
+	static bool IsSupported(const FAvaMaskMaterialReference& InInstance, FName InTag = NAME_None);
 
 protected:
 	// ~Begin TAvaMaskMaterialCollectionHandle
diff --git a/Engine/Plugins/Experimental/Avalanche/Source/AvalancheMask/Private/Handling/AvaMaterialInstanceHandle.cpp b/Engine/Plugins/Experimental/Avalanche/Source/AvalancheMask/Private/Handling/AvaMaterialInstanceHandle.cpp
--- a/Engine/Plugins/Experimental/Avalanche/Source/AvalancheMask/Private/Handling/AvaMaterialInstanceHandle.cpp
+++ b/Engine/Plugins/Experimental/Avalanche/Source/AvalancheMask/Private/Handling/AvaMaterialInstanceHandle.cpp
@@ -2,6 +2,7 @@
 
 #include "Handling/AvaMaterialInstanceHandle.h"
 
+#include "AvaMaskMaterialReference.h"
 #include "Engine/Texture.h"
 #include "Materials/MaterialInstanceDynamic.h"
 #include "Materials/MaterialInterface.h"
@@ -48,15 +49,9 @@ bool FAvaMaterialInstanceHandle::IsValid() const
 	return WeakParentMaterial.IsValid();
 }
 
-bool FAvaMaterialInstanceHandle::IsSupported(
-	const UStruct* InStruct
-	, const TVariant<UObject*, FStructView>& InInstance
-    , FName InTag)
+bool FAvaMaterialInstanceHandle::IsSupported(const FAvaMaskMaterialReference& InInstance, FName InTag)
 {
-	return InStruct
-		&& ::IsValid(InStruct)
-		&& InStruct->IsChildOf(UMaterialInterface::StaticClass())
-		&& InInstance.TryGet<UObject*>() != nullptr;
+	return !!InInstance.GetTypedObject<UMaterialInterface>();
 }
 
 UMaterialInterface* FAvaMaterialInstanceHandle::GetParentMaterial() const
diff --git a/Engine/Plugins/Experimental/Avalanche/Source/AvalancheMask/Private/Handling/AvaMaterialInstanceHandle.h b/Engine/Plugins/Experimental/Avalanche/Source/AvalancheMask/Private/Handling/AvaMaterialInstanceHandle.h
--- a/Engine/Plugins/Experimental/Avalanche/Source/AvalancheMask/Private/Handling/AvaMaterialInstanceHandle.h
+++ b/Engine/Plugins/Experimental/Avalanche/Source/AvalancheMask/Private/Handling/AvaMaterialInstanceHandle.h
@@ -12,6 +12,7 @@
 class UMaterialInstance;
 class UMaterialInstanceDynamic;
 class UMaterialInterface;
+struct FAvaMaskMaterialReference;
 
 class FAvaMaterialInstanceHandle
 	: public IAvaMaterialHandle
@@ -32,7 +33,7 @@ public:
 	virtual bool IsValid() const override;
 	// ~End IAvaObjectHandle
 
-	static bool IsSupported(const UStruct* InStruct, const TVariant<UObject*, FStructView>& InInstance, FName InTag = NAME_None);
+	static bool IsSupported(const FAvaMaskMaterialReference& InInstance, FName InTag = NAME_None);
 
 protected:
 	virtual UMaterialInterface* GetParentMaterial() const;
diff --git a/Engine/Plugins/Experimental/Avalanche/Source/AvalancheMask/Private/Handling/AvaObjectHandleSubsystem.cpp b/Engine/Plugins/Experimental/Avalanche/Source/AvalancheMask/Private/Handling/AvaObjectHandleSubsystem.cpp
--- a/Engine/Plugins/Experimental/Avalanche/Source/AvalancheMask/Private/Handling/AvaObjectHandleSubsystem.cpp
+++ b/Engine/Plugins/Experimental/Avalanche/Source/AvalancheMask/Private/Handling/AvaObjectHandleSubsystem.cpp
@@ -6,6 +6,7 @@
 #include "AvaMaskActorMaterialCollectionHandle.h"
 #include "AvaMaskAvaShapeMaterialCollectionHandle.h"
 #include "AvaMaskLog.h"
+#include "AvaMaskMaterialReference.h"
 #include "AvaMaskMediaPlateMaterialHandle.h"
 #include "AvaMaskText3DActorMaterialCollectionHandle.h"
 #include "Engine/Engine.h"
@@ -25,34 +26,29 @@ void UAvaObjectHandleSubsystem::Initialize(FSubsystemCollectionBase& Collection)
 	FindObjectHandleFactories();
 }
 
-TSharedPtr<IAvaObjectHandle> UAvaObjectHandleSubsystem::MakeHandle(
-	const UStruct* InStruct
-	, const TVariant<UObject*, FStructView>& InInstance
-	, FName InTag)
+TSharedPtr<IAvaObjectHandle> UAvaObjectHandleSubsystem::MakeHandleDirect(const FAvaMaskMaterialReference& InInstance, FName InTag)
 {
 	for (const TPair<FIsSupportedFunction, FMakeHandleFunction>& SupportedFactoryPair : ObjectHandleFactories)
 	{
-		if (SupportedFactoryPair.Key(InStruct, InInstance, InTag))
+		if (SupportedFactoryPair.Key(InInstance, InTag))
 		{
 			TSharedPtr<IAvaObjectHandle> Handle = SupportedFactoryPair.Value(InInstance);
 			if (!UE::Ava::Internal::IsHandleValid(Handle))
 			{
-				const FString ObjectName = GetInstanceName(InStruct, InInstance);
 				if (InTag.IsNone())
 				{
-					UE_LOG(LogAvaMask, Warning, TEXT("Object Handle for '%s' was created but invalid."), *ObjectName)
+					UE_LOG(LogAvaMask, Warning, TEXT("Object Handle for '%s' was created but invalid."), *InInstance.ToString())
 				}
 				else
 				{
-					UE_LOG(LogAvaMask, Warning, TEXT("Object Handle for '%s' (with tag '%s') was created but invalid."), *ObjectName, *InTag.ToString())
+					UE_LOG(LogAvaMask, Warning, TEXT("Object Handle for '%s' (with tag '%s') was created but invalid."), *InInstance.ToString(), *InTag.ToString())
 				}
 			}
 			return Handle;
 		}
 	}
 
-	UE_LOG(LogAvaMask, Display, TEXT("No ObjectHandle found for '%s'"), *InStruct->GetName());
-	
+	UE_LOG(LogAvaMask, Display, TEXT("No ObjectHandle found for '%s'"), *InInstance.ToString());
 	return nullptr;
 }
 
@@ -63,30 +59,30 @@ void UAvaObjectHandleSubsystem::FindObjectHandleFactories()
 	{
 		ObjectHandleFactories.Add({
 			FAvaMaskAvaTextActorMaterialCollectionHandle::IsSupported
-			, [](const TVariant<UObject*, FStructView>& InObject) -> TSharedPtr<IAvaMaskMaterialCollectionHandle>
+			, [](const FAvaMaskMaterialReference& InObject) -> TSharedPtr<IAvaMaskMaterialCollectionHandle>
 			{
-				return MakeShared<FAvaMaskAvaTextActorMaterialCollectionHandle>(Cast<AActor>(InObject.Get<UObject*>()));
+				return MakeShared<FAvaMaskAvaTextActorMaterialCollectionHandle>(InObject.GetTypedObject<AActor>());
 			}});
-	
+
 		ObjectHandleFactories.Add({
 			FAvaMaskAvaShapeMaterialCollectionHandle::IsSupported
-			, [](const TVariant<UObject*, FStructView>& InObject) -> TSharedPtr<IAvaMaskMaterialCollectionHandle>
+			, [](const FAvaMaskMaterialReference& InObject) -> TSharedPtr<IAvaMaskMaterialCollectionHandle>
 			{
-				return MakeShared<FAvaMaskAvaShapeMaterialCollectionHandle>(Cast<AActor>(InObject.Get<UObject*>()));
+				return MakeShared<FAvaMaskAvaShapeMaterialCollectionHandle>(InObject.GetTypedObject<AActor>());
 			}});
-	
+
 		ObjectHandleFactories.Add({
 			FAvaMaskText3DActorMaterialCollectionHandle::IsSupported
-			, [](const TVariant<UObject*, FStructView>& InObject) -> TSharedPtr<IAvaMaskMaterialCollectionHandle>
+			, [](const FAvaMaskMaterialReference& InObject) -> TSharedPtr<IAvaMaskMaterialCollectionHandle>
 			{
-				return MakeShared<FAvaMaskText3DActorMaterialCollectionHandle>(Cast<AActor>(InObject.Get<UObject*>()));
+				return MakeShared<FAvaMaskText3DActorMaterialCollectionHandle>(InObject.GetTypedObject<AActor>());
 			}});
-	
+
 		ObjectHandleFactories.Add({
 			FAvaMaskActorMaterialCollectionHandle::IsSupported
-			, [](const TVariant<UObject*, FStructView>& InObject) -> TSharedPtr<IAvaMaskMaterialCollectionHandle>
+			, [](const FAvaMaskMaterialReference& InObject) -> TSharedPtr<IAvaMaskMaterialCollectionHandle>
 			{
-				return MakeShared<FAvaMaskActorMaterialCollectionHandle>(Cast<AActor>(InObject.Get<UObject*>()));
+				return MakeShared<FAvaMaskActorMaterialCollectionHandle>(InObject.GetTypedObject<AActor>());
 			}});
 	}
 
@@ -94,66 +90,30 @@ void UAvaObjectHandleSubsystem::FindObjectHandleFactories()
 	{
 		ObjectHandleFactories.Add({
 			FAvaMaskDesignedMaterialHandle::IsSupported
-			, [](const TVariant<UObject*, FStructView>& InMaterial) -> TSharedPtr<IAvaMaskMaterialHandle>
+			, [](const FAvaMaskMaterialReference& InMaterial) -> TSharedPtr<IAvaMaskMaterialHandle>
 			{
-				return MakeShared<FAvaMaskDesignedMaterialHandle>(Cast<UDynamicMaterialInstance>(InMaterial.Get<UObject*>()));
+				return MakeShared<FAvaMaskDesignedMaterialHandle>(InMaterial.GetTypedObject<UDynamicMaterialInstance>());
 			}});
 
 		ObjectHandleFactories.Add({
 			FAvaMaskParametricMaterialHandle::IsSupported
-			, [](const TVariant<UObject*, FStructView>& InMaterial) -> TSharedPtr<IAvaMaskMaterialHandle>
+			, [](const FAvaMaskMaterialReference& InMaterial) -> TSharedPtr<IAvaMaskMaterialHandle>
 			{
-				return MakeShared<FAvaMaskParametricMaterialHandle>(InMaterial.Get<FStructView>());
+				return MakeShared<FAvaMaskParametricMaterialHandle>(InMaterial);
 			}});
 
 		ObjectHandleFactories.Add({
 			FAvaMaskMediaPlateMaterialHandle::IsSupported
-			, [](const TVariant<UObject*, FStructView>& InMaterial) -> TSharedPtr<IAvaMaskMaterialHandle>
+			, [](const FAvaMaskMaterialReference& InMaterial) -> TSharedPtr<IAvaMaskMaterialHandle>
 			{
-				return MakeShared<FAvaMaskMediaPlateMaterialHandle>(Cast<UMaterialInterface>(InMaterial.Get<UObject*>()));
+				return MakeShared<FAvaMaskMediaPlateMaterialHandle>(InMaterial.GetTypedObject<UMaterialInterface>());
 			}});
 
 		ObjectHandleFactories.Add({
 			FAvaMaskMaterialInstanceHandle::IsSupported
-			, [](const TVariant<UObject*, FStructView>& InMaterial) -> TSharedPtr<IAvaMaskMaterialHandle>
+			, [](const FAvaMaskMaterialReference& InMaterial) -> TSharedPtr<IAvaMaskMaterialHandle>
 			{
-				return MakeShared<FAvaMaskMaterialInstanceHandle>(Cast<UMaterialInterface>(InMaterial.Get<UObject*>()));
+				return MakeShared<FAvaMaskMaterialInstanceHandle>(InMaterial.GetTypedObject<UMaterialInterface>());
 			}});
 	}
 }
-
-FString UAvaObjectHandleSubsystem::GetInstanceName(const UStruct* InStruct, const TVariant<UObject*, FStructView>& InInstance)
-{
-	struct FNameGetter
-	{
-		FString operator()(const UObject* InValue) const
-		{
-			if (::IsValid(InValue))
-			{
-				return InValue->GetName();
-			}
-
-			return TEXT("");
-		}
-		
-		FString operator()(const FStructView& InValue) const
-		{
-			return InValue.GetScriptStruct()->GetName();
-		}
-	};
-
-	FNameGetter NameGetter;
-
-	if (FString Name = Visit(NameGetter, InInstance);
-		!Name.IsEmpty())
-	{
-		return Name;
-	}
-
-	if (const UStruct* Struct = InStruct)
-	{
-		return Struct->GetName();
-	}
-
-	return TEXT("");
-}
diff --git a/Engine/Plugins/Experimental/Avalanche/Source/AvalancheMask/Private/Handling/AvaObjectHandleSubsystem.h b/Engine/Plugins/Experimental/Avalanche/Source/AvalancheMask/Private/Handling/AvaObjectHandleSubsystem.h
--- a/Engine/Plugins/Experimental/Avalanche/Source/AvalancheMask/Private/Handling/AvaObjectHandleSubsystem.h
+++ b/Engine/Plugins/Experimental/Avalanche/Source/AvalancheMask/Private/Handling/AvaObjectHandleSubsystem.h
@@ -4,18 +4,15 @@
 
 #include "AvaMaskLog.h"
 #include "IAvaObjectHandle.h"
-#include "Misc/TVariant.h"
-#include "StructUtils/StructView.h"
 #include "Subsystems/EngineSubsystem.h"
-
 #include "AvaObjectHandleSubsystem.generated.h"
 
 class UAvaObjectHandleSubsystem;
+struct FAvaMaskMaterialReference;
 
 /** Responsible for providing Handlers for a given UObject. */
 UCLASS()
-class UAvaObjectHandleSubsystem
-	: public UEngineSubsystem
+class UAvaObjectHandleSubsystem : public UEngineSubsystem
 {
 	GENERATED_BODY()
 
@@ -24,7 +21,7 @@ public:
 	virtual void Initialize(FSubsystemCollectionBase& Collection) override;
 	// ~End USubsystem
 
-	TSharedPtr<IAvaObjectHandle> MakeHandle(const UStruct* InStruct, const TVariant<UObject*, FStructView>& InInstance, FName InTag = NAME_None);
+	TSharedPtr<IAvaObjectHandle> MakeHandleDirect(const FAvaMaskMaterialReference& InInstance, FName InTag = NAME_None);
 
 	template <typename HandleType, typename ObjectType = UObject
 		UE_REQUIRES(!std::is_const_v<HandleType> && std::is_base_of_v<IAvaObjectHandle, HandleType>
@@ -36,42 +33,20 @@ public:
 			UE_LOG(LogAvaMask, Warning, TEXT("Invalid or null object provided to MakeHandle"));
 			return nullptr;
 		}
-
-		TVariant<UObject*, FStructView> Variant(TInPlaceType<UObject*>(), InInstance);
-		
-    	return StaticCastSharedPtr<HandleType>(MakeHandle(InInstance->GetClass(), MoveTemp(Variant), InTag));
+    	return StaticCastSharedPtr<HandleType>(MakeHandleDirect(FAvaMaskMaterialReference(InInstance), InTag));
     }
 
-	template <typename HandleType, typename StructType
-		UE_REQUIRES(!std::is_const_v<HandleType> && std::is_base_of_v<IAvaObjectHandle, HandleType>
-				&&	!std::is_const_v<StructType> && TModels_V<CStaticStructProvider, std::decay_t<StructType>>)>
-	TSharedPtr<HandleType> MakeHandle(const FStructView& InInstance, FName InTag = NAME_None)
+	template<typename HandleType UE_REQUIRES(!std::is_const_v<HandleType> && std::is_base_of_v<IAvaObjectHandle, HandleType>)>
+	TSharedPtr<HandleType> MakeHandle(const FAvaMaskMaterialReference& InInstance, FName InTag = NAME_None)
 	{
-		if (!InInstance.IsValid())
-		{
-			UE_LOG(LogAvaMask, Warning, TEXT("Invalid or null StructView provided to MakeHandle"));
-			return nullptr;
-		}
-
-		if (!ensure(InInstance.GetPtr<StructType>()))
-		{
-			UE_LOG(LogAvaMask, Warning, TEXT("StructView does not point to the expected struct type \'%s\'"), *StructType::StaticStruct()->GetName());
-			return nullptr;			
-		}
-		
-		TVariant<UObject*, FStructView> Variant(TInPlaceType<FStructView>(), InInstance);
-		
-		return StaticCastSharedPtr<HandleType>(MakeHandle(InInstance.GetScriptStruct(), MoveTemp(Variant), InTag));
+		return StaticCastSharedPtr<HandleType>(MakeHandleDirect(InInstance, InTag));
 	}
- 
+
 private:
 	void FindObjectHandleFactories();
 
-	FString GetInstanceName(const UStruct* InStruct, const TVariant<UObject*, FStructView>& InInstance);
-
-private:
-	using FIsSupportedFunction = TFunction<bool(const UStruct*, const TVariant<UObject*, FStructView>&, FName)>;
-	using FMakeHandleFunction = TFunction<TSharedPtr<IAvaObjectHandle>(const TVariant<UObject*, FStructView>&)>;
+	using FIsSupportedFunction = TFunction<bool(const FAvaMaskMaterialReference&, FName)>;
+	using FMakeHandleFunction = TFunction<TSharedPtr<IAvaObjectHandle>(const FAvaMaskMaterialReference&)>;
 
 	TArray<TPair<FIsSupportedFunction, FMakeHandleFunction>> ObjectHandleFactories;
 };
diff --git a/Engine/Plugins/Experimental/Avalanche/Source/AvalancheMask/Private/Handling/AvaParametricMaterialHandle.cpp b/Engine/Plugins/Experimental/Avalanche/Source/AvalancheMask/Private/Handling/AvaParametricMaterialHandle.cpp
--- a/Engine/Plugins/Experimental/Avalanche/Source/AvalancheMask/Private/Handling/AvaParametricMaterialHandle.cpp
+++ b/Engine/Plugins/Experimental/Avalanche/Source/AvalancheMask/Private/Handling/AvaParametricMaterialHandle.cpp
@@ -1,12 +1,33 @@
 ﻿// Copyright Epic Games, Inc. All Rights Reserved.
 
 #include "Handling/AvaParametricMaterialHandle.h"
-
 #include "AvaShapeParametricMaterial.h"
+#include "DynamicMeshes/AvaShapeDynMeshBase.h"
 
-FAvaParametricMaterialHandle::FAvaParametricMaterialHandle(const FStructView& InParametricMaterial)
-	: FAvaMaterialInstanceHandle(InParametricMaterial.GetPtr<FAvaShapeParametricMaterial>()->GetMaterial())
-	, ParametricMaterial(InParametricMaterial)
+namespace UE::AvaMask::Private
+{
+	FAvaShapeParametricMaterial* GetParametricMaterialData(const FAvaMaskMaterialReference& InParametricMaterial)
+	{
+		if (UAvaShapeDynamicMeshBase* ShapeComponent = InParametricMaterial.GetTypedObject<UAvaShapeDynamicMeshBase>())
+		{
+			return ShapeComponent->GetParametricMaterial(InParametricMaterial.Index);
+		}
+		return nullptr;
+	}
+
+	UMaterialInterface* GetParametricMaterial(const FAvaMaskMaterialReference& InParametricMaterial)
+	{
+		if (const FAvaShapeParametricMaterial* ParametricMaterial = GetParametricMaterialData(InParametricMaterial))
+		{
+			return ParametricMaterial->GetMaterial();
+		}
+		return nullptr;
+	}
+}
+
+FAvaParametricMaterialHandle::FAvaParametricMaterialHandle(const FAvaMaskMaterialReference& InMaterialReference)
+	: FAvaMaterialInstanceHandle(UE::AvaMask::Private::GetParametricMaterial(InMaterialReference))
+	, MaterialReference(InMaterialReference)
 {
 }
 
@@ -22,28 +43,17 @@ void FAvaParametricMaterialHandle::CopyParametersFrom(UMaterialInstance* InSourc
 
 bool FAvaParametricMaterialHandle::IsValid() const
 {
-	return this && ParametricMaterial.IsValid() && WeakMaterialInstance.IsValid();
+	return !!UE::AvaMask::Private::GetParametricMaterialData(MaterialReference);
 }
 
-bool FAvaParametricMaterialHandle::IsSupported(
-	const UStruct* InStruct
-	, const TVariant<UObject*, FStructView>& InInstance
-	, FName InTag)
+bool FAvaParametricMaterialHandle::IsSupported(const FAvaMaskMaterialReference& InInstance, FName InTag)
 {
-	return InStruct
-		&& ::IsValid(InStruct)
-		&& InStruct->IsChildOf(FAvaShapeParametricMaterial::StaticStruct())
-		&& InInstance.TryGet<FStructView>() != nullptr;
+	return !!UE::AvaMask::Private::GetParametricMaterialData(InInstance);
 }
 
 FAvaShapeParametricMaterial* FAvaParametricMaterialHandle::GetParametricMaterial()
 {
-	if (ParametricMaterial.IsValid())
-	{
-		return ParametricMaterial.GetPtr<FAvaShapeParametricMaterial>();
-	}
-
-	return nullptr;
+	return UE::AvaMask::Private::GetParametricMaterialData(MaterialReference);
 }
 
 UMaterialInstanceDynamic* FAvaParametricMaterialHandle::GetMaterialInstance()
@@ -52,6 +62,5 @@ UMaterialInstanceDynamic* FAvaParametricMaterialHandle::GetMaterialInstance()
 	{
 		return ParametricMtl->GetMaterial();
 	}
-	
 	return FAvaMaterialInstanceHandle::GetMaterialInstance();
 }
diff --git a/Engine/Plugins/Experimental/Avalanche/Source/AvalancheMask/Private/Handling/AvaParametricMaterialHandle.h b/Engine/Plugins/Experimental/Avalanche/Source/AvalancheMask/Private/Handling/AvaParametricMaterialHandle.h
--- a/Engine/Plugins/Experimental/Avalanche/Source/AvalancheMask/Private/Handling/AvaParametricMaterialHandle.h
+++ b/Engine/Plugins/Experimental/Avalanche/Source/AvalancheMask/Private/Handling/AvaParametricMaterialHandle.h
@@ -2,6 +2,7 @@
 
 #pragma once
 
+#include "AvaMaskMaterialReference.h"
 #include "AvaType.h"
 #include "Handling/AvaMaterialInstanceHandle.h"
 #include "IAvaMaskMaterialHandle.h"
@@ -9,13 +10,12 @@
 
 struct FAvaShapeParametricMaterial;
 
-class FAvaParametricMaterialHandle
-	: public FAvaMaterialInstanceHandle
+class FAvaParametricMaterialHandle : public FAvaMaterialInstanceHandle
 {
 public:
 	UE_AVA_INHERITS_WITH_SUPER(FAvaParametricMaterialHandle, FAvaMaterialInstanceHandle);
 	
-	explicit FAvaParametricMaterialHandle(const FStructView& InParametricMaterial);
+	explicit FAvaParametricMaterialHandle(const FAvaMaskMaterialReference& InParametricMaterial);
 
 	// ~Begin IAvaMaterialHandle
 	virtual void CopyParametersFrom(UMaterialInstance* InSourceMaterial) override;
@@ -25,12 +25,11 @@ public:
 	virtual bool IsValid() const override;
 	// ~End IAvaObjectHandle
 
-	static bool IsSupported(const UStruct* InStruct, const TVariant<UObject*, FStructView>& InInstance, FName InTag = NAME_None);
+	static bool IsSupported(const FAvaMaskMaterialReference& InInstance, FName InTag = NAME_None);
 
 protected:
 	FAvaShapeParametricMaterial* GetParametricMaterial();
 	virtual UMaterialInstanceDynamic* GetMaterialInstance() override;
 
-protected:
-	FStructView ParametricMaterial = nullptr;
+	FAvaMaskMaterialReference MaterialReference;
 };
