Motion Design: fix build errors

diff --git a/Engine/Plugins/Experimental/Avalanche/Source/AvalancheMask/Private/Handling/AvaObjectHandleSubsystem.h b/Engine/Plugins/Experimental/Avalanche/Source/AvalancheMask/Private/Handling/AvaObjectHandleSubsystem.h
--- a/Engine/Plugins/Experimental/Avalanche/Source/AvalancheMask/Private/Handling/AvaObjectHandleSubsystem.h
+++ b/Engine/Plugins/Experimental/Avalanche/Source/AvalancheMask/Private/Handling/AvaObjectHandleSubsystem.h
@@ -3,12 +3,12 @@
 #pragma once
 
 #include "AvaMaskLog.h"
+#include "AvaMaskMaterialReference.h"
 #include "IAvaObjectHandle.h"
 #include "Subsystems/EngineSubsystem.h"
 #include "AvaObjectHandleSubsystem.generated.h"
 
 class UAvaObjectHandleSubsystem;
-struct FAvaMaskMaterialReference;
 
 /** Responsible for providing Handlers for a given UObject. */
 UCLASS()
