[Motion Design] Rundown RC referenced assets (page values) get loaded (async) by the playback transitions.

- For Runtime: the playback transition will be loading the referenced assets prior to starting the transition. Works for both local rundown page player and remote playback server.
- For Rundown Editor: the detail panel will do an async load of the assets and refresh when it is ready.

diff --git a/Engine/Plugins/Experimental/Avalanche/Source/AvalancheMedia/Internal/Playback/AvaPlaybackUtils.h b/Engine/Plugins/Experimental/Avalanche/Source/AvalancheMedia/Internal/Playback/AvaPlaybackUtils.h
--- a/Engine/Plugins/Experimental/Avalanche/Source/AvalancheMedia/Internal/Playback/AvaPlaybackUtils.h
+++ b/Engine/Plugins/Experimental/Avalanche/Source/AvalancheMedia/Internal/Playback/AvaPlaybackUtils.h
@@ -4,6 +4,9 @@
 
 #include "Containers/Array.h"
 #include "Containers/UnrealString.h"
+#include "Delegates/Delegate.h"
+#include "Templates/SharedPointer.h"
+#include "UObject/SoftObjectPath.h"
 
 class UPackage;
 struct FAssetData;
@@ -64,4 +67,32 @@ namespace UE::AvaPlayback::Utils
 	{
 		return StaticEnum<InEnumType>()->GetNameStringByValue(static_cast<int64>(InValue));
 	}
+
+	/**
+	 * Utility class to perform async load of assets. 
+	 */
+	class AVALANCHEMEDIA_API FAsyncAssetLoader : public TSharedFromThis<FAsyncAssetLoader>
+	{
+	public:
+		/**
+		 * @brief Issue the async load command. Returns immediately.
+		 * @param InAssetsToLoad Assets to load.
+		 *
+		 * The specified asset to load will added to the current set of pending assets.
+		 */
+		void BeginLoadingAssets(const TArray<FSoftObjectPath>& InAssetsToLoad);
+
+		/** Returns true when all the requested assets have finished loading. */
+		bool IsLoadingCompleted() const
+		{
+			return PendingAssets.IsEmpty();
+		}
+
+		/** Delegate called when all assets have finished loading. */
+		DECLARE_MULTICAST_DELEGATE(FOnLoadingCompleted);
+		FOnLoadingCompleted OnLoadingCompleted;
+
+	private:
+		TSet<FSoftObjectPath> PendingAssets;
+	};
 }
\ No newline at end of file
diff --git a/Engine/Plugins/Experimental/Avalanche/Source/AvalancheMedia/Private/Playback/AvaPlaybackUtils.cpp b/Engine/Plugins/Experimental/Avalanche/Source/AvalancheMedia/Private/Playback/AvaPlaybackUtils.cpp
--- a/Engine/Plugins/Experimental/Avalanche/Source/AvalancheMedia/Private/Playback/AvaPlaybackUtils.cpp
+++ b/Engine/Plugins/Experimental/Avalanche/Source/AvalancheMedia/Private/Playback/AvaPlaybackUtils.cpp
@@ -206,4 +206,25 @@ namespace UE::AvaPlayback::Utils
 	{
 		return FString::Printf(TEXT("[%d]"), GFrameNumber);
 	}
+
+	void FAsyncAssetLoader::BeginLoadingAssets(const TArray<FSoftObjectPath>& InAssetsToLoad)
+	{
+		for (const FSoftObjectPath& AssetToLoad : InAssetsToLoad)
+		{
+			if (AssetToLoad.ResolveObject() // Already loaded.
+				|| PendingAssets.Contains(AssetToLoad)) // Already pending	
+			{
+				continue;
+			}
+
+			AssetToLoad.LoadAsync(FLoadSoftObjectPathAsyncDelegate::CreateSPLambda(this, [this](const FSoftObjectPath& InObjectPath, UObject* InLoadedObject)
+			{
+				PendingAssets.Remove(InObjectPath);
+				if (PendingAssets.IsEmpty())
+				{
+					OnLoadingCompleted.Broadcast();
+				}
+			}));
+		}
+	}
 }
\ No newline at end of file
diff --git a/Engine/Plugins/Experimental/Avalanche/Source/AvalancheMedia/Private/Playback/Transition/AvaPlaybackServerTransition.cpp b/Engine/Plugins/Experimental/Avalanche/Source/AvalancheMedia/Private/Playback/Transition/AvaPlaybackServerTransition.cpp
--- a/Engine/Plugins/Experimental/Avalanche/Source/AvalancheMedia/Private/Playback/Transition/AvaPlaybackServerTransition.cpp
+++ b/Engine/Plugins/Experimental/Avalanche/Source/AvalancheMedia/Private/Playback/Transition/AvaPlaybackServerTransition.cpp
@@ -271,7 +271,7 @@ bool UAvaPlaybackServerTransition::IsVisibilityConstrained(const UAvaPlayable* I
 }
 
 
-bool UAvaPlaybackServerTransition::CanStart(bool& bOutShouldDiscard) const
+bool UAvaPlaybackServerTransition::CanStart(bool& bOutShouldDiscard)
 {
 	using namespace UE::AvaPlaybackServerTransition::Private;
 
@@ -283,6 +283,19 @@ bool UAvaPlaybackServerTransition::CanStart(bool& bOutShouldDiscard) const
 		bOutShouldDiscard = false;
 		return false;
 	}
+	
+	if (!EnterValues.IsEmpty() && !AsyncAssetLoader)
+	{
+		// Start loading the RC referenced assets.
+		AsyncAssetLoader = MakeShared<UE::AvaPlayback::Utils::FAsyncAssetLoader>();
+		TSet<FSoftObjectPath> Assets;
+		for (const TSharedPtr<FAvaPlayableRemoteControlValues>& Values : EnterValues)
+		{
+			FAvaPlayableRemoteControlValues::CollectReferencedAssetPaths(Values->ControllerValues, Assets);
+			FAvaPlayableRemoteControlValues::CollectReferencedAssetPaths(Values->EntityValues, Assets);
+		}
+		AsyncAssetLoader->BeginLoadingAssets(Assets.Array());
+	}
 
 	if (!CanStartTransition(this, EnterPlaybackInstancesWeak, bOutShouldDiscard))
 	{
@@ -302,6 +315,13 @@ bool UAvaPlaybackServerTransition::CanStart(bool& bOutShouldDiscard) const
 		return false;
 	}
 
+	// Wait for RC referenced assets to be loaded.
+	if (AsyncAssetLoader && !AsyncAssetLoader->IsLoadingCompleted())
+	{
+		bOutShouldDiscard = false;
+		return false;
+	}
+
 	bOutShouldDiscard = true;
 	return true;
 }
diff --git a/Engine/Plugins/Experimental/Avalanche/Source/AvalancheMedia/Private/Playback/Transition/AvaPlaybackServerTransition.h b/Engine/Plugins/Experimental/Avalanche/Source/AvalancheMedia/Private/Playback/Transition/AvaPlaybackServerTransition.h
--- a/Engine/Plugins/Experimental/Avalanche/Source/AvalancheMedia/Private/Playback/Transition/AvaPlaybackServerTransition.h
+++ b/Engine/Plugins/Experimental/Avalanche/Source/AvalancheMedia/Private/Playback/Transition/AvaPlaybackServerTransition.h
@@ -12,6 +12,11 @@ class FAvaPlaybackServer;
 class UAvaPlayableTransition;
 class UAvaPlaybackGraph;
 
+namespace UE::AvaPlayback::Utils
+{
+	class FAsyncAssetLoader;
+}
+
 /**
  * Class for creating and tracking playback graph instance transitions on the server.
  * It is responsible for creating the playable transition object when requested from
@@ -50,7 +55,7 @@ public:
 	//~ End IAvaPlayableVisibilityConstraint
 	
 	//~ Begin UAvaPlaybackTransition
-	virtual bool CanStart(bool& bOutShouldDiscard) const override;
+	virtual bool CanStart(bool& bOutShouldDiscard) override;
 	virtual void Start() override;
 	virtual void Stop() override;
 	virtual bool IsRunning() const override;
@@ -95,7 +100,9 @@ protected:
 	TArray<TWeakPtr<FAvaPlaybackInstance>> PlayingPlaybackInstancesWeak;
 	TArray<TWeakPtr<FAvaPlaybackInstance>> ExitPlaybackInstancesWeak;
 	TArray<TSharedPtr<FAvaPlayableRemoteControlValues>> EnterValues;
-	
+
+	TSharedPtr<UE::AvaPlayback::Utils::FAsyncAssetLoader> AsyncAssetLoader;
+
 	UPROPERTY(Transient)
 	TObjectPtr<UAvaPlayableTransition> PlayableTransition;
 };
\ No newline at end of file
diff --git a/Engine/Plugins/Experimental/Avalanche/Source/AvalancheMedia/Private/Playback/Transition/AvaPlaybackTransition.h b/Engine/Plugins/Experimental/Avalanche/Source/AvalancheMedia/Private/Playback/Transition/AvaPlaybackTransition.h
--- a/Engine/Plugins/Experimental/Avalanche/Source/AvalancheMedia/Private/Playback/Transition/AvaPlaybackTransition.h
+++ b/Engine/Plugins/Experimental/Avalanche/Source/AvalancheMedia/Private/Playback/Transition/AvaPlaybackTransition.h
@@ -21,7 +21,7 @@ public:
 	 * @param bOutShouldDiscard indicate if the pending start transition command should be discarded.
 	 * @return true if the transition can start, false otherwise.
 	 */
-	virtual bool CanStart(bool& bOutShouldDiscard) const { bOutShouldDiscard = true; return false; }
+	virtual bool CanStart(bool& bOutShouldDiscard) { bOutShouldDiscard = true; return false; }
 
 	virtual void Start() {}
 
diff --git a/Engine/Plugins/Experimental/Avalanche/Source/AvalancheMedia/Private/Rundown/Transition/AvaRundownPageTransition.cpp b/Engine/Plugins/Experimental/Avalanche/Source/AvalancheMedia/Private/Rundown/Transition/AvaRundownPageTransition.cpp
--- a/Engine/Plugins/Experimental/Avalanche/Source/AvalancheMedia/Private/Rundown/Transition/AvaRundownPageTransition.cpp
+++ b/Engine/Plugins/Experimental/Avalanche/Source/AvalancheMedia/Private/Rundown/Transition/AvaRundownPageTransition.cpp
@@ -179,9 +179,26 @@ bool UAvaRundownPageTransition::IsVisibilityConstrained(const UAvaPlayable* InPl
 	return bIsPlayableInThisTransition && !bAllPlayablesLoaded;
 }
 
-bool UAvaRundownPageTransition::CanStart(bool& bOutShouldDiscard) const
+bool UAvaRundownPageTransition::CanStart(bool& bOutShouldDiscard)
 {
 	using namespace UE::AvaMedia::Rundown::PageTransition::Private;
+	if (!AsyncAssetLoader)
+	{
+		// Start loading the RC referenced assets.
+		AsyncAssetLoader = MakeShared<UE::AvaPlayback::Utils::FAsyncAssetLoader>();
+		TSet<FSoftObjectPath> Assets;
+		for (const TWeakObjectPtr<UAvaRundownPagePlayer>& PlayerWeak : EnterPlayersWeak)
+		{
+			if (const UAvaRundownPagePlayer* Player = PlayerWeak.Get())
+			{
+				const TSharedPtr<FAvaPlayableRemoteControlValues> Values = GetRemoteControlValues(Player);
+				FAvaPlayableRemoteControlValues::CollectReferencedAssetPaths(Values->ControllerValues, Assets);
+				FAvaPlayableRemoteControlValues::CollectReferencedAssetPaths(Values->EntityValues, Assets);
+			}
+		}
+		AsyncAssetLoader->BeginLoadingAssets(Assets.Array());
+	}
+
 	for (const TWeakObjectPtr<UAvaRundownPagePlayer>& PlayerWeak : EnterPlayersWeak)
 	{
 		if (const UAvaRundownPagePlayer* Player = PlayerWeak.Get())
@@ -231,6 +248,13 @@ bool UAvaRundownPageTransition::CanStart(bool& bOutShouldDiscard) const
 			return false;
 		}
 	}
+
+	// Wait for RC referenced assets to be loaded.
+	if (AsyncAssetLoader && !AsyncAssetLoader->IsLoadingCompleted())
+	{
+		bOutShouldDiscard = false;
+		return false;
+	}
 	
 	bOutShouldDiscard = true;
 	return true;
diff --git a/Engine/Plugins/Experimental/Avalanche/Source/AvalancheMedia/Private/Rundown/Transition/AvaRundownPageTransition.h b/Engine/Plugins/Experimental/Avalanche/Source/AvalancheMedia/Private/Rundown/Transition/AvaRundownPageTransition.h
--- a/Engine/Plugins/Experimental/Avalanche/Source/AvalancheMedia/Private/Rundown/Transition/AvaRundownPageTransition.h
+++ b/Engine/Plugins/Experimental/Avalanche/Source/AvalancheMedia/Private/Rundown/Transition/AvaRundownPageTransition.h
@@ -16,6 +16,11 @@ class UAvaRundownPagePlayer;
 class UAvaRundownPlaybackInstancePlayer;
 enum class EAvaPlayableTransitionEntryRole : uint8;
 
+namespace UE::AvaPlayback::Utils
+{
+	class FAsyncAssetLoader;
+}
+
 /**
  * Class for creating and tracking page transitions in the rundown.
  * It is responsible for creating the playable transition object when requested from
@@ -42,7 +47,7 @@ public:
 	//~ End IAvaPlayableVisibilityConstraint
 	
 	//~ Begin UAvaPlaybackTransition
-	virtual bool CanStart(bool& bOutShouldDiscard) const override;
+	virtual bool CanStart(bool& bOutShouldDiscard) override;
 	virtual void Start() override;
 	virtual void Stop() override;
 	virtual bool IsRunning() const override;
@@ -110,6 +115,8 @@ protected:
 	TSet<FAvaTagId> CachedTransitionLayers;
 	
 	TSet<FGuid> InstancesMarkedForDiscard;
+
+	TSharedPtr<UE::AvaPlayback::Utils::FAsyncAssetLoader> AsyncAssetLoader;
 	
 	UPROPERTY(Transient)
 	TObjectPtr<UAvaPlayableTransition> PlayableTransition;
diff --git a/Engine/Plugins/Experimental/Avalanche/Source/AvalancheMediaEditor/Private/Rundown/DetailsView/RemoteControl/Properties/SAvaRundownPageRemoteControlProps.cpp b/Engine/Plugins/Experimental/Avalanche/Source/AvalancheMediaEditor/Private/Rundown/DetailsView/RemoteControl/Properties/SAvaRundownPageRemoteControlProps.cpp
--- a/Engine/Plugins/Experimental/Avalanche/Source/AvalancheMediaEditor/Private/Rundown/DetailsView/RemoteControl/Properties/SAvaRundownPageRemoteControlProps.cpp
+++ b/Engine/Plugins/Experimental/Avalanche/Source/AvalancheMediaEditor/Private/Rundown/DetailsView/RemoteControl/Properties/SAvaRundownPageRemoteControlProps.cpp
@@ -5,6 +5,7 @@
 #include "Modules/ModuleManager.h"
 #include "Playable/AvaPlayableRemoteControl.h"
 #include "Playable/AvaPlayableRemoteControlValues.h"
+#include "Playback/AvaPlaybackUtils.h"
 #include "PropertyEditorModule.h"
 #include "RemoteControlEntity.h"
 #include "RemoteControlPreset.h"
@@ -14,6 +15,7 @@
 #include "Rundown/AvaRundownManagedInstanceCache.h"
 #include "Rundown/AvaRundownPage.h"
 #include "Rundown/DetailsView/RemoteControl/Properties/AvaRundownPagePropertyContextMenu.h"
+#include "Rundown/DetailsView/SAvaRundownPageDetails.h"
 #include "SAvaRundownRCPropertyItemRow.h"
 #include "SlateOptMacros.h"
 #include "UObject/NameTypes.h"
@@ -73,10 +75,18 @@ public:
 	TSet<FProperty*> OngoingPropertyChanges;
 };
 
-void SAvaRundownPageRemoteControlProps::Construct(const FArguments& InArgs, TSharedPtr<FAvaRundownEditor> InRundownEditor)
+void SAvaRundownPageRemoteControlProps::Construct(const FArguments& InArgs, const TSharedRef<SAvaRundownPageDetails>& InPageDetailPanel, const TSharedPtr<FAvaRundownEditor>& InRundownEditor)
 {
 	RundownEditorWeak = InRundownEditor;
 	ActivePageId = FAvaRundownPage::InvalidPageId;
+	AsyncAssetLoader = MakeShared<UE::AvaPlayback::Utils::FAsyncAssetLoader>();
+	AsyncAssetLoader->OnLoadingCompleted.AddSPLambda(&InPageDetailPanel.Get(), [PageDetailPanelWeak = InPageDetailPanel.ToWeakPtr()]()
+	{
+		if (const TSharedPtr<SAvaRundownPageDetails> PageDetailPanel = PageDetailPanelWeak.Pin())
+		{
+			PageDetailPanel->RefreshSelectedPage();	// Refresh the detail panel when assets are loaded.
+		}
+	});
 
 	CommandList = MakeShared<FUICommandList>();
 	NotifyHook = MakeShared<FAvaRundownPageRCPropsNotifyHookImpl>(SharedThis<SAvaRundownPageRemoteControlProps>(this).ToWeakPtr());
@@ -206,6 +216,7 @@ void SAvaRundownPageRemoteControlProps::Refresh(const TArray<int32>& InSelectedP
 		return;
 	}
 
+	int32 PreviousPageId = ActivePageId;
 	ActivePageId = InSelectedPageIds.IsEmpty() ? FAvaRundownPage::InvalidPageId : InSelectedPageIds[0];
 	
 	if (const TSharedPtr<FAvaRundownEditor> RundownEditor = RundownEditorWeak.Pin())
@@ -237,6 +248,15 @@ void SAvaRundownPageRemoteControlProps::Refresh(const TArray<int32>& InSelectedP
 						UE_LOG(LogAvaPlayableRemoteControl, Log, TEXT("Page %d had stale values that where pruned."), ActivePage->GetPageId());
 						RundownEditor->MarkAsModified();
 					}
+
+					if (PreviousPageId != ActivePageId)
+					{
+						TSet<FSoftObjectPath> Assets;
+						const FAvaPlayableRemoteControlValues& Values = ActivePage->GetRemoteControlValues();
+						FAvaPlayableRemoteControlValues::CollectReferencedAssetPaths(Values.ControllerValues, Assets);
+						FAvaPlayableRemoteControlValues::CollectReferencedAssetPaths(Values.EntityValues, Assets);
+						AsyncAssetLoader->BeginLoadingAssets(Assets.Array());
+					}
 				}
 			}
 		}
diff --git a/Engine/Plugins/Experimental/Avalanche/Source/AvalancheMediaEditor/Private/Rundown/DetailsView/RemoteControl/Properties/SAvaRundownPageRemoteControlProps.h b/Engine/Plugins/Experimental/Avalanche/Source/AvalancheMediaEditor/Private/Rundown/DetailsView/RemoteControl/Properties/SAvaRundownPageRemoteControlProps.h
--- a/Engine/Plugins/Experimental/Avalanche/Source/AvalancheMediaEditor/Private/Rundown/DetailsView/RemoteControl/Properties/SAvaRundownPageRemoteControlProps.h
+++ b/Engine/Plugins/Experimental/Avalanche/Source/AvalancheMediaEditor/Private/Rundown/DetailsView/RemoteControl/Properties/SAvaRundownPageRemoteControlProps.h
@@ -21,6 +21,7 @@ class FAvaRundownRCFieldItem;
 class FName;
 class FUICommandList;
 class ITableRow;
+class SAvaRundownPageDetails;
 class SAvaRundownPageRemoteControlProps;
 class STableViewBase;
 class SWidget;
@@ -32,6 +33,11 @@ struct FGuid;
 struct FRemoteControlEntity;
 template<typename ItemType> class STreeView;
 
+namespace UE::AvaPlayback::Utils
+{
+	class FAsyncAssetLoader;
+}
+
 DECLARE_MULTICAST_DELEGATE_TwoParams(FAvaRundownRCPropertyHeaderRowExtensionDelegate, TSharedRef<SAvaRundownPageRemoteControlProps> Panel,
 	TSharedRef<SHeaderRow>& HeaderRow)
 DECLARE_DELEGATE_ThreeParams(FAvaRundownRCPropertyTableRowExtensionDelegate, TSharedRef<SAvaRundownPageRemoteControlProps> Panel,
@@ -66,7 +72,7 @@ public:
 	SLATE_BEGIN_ARGS(SAvaRundownPageRemoteControlProps) {}
 	SLATE_END_ARGS()
 
-	void Construct(const FArguments& InArgs, TSharedPtr<FAvaRundownEditor> InRundownEditor);
+	void Construct(const FArguments& InArgs, const TSharedRef<SAvaRundownPageDetails>& InPageDetailPanel, const TSharedPtr<FAvaRundownEditor>& InRundownEditor);
 	virtual ~SAvaRundownPageRemoteControlProps() override;
 
 	/** Update the current page's remote control values from the defaults then refresh the widget. */
@@ -173,4 +179,6 @@ private:
 
 	/** Whether the properties that are set by controllers are shown */
 	bool bShowControlledProperties = false;
+
+	TSharedPtr<UE::AvaPlayback::Utils::FAsyncAssetLoader> AsyncAssetLoader;
 };
\ No newline at end of file
diff --git a/Engine/Plugins/Experimental/Avalanche/Source/AvalancheMediaEditor/Private/Rundown/DetailsView/SAvaRundownPageDetails.cpp b/Engine/Plugins/Experimental/Avalanche/Source/AvalancheMediaEditor/Private/Rundown/DetailsView/SAvaRundownPageDetails.cpp
--- a/Engine/Plugins/Experimental/Avalanche/Source/AvalancheMediaEditor/Private/Rundown/DetailsView/SAvaRundownPageDetails.cpp
+++ b/Engine/Plugins/Experimental/Avalanche/Source/AvalancheMediaEditor/Private/Rundown/DetailsView/SAvaRundownPageDetails.cpp
@@ -163,7 +163,7 @@ void SAvaRundownPageDetails::Construct(const FArguments& InArgs, const TSharedPt
 			+ SVerticalBox::Slot()
 			.AutoHeight()
 			[
-				SAssignNew(RemoteControlProps, SAvaRundownPageRemoteControlProps, InRundownEditor)
+				SAssignNew(RemoteControlProps, SAvaRundownPageRemoteControlProps, SharedThis(this), InRundownEditor)
 			]
 		]
 	];
diff --git a/Engine/Plugins/Experimental/Avalanche/Source/AvalancheMediaEditor/Private/Rundown/DetailsView/SAvaRundownPageDetails.h b/Engine/Plugins/Experimental/Avalanche/Source/AvalancheMediaEditor/Private/Rundown/DetailsView/SAvaRundownPageDetails.h
--- a/Engine/Plugins/Experimental/Avalanche/Source/AvalancheMediaEditor/Private/Rundown/DetailsView/SAvaRundownPageDetails.h
+++ b/Engine/Plugins/Experimental/Avalanche/Source/AvalancheMediaEditor/Private/Rundown/DetailsView/SAvaRundownPageDetails.h
@@ -34,6 +34,11 @@ public:
 	void OnPageEvent(const TArray<int32>& InSelectedPageIds, UE::AvaRundown::EPageEvent InPageEvent);
 	void OnManagedInstanceCacheEntryInvalidated(const FSoftObjectPath& InAssetPath);
 
+	void RefreshSelectedPage()
+	{
+		QueueRefreshSelectedPage();
+	}
+
 protected:
 	FReply ToggleExposedPropertiesVisibility();
 
