MotionDesign : Shapes
- Prevent shapes dynamic mesh updates when package is saving to avoid crash

diff --git a/Engine/Plugins/Experimental/Avalanche/Source/AvalancheShapes/Private/DynamicMeshes/AvaShapeDynMeshBase.cpp b/Engine/Plugins/Experimental/Avalanche/Source/AvalancheShapes/Private/DynamicMeshes/AvaShapeDynMeshBase.cpp
--- a/Engine/Plugins/Experimental/Avalanche/Source/AvalancheShapes/Private/DynamicMeshes/AvaShapeDynMeshBase.cpp
+++ b/Engine/Plugins/Experimental/Avalanche/Source/AvalancheShapes/Private/DynamicMeshes/AvaShapeDynMeshBase.cpp
@@ -25,6 +25,7 @@
 #include "Parameterization/DynamicMeshUVEditor.h"
 #include "PhysicsEngine/AggregateGeom.h"
 #include "Subsystems/ActorModifierCoreSubsystem.h"
+#include "UObject/Package.h"
 
 #if WITH_EDITOR
 #include "DMObjectMaterialProperty.h"
@@ -40,6 +41,27 @@ DEFINE_LOG_CATEGORY_STATIC(LogAvaDynamicMesh, Log, All);
 
 using namespace UE::Geometry;
 
+namespace UE::Private::AvaShape
+{
+	/** Edit underlying dynamic mesh */
+	bool EditMesh(const UAvaShapeDynamicMeshBase* InContext, const TFunctionRef<void(FDynamicMesh3&)>& InFunctor, EDynamicMeshChangeType InChangeType = EDynamicMeshChangeType::GeneralEdit, EDynamicMeshAttributeChangeFlags InAttributeChangeType = EDynamicMeshAttributeChangeFlags::Unknown, bool bInBroadcastChangeEvent = true)
+	{
+		checkf(!!InContext, TEXT("Invalid shape context provided to edit mesh"))
+
+		if (UDynamicMeshComponent* MeshComponent = InContext->GetShapeMeshComponent())
+		{
+			if (UDynamicMesh* Mesh = MeshComponent->GetDynamicMesh())
+			{
+				const bool bIsSaving = UE::IsSavingPackage(InContext->GetPackage());
+				Mesh->EditMesh(InFunctor, InChangeType, InAttributeChangeType, !bIsSaving ? bInBroadcastChangeEvent : false);
+				return true;
+			}
+		}
+
+		return false;
+	}
+}
+
 UAvaShapeDynamicMeshBase::UAvaShapeDynamicMeshBase(const FLinearColor& InVertexColor, float InUniformScaledSize, bool bInAllowEditSize)
 	: bAllowEditSize(bInAllowEditSize)
 	, UniformScaledSize(InUniformScaledSize)
@@ -279,58 +301,55 @@ void UAvaShapeDynamicMeshBase::ToggleGizmo_Implementation(
 
 void UAvaShapeDynamicMeshBase::UpdateDynamicMesh()
 {
-	if (IsValid(this) && !HasAnyFlags(RF_BeginDestroyed | RF_FinishDestroyed | RF_ClassDefaultObject))
-	{
-		RestoreCachedMesh();
-		CachedMesh.Reset();
+	RestoreCachedMesh();
+	CachedMesh.Reset();
 
-		bool bUpdateMesh = false;
+	bool bUpdateMesh = false;
 
-		if (bAllMeshDirty || bAnyMeshDirty)
-		{
-			bUpdateMesh = GenerateMesh();
-		}
+	if (bAllMeshDirty || bAnyMeshDirty)
+	{
+		bUpdateMesh = GenerateMesh();
+	}
 
-		if (bVerticesDirty)
-		{
-			UpdateVertices();
-		}
+	if (bVerticesDirty)
+	{
+		UpdateVertices();
+	}
 
-		for (const int32 MeshIndex : GetMeshesIndexes())
+	for (const int32 MeshIndex : GetMeshesIndexes())
+	{
+		FAvaShapeMesh* Mesh = GetMesh(MeshIndex);
+		if (Mesh->bUpdateRequired)
 		{
-			FAvaShapeMesh* Mesh = GetMesh(MeshIndex);
-			if (Mesh->bUpdateRequired)
-			{
-				OnMeshSectionUpdated(*Mesh);
-			}
+			OnMeshSectionUpdated(*Mesh);
 		}
+	}
 
-		if (bHasNewMeshRegenWorldLocation)
-		{
-			AActor* OwningActor = GetOwner();
-			bHasNewMeshRegenWorldLocation = false;
-			FVector Delta = MeshRegenWorldLocation - OwningActor->GetRootComponent()->GetComponentLocation();
+	if (bHasNewMeshRegenWorldLocation)
+	{
+		AActor* OwningActor = GetOwner();
+		bHasNewMeshRegenWorldLocation = false;
+		FVector Delta = MeshRegenWorldLocation - OwningActor->GetRootComponent()->GetComponentLocation();
 
 #if WITH_EDITOR
-			FActorElementEditorViewportInteractionCustomization::ApplyDeltaToActor(
-				OwningActor, true, &Delta,
-				nullptr, nullptr, FVector::ZeroVector, FInputDeviceState()
-			); // Pivot is not used for just translation
+		FActorElementEditorViewportInteractionCustomization::ApplyDeltaToActor(
+			OwningActor, true, &Delta,
+			nullptr, nullptr, FVector::ZeroVector, FInputDeviceState()
+		); // Pivot is not used for just translation
 #else
-			OwningActor->ApplyWorldOffset(Delta, false);
+		OwningActor->ApplyWorldOffset(Delta, false);
 #endif
-			MeshRegenWorldLocation = FVector::ZeroVector;
-		}
-
-		GenerateUV();
+		MeshRegenWorldLocation = FVector::ZeroVector;
+	}
 
-		if (bUpdateMesh)
-		{
-			GenerateTangents();
-		}
+	GenerateUV();
 
-		SaveCachedMesh();
+	if (bUpdateMesh)
+	{
+		GenerateTangents();
 	}
+
+	SaveCachedMesh();
 }
 
 void UAvaShapeDynamicMeshBase::RunUpdate(bool bDoAsync)
@@ -347,10 +366,29 @@ void UAvaShapeDynamicMeshBase::RunUpdate(bool bDoAsync)
 
 	MeshUpdateStatus = EAvaDynamicMeshUpdateState::UpdateInProgress;
 
+	auto ExecUpdate = [this]()
+	{
+		if (IsValid(this) && !HasAnyFlags(RF_BeginDestroyed | RF_FinishDestroyed | RF_ClassDefaultObject))
+		{
+			if (!UE::IsSavingPackage(this))
+			{
+				FActorModifierCoreScopedLock ModifierLock(UActorModifierCoreSubsystem::Get()->GetActorModifierStack(GetShapeActor()));
+
+				UpdateDynamicMesh();
+				MeshUpdateStatus = EAvaDynamicMeshUpdateState::UpToDate;
+				OnMeshUpdateFinished();
+			}
+			else
+			{
+				MeshUpdateStatus = EAvaDynamicMeshUpdateState::UpdateRequired;
+			}
+		}
+	};
+
 	if (bDoAsync)
 	{
 		TWeakObjectPtr<UAvaShapeDynamicMeshBase> ThisWeak(this);
-		Async(EAsyncExecution::TaskGraphMainThread, [ThisWeak]()
+		Async(EAsyncExecution::TaskGraphMainThread, [ThisWeak, ExecUpdate]()
 		{
 			UAvaShapeDynamicMeshBase* const This = ThisWeak.Get(true);
 			if (!This)
@@ -358,20 +396,12 @@ void UAvaShapeDynamicMeshBase::RunUpdate(bool bDoAsync)
 				return;
 			}
 
-			FActorModifierCoreScopedLock ModifierLock(UActorModifierCoreSubsystem::Get()->GetActorModifierStack(This->GetShapeActor()));
-
-			This->UpdateDynamicMesh();
-			This->MeshUpdateStatus = EAvaDynamicMeshUpdateState::UpToDate;
-			This->OnMeshUpdateFinished();
+			ExecUpdate();
 		});
 	}
 	else
 	{
-		FActorModifierCoreScopedLock ModifierLock(UActorModifierCoreSubsystem::Get()->GetActorModifierStack(GetShapeActor()));
-
-		UpdateDynamicMesh();
-		MeshUpdateStatus = EAvaDynamicMeshUpdateState::UpToDate;
-		OnMeshUpdateFinished();
+		ExecUpdate();
 	}
 }
 
@@ -1517,62 +1547,65 @@ bool UAvaShapeDynamicMeshBase::CreateDynamicMesh(FAvaShapeMesh& InMesh)
 
 	InMesh.bUpdateRequired = false;
 
-	ProcMesh->GetDynamicMesh()->EditMesh([this, &InMesh](FDynamicMesh3& Mesh)
-	{
-		// clear ids before adding them
-		InMesh.ClearIds();
-
-		FDynamicMeshNormalOverlay* const NormalOverlay = Mesh.Attributes()->PrimaryNormals();
-		FDynamicMeshMaterialAttribute* const MaterialAttr = Mesh.Attributes()->GetMaterialID();
-		FDynamicMeshPolygroupAttribute* const PolyGroupAttr = Mesh.Attributes()->GetPolygroupLayer(InMesh.GetMeshIndex());
-
-		// process vertices infos
-		for (int32 v = 0; v < InMesh.Vertices.Num(); v++)
+	UE::Private::AvaShape::EditMesh(
+		this,
+		[this, &InMesh](FDynamicMesh3& Mesh)
 		{
-			int32 VId = Mesh.AppendVertex(InMesh.Vertices[v]);
-			InMesh.VerticeIds.Add(VId);
+			// clear ids before adding them
+			InMesh.ClearIds();
 
-			FVector3f Normal = static_cast<FVector3f>(InMesh.Normals[v]);
-			int32 NId = NormalOverlay->AppendElement(Normal);
-			InMesh.NormalIds.Add(NId);
-		}
+			FDynamicMeshNormalOverlay* const NormalOverlay = Mesh.Attributes()->PrimaryNormals();
+			FDynamicMeshMaterialAttribute* const MaterialAttr = Mesh.Attributes()->GetMaterialID();
+			FDynamicMeshPolygroupAttribute* const PolyGroupAttr = Mesh.Attributes()->GetPolygroupLayer(InMesh.GetMeshIndex());
 
-		// process triangles
-		for (int32 t = 0; t < InMesh.Triangles.Num(); t += 3)
-		{
-			int32 Idx1 = InMesh.Triangles[t];
-			int32 Idx2 = InMesh.Triangles[t + 1];
-			int32 Idx3 = InMesh.Triangles[t + 2];
-
-			if (!InMesh.VerticeIds.IsValidIndex(Idx1) ||
-				!InMesh.VerticeIds.IsValidIndex(Idx2) ||
-				!InMesh.VerticeIds.IsValidIndex(Idx3))
+			// process vertices infos
+			for (int32 v = 0; v < InMesh.Vertices.Num(); v++)
 			{
-				UE_LOG(LogAvaDynamicMesh, Warning, TEXT("CreateDynamicMesh %s %i : Invalid Vertice idx for triangle %i %i %i"), *GetMeshName(), InMesh.GetMeshIndex(), Idx1, Idx2, Idx3);
-				continue;
+				int32 VId = Mesh.AppendVertex(InMesh.Vertices[v]);
+				InMesh.VerticeIds.Add(VId);
+
+				FVector3f Normal = static_cast<FVector3f>(InMesh.Normals[v]);
+				int32 NId = NormalOverlay->AppendElement(Normal);
+				InMesh.NormalIds.Add(NId);
 			}
 
-			// get vertice id from DM
-			int32 VId1 = InMesh.VerticeIds[Idx1];
-			int32 VId2 = InMesh.VerticeIds[Idx2];
-			int32 VId3 = InMesh.VerticeIds[Idx3];
-			// create triangle from vertices id
-			int32 TId = Mesh.AppendTriangle(VId1, VId2, VId3, InMesh.GetMeshIndex());
-			if (TId < 0)
+			// process triangles
+			for (int32 t = 0; t < InMesh.Triangles.Num(); t += 3)
 			{
-				UE_LOG(LogAvaDynamicMesh, Warning, TEXT("CreateDynamicMesh %s %i : Invalid Triangle ID for mesh"), *GetMeshName(), InMesh.GetMeshIndex());
-				continue;
+				int32 Idx1 = InMesh.Triangles[t];
+				int32 Idx2 = InMesh.Triangles[t + 1];
+				int32 Idx3 = InMesh.Triangles[t + 2];
+
+				if (!InMesh.VerticeIds.IsValidIndex(Idx1) ||
+					!InMesh.VerticeIds.IsValidIndex(Idx2) ||
+					!InMesh.VerticeIds.IsValidIndex(Idx3))
+				{
+					UE_LOG(LogAvaDynamicMesh, Warning, TEXT("CreateDynamicMesh %s %i : Invalid Vertice idx for triangle %i %i %i"), *GetMeshName(), InMesh.GetMeshIndex(), Idx1, Idx2, Idx3);
+					continue;
+				}
+
+				// get vertice id from DM
+				int32 VId1 = InMesh.VerticeIds[Idx1];
+				int32 VId2 = InMesh.VerticeIds[Idx2];
+				int32 VId3 = InMesh.VerticeIds[Idx3];
+				// create triangle from vertices id
+				int32 TId = Mesh.AppendTriangle(VId1, VId2, VId3, InMesh.GetMeshIndex());
+				if (TId < 0)
+				{
+					UE_LOG(LogAvaDynamicMesh, Warning, TEXT("CreateDynamicMesh %s %i : Invalid Triangle ID for mesh"), *GetMeshName(), InMesh.GetMeshIndex());
+					continue;
+				}
+				InMesh.TriangleIds.Add(TId);
+				// material
+				MaterialAttr->SetValue(TId, InMesh.GetMeshIndex());
+				// poly group
+				PolyGroupAttr->SetValue(TId, InMesh.GetMeshIndex());
+				// normals
+				const FIndex3i NormalIds(InMesh.NormalIds[Idx1], InMesh.NormalIds[Idx2], InMesh.NormalIds[Idx3]);
+				NormalOverlay->SetTriangle(TId, NormalIds, true);
 			}
-			InMesh.TriangleIds.Add(TId);
-			// material
-			MaterialAttr->SetValue(TId, InMesh.GetMeshIndex());
-			// poly group
-			PolyGroupAttr->SetValue(TId, InMesh.GetMeshIndex());
-			// normals
-			const FIndex3i NormalIds(InMesh.NormalIds[Idx1], InMesh.NormalIds[Idx2], InMesh.NormalIds[Idx3]);
-			NormalOverlay->SetTriangle(TId, NormalIds, true);
 		}
-	}, EDynamicMeshChangeType::GeneralEdit, EDynamicMeshAttributeChangeFlags::Unknown);
+	);
 
 	ProcMesh->SetMaterial(InMesh.GetMeshIndex(), GetMaterial(InMesh.GetMeshIndex()));
 
@@ -1610,20 +1643,25 @@ bool UAvaShapeDynamicMeshBase::UpdateDynamicMesh(FAvaShapeMesh& InMesh)
 	InMesh.bUpdateRequired = false;
 
 	// same topology, update vertices only (location, normal, color, uv)
-	ProcMesh->GetDynamicMesh()->EditMesh([this, InMesh](FDynamicMesh3& EditMesh)
-	{
-		FDynamicMeshNormalOverlay* const NormalOverlay = EditMesh.Attributes()->PrimaryNormals();
-
-		ParallelFor(InMesh.Vertices.Num(), [InMesh, &EditMesh, NormalOverlay](int32 Idx)
+	UE::Private::AvaShape::EditMesh(
+		this,
+		[this, InMesh](FDynamicMesh3& EditMesh)
 		{
-			// vertices
-			int32 VId = InMesh.VerticeIds[Idx];
-			EditMesh.SetVertex(VId, InMesh.Vertices[Idx]);
-			// normals
-			int32 NId = InMesh.NormalIds[Idx];
-			NormalOverlay->SetElement(NId, static_cast<FVector3f>(InMesh.Normals[Idx]));
-		});
-	}, EDynamicMeshChangeType::MeshVertexChange, EDynamicMeshAttributeChangeFlags::Unknown);
+			FDynamicMeshNormalOverlay* const NormalOverlay = EditMesh.Attributes()->PrimaryNormals();
+
+			ParallelFor(InMesh.Vertices.Num(), [InMesh, &EditMesh, NormalOverlay](int32 Idx)
+			{
+				// vertices
+				int32 VId = InMesh.VerticeIds[Idx];
+				EditMesh.SetVertex(VId, InMesh.Vertices[Idx]);
+				// normals
+				int32 NId = InMesh.NormalIds[Idx];
+				NormalOverlay->SetElement(NId, static_cast<FVector3f>(InMesh.Normals[Idx]));
+			});
+		},
+		EDynamicMeshChangeType::MeshVertexChange,
+		EDynamicMeshAttributeChangeFlags::Unknown
+	);
 
 	ProcMesh->SetMaterial(InMesh.GetMeshIndex(), GetMaterial(InMesh.GetMeshIndex()));
 
@@ -1909,13 +1947,6 @@ void UAvaShapeDynamicMeshBase::SetMeshRegenWorldLocation(const FVector& NewLocat
 
 bool UAvaShapeDynamicMeshBase::ApplyUVsManually(FAvaShapeMesh& InMesh)
 {
-	UDynamicMeshComponent* DynMeshComponent = GetShapeMeshComponent();
-
-	if (!DynMeshComponent)
-	{
-		return false;
-	}
-
 	if (InMesh.TriangleIds.IsEmpty())
 	{
 		return false;
@@ -1928,130 +1959,131 @@ bool UAvaShapeDynamicMeshBase::ApplyUVsManually(FAvaShapeMesh& InMesh)
 
 	InMesh.UVIds.Empty();
 
-	DynMeshComponent->GetDynamicMesh()->EditMesh([this, &InMesh](FDynamicMesh3& EditMesh)
-	{
-		FDynamicMeshUVOverlay* const UVOverlay = EditMesh.Attributes()->GetUVLayer(0);
-
-		// used to quickly find index O(1) instead of O(n)
-		TMap<int32, int32> BaseToOverlayVIDMap;
-		for (int32 Idx = 0; Idx < InMesh.VerticeIds.Num(); Idx++)
+	return UE::Private::AvaShape::EditMesh(
+		this,
+		[this, &InMesh](FDynamicMesh3& EditMesh)
 		{
-			FVector2f UV = static_cast<FVector2f>(InMesh.UVs[Idx]);
-			int32 UId = UVOverlay->AppendElement(UV);
-			InMesh.UVIds.Add(UId);
-			BaseToOverlayVIDMap.Add(InMesh.VerticeIds[Idx], UId);
-		}
+			FDynamicMeshUVOverlay* const UVOverlay = EditMesh.Attributes()->GetUVLayer(0);
 
-		for (int32 TID : InMesh.TriangleIds)
-		{
-			FIndex3i TriVtx = EditMesh.GetTriangle(TID);
-			FIndex3i UVElem;
-			for (int32 j = 0; j < 3; ++j)
+			// used to quickly find index O(1) instead of O(n)
+			TMap<int32, int32> BaseToOverlayVIDMap;
+			for (int32 Idx = 0; Idx < InMesh.VerticeIds.Num(); Idx++)
 			{
-				const int32* FoundElementID = BaseToOverlayVIDMap.Find(TriVtx[j]);
-				if (FoundElementID == nullptr)
-				{
-					UE_LOG(LogAvaDynamicMesh, Warning, TEXT("ApplyUVsManually %s %i: vertice id %i for uv is invalid, skipping triangle"), *GetMeshName(), InMesh.GetMeshIndex(), TriVtx[j]);
-					break;
-				}
-				else
+				FVector2f UV = static_cast<FVector2f>(InMesh.UVs[Idx]);
+				int32 UId = UVOverlay->AppendElement(UV);
+				InMesh.UVIds.Add(UId);
+				BaseToOverlayVIDMap.Add(InMesh.VerticeIds[Idx], UId);
+			}
+
+			for (int32 TID : InMesh.TriangleIds)
+			{
+				FIndex3i TriVtx = EditMesh.GetTriangle(TID);
+				FIndex3i UVElem;
+				for (int32 j = 0; j < 3; ++j)
 				{
-					UVElem[j] = *FoundElementID;
+					const int32* FoundElementID = BaseToOverlayVIDMap.Find(TriVtx[j]);
+					if (FoundElementID == nullptr)
+					{
+						UE_LOG(LogAvaDynamicMesh, Warning, TEXT("ApplyUVsManually %s %i: vertice id %i for uv is invalid, skipping triangle"), *GetMeshName(), InMesh.GetMeshIndex(), TriVtx[j]);
+						break;
+					}
+					else
+					{
+						UVElem[j] = *FoundElementID;
+					}
 				}
+				UVOverlay->SetTriangle(TID, UVElem);
 			}
-			UVOverlay->SetTriangle(TID, UVElem);
-		}
-
-	}, EDynamicMeshChangeType::AttributeEdit, EDynamicMeshAttributeChangeFlags::UVs);
-
-	return true;
+		},
+		EDynamicMeshChangeType::AttributeEdit,
+		EDynamicMeshAttributeChangeFlags::UVs
+	);
 }
 
 bool UAvaShapeDynamicMeshBase::ApplyUVsPlanarProjection(FAvaShapeMesh& InMesh, FRotator PlaneRotation, FVector2D PlaneSize)
 {
-	UDynamicMeshComponent* DynMeshComponent = GetShapeMeshComponent();
-
-	if (!DynMeshComponent || InMesh.TriangleIds.IsEmpty())
+	if (InMesh.TriangleIds.IsEmpty())
 	{
 		return false;
 	}
 
-	DynMeshComponent->GetDynamicMesh()->EditMesh([this, &InMesh, PlaneRotation, PlaneSize](FDynamicMesh3& EditMesh) {
-
-		const FTransform PlaneTransform(PlaneRotation, FVector(0), FVector(PlaneSize.X, PlaneSize.Y, 0));
-		FDynamicMeshUVOverlay* UVOverlay = EditMesh.Attributes()->GetUVLayer(0);
-		FDynamicMeshUVEditor UVEditor(&EditMesh, UVOverlay);
-		const FFrame3d ProjectionFrame(PlaneTransform);
-		const FVector Scale = PlaneTransform.GetScale3D();
-		const FVector2d Dimensions(Scale.X, Scale.Y);
-		FUVEditResult Result;
-
-		UVEditor.SetTriangleUVsFromPlanarProjection(InMesh.TriangleIds, [](const FVector3d& Pos) {
-			return Pos;
-		}, ProjectionFrame, Dimensions, &Result);
+	return UE::Private::AvaShape::EditMesh(
+		this,
+		[this, &InMesh, PlaneRotation, PlaneSize](FDynamicMesh3& EditMesh)
+		{
+			const FTransform PlaneTransform(PlaneRotation, FVector(0), FVector(PlaneSize.X, PlaneSize.Y, 0));
+			FDynamicMeshUVOverlay* UVOverlay = EditMesh.Attributes()->GetUVLayer(0);
+			FDynamicMeshUVEditor UVEditor(&EditMesh, UVOverlay);
+			const FFrame3d ProjectionFrame(PlaneTransform);
+			const FVector Scale = PlaneTransform.GetScale3D();
+			const FVector2d Dimensions(Scale.X, Scale.Y);
+			FUVEditResult Result;
 
-		InMesh.UVIds = Result.NewUVElements;
+			UVEditor.SetTriangleUVsFromPlanarProjection(InMesh.TriangleIds, [](const FVector3d& Pos) {
+				return Pos;
+			}, ProjectionFrame, Dimensions, &Result);
 
-		// copy generated uv to cache
-		InMesh.UVs.Empty();
-		for(const int32 UId : InMesh.UVIds)
-		{
-			InMesh.UVs.Add(static_cast<FVector2D>(UVOverlay->GetElement(UId)));
-		}
-	}, EDynamicMeshChangeType::AttributeEdit, EDynamicMeshAttributeChangeFlags::UVs);
+			InMesh.UVIds = Result.NewUVElements;
 
-	return true;
+			// copy generated uv to cache
+			InMesh.UVs.Empty();
+			for(const int32 UId : InMesh.UVIds)
+			{
+				InMesh.UVs.Add(static_cast<FVector2D>(UVOverlay->GetElement(UId)));
+			}
+		}, 
+		EDynamicMeshChangeType::AttributeEdit,
+		EDynamicMeshAttributeChangeFlags::UVs
+	);
 }
 
 bool UAvaShapeDynamicMeshBase::ApplyUVsBoxProjection(FAvaShapeMesh& InMesh, FRotator BoxRotation, FVector BoxSize)
 {
-	UDynamicMeshComponent* DynMeshComponent = GetShapeMeshComponent();
-
-	if (!DynMeshComponent || InMesh.TriangleIds.IsEmpty())
+	if (InMesh.TriangleIds.IsEmpty())
 	{
 		return false;
 	}
 
-	DynMeshComponent->GetDynamicMesh()->EditMesh([this, &InMesh, BoxRotation, BoxSize](FDynamicMesh3& EditMesh) {
-
-		const FTransform PlaneTransform(BoxRotation, FVector(0), BoxSize);
-		FDynamicMeshUVOverlay* UVOverlay = EditMesh.Attributes()->GetUVLayer(0);
-		FDynamicMeshUVEditor UVEditor(&EditMesh, UVOverlay);
-		const FFrame3d ProjectionFrame(PlaneTransform);
-		FUVEditResult Result;
+	return UE::Private::AvaShape::EditMesh(
+		this,
+		[this, &InMesh, BoxRotation, BoxSize](FDynamicMesh3& EditMesh)
+		{
+			const FTransform PlaneTransform(BoxRotation, FVector(0), BoxSize);
+			FDynamicMeshUVOverlay* UVOverlay = EditMesh.Attributes()->GetUVLayer(0);
+			FDynamicMeshUVEditor UVEditor(&EditMesh, UVOverlay);
+			const FFrame3d ProjectionFrame(PlaneTransform);
+			FUVEditResult Result;
 
-		UVEditor.SetTriangleUVsFromBoxProjection(InMesh.TriangleIds, [this](const FVector3d& Pos) {
-			return Pos;
-		}, ProjectionFrame, BoxSize, 2, &Result);
+			UVEditor.SetTriangleUVsFromBoxProjection(InMesh.TriangleIds, [this](const FVector3d& Pos) {
+				return Pos;
+			}, ProjectionFrame, BoxSize, 2, &Result);
 
-		InMesh.UVIds = Result.NewUVElements;
+			InMesh.UVIds = Result.NewUVElements;
 
-		// copy generated uv to cache
-		InMesh.UVs.Empty();
-		for (const int32 UId : InMesh.UVIds)
-		{
-			InMesh.UVs.Add(static_cast<FVector2D>(UVOverlay->GetElement(UId)));
-		}
-	}, EDynamicMeshChangeType::AttributeEdit, EDynamicMeshAttributeChangeFlags::UVs);
-
-	return true;
+			// copy generated uv to cache
+			InMesh.UVs.Empty();
+			for (const int32 UId : InMesh.UVIds)
+			{
+				InMesh.UVs.Add(static_cast<FVector2D>(UVOverlay->GetElement(UId)));
+			}
+		},
+		EDynamicMeshChangeType::AttributeEdit,
+		EDynamicMeshAttributeChangeFlags::UVs
+	);
 }
 
 bool UAvaShapeDynamicMeshBase::ApplyUVsTransform(FAvaShapeMesh& InMesh, FAvaShapeMaterialUVParameters& InParams, FVector2D ShapeSize, FVector2D UVOffset, float UVFixRotation) const
 {
-	UDynamicMeshComponent* DynMesh = GetShapeMeshComponent();
-
-	if (!DynMesh)
-	{
-		return false;
-	}
-
 	bool bResult = true;
-	DynMesh->GetDynamicMesh()->EditMesh([InMesh, InParams, ShapeSize, UVOffset, UVFixRotation, &bResult](FDynamicMesh3& EditMesh) {
-		bResult = UE::AvaShapes::TransformMeshUVs(EditMesh, InMesh.UVIds, InParams, ShapeSize, UVOffset, UVFixRotation);
-	}, EDynamicMeshChangeType::AttributeEdit, EDynamicMeshAttributeChangeFlags::UVs);
-
-	return bResult;
+	return UE::Private::AvaShape::EditMesh(
+		this,
+		[InMesh, InParams, ShapeSize, UVOffset, UVFixRotation, &bResult](FDynamicMesh3& EditMesh)
+		{
+			bResult = UE::AvaShapes::TransformMeshUVs(EditMesh, InMesh.UVIds, InParams, ShapeSize, UVOffset, UVFixRotation);
+		},
+		EDynamicMeshChangeType::AttributeEdit,
+		EDynamicMeshAttributeChangeFlags::UVs
+	) && bResult;
 }
 
 void UAvaShapeDynamicMeshBase::OnParametricMaterialChanged(const FAvaShapeParametricMaterial& InMaterial)
@@ -2080,13 +2112,15 @@ void UAvaShapeDynamicMeshBase::OnParametricMaterialParameterChanged(const FAvaSh
 
 void UAvaShapeDynamicMeshBase::RestoreCachedMesh()
 {
-	UDynamicMeshComponent* DynamicMeshComponent = GetShapeMeshComponent();
-	if (DynamicMeshComponent && CachedMesh.IsSet())
+	if (CachedMesh.IsSet())
 	{
-		DynamicMeshComponent->EditMesh([this](FDynamicMesh3& EditMesh)
-		{
-			EditMesh = CachedMesh.GetValue();
-		});
+		UE::Private::AvaShape::EditMesh(
+			this,
+			[this](FDynamicMesh3& EditMesh)
+			{
+				EditMesh = CachedMesh.GetValue();
+			}
+		);
 	}
 }
 
@@ -2255,27 +2289,34 @@ void UAvaShapeDynamicMeshBase::InitializeDynamicMesh()
 		return;
 	}
 
-	MeshComponent->EditMesh([this](FDynamicMesh3& EditMesh) {
-		EditMesh.Clear();
-		// for sections
-		EditMesh.EnableTriangleGroups();
-		// for materials, tangents, uv, colors
-		EditMesh.EnableAttributes();
-		EditMesh.Attributes()->EnableMaterialID();
-		EditMesh.Attributes()->GetMaterialID()->SetName(FName(GetMeshName()));
-		EditMesh.Attributes()->EnableTangents();
-		// set name for PolyGroup layers
-		EditMesh.Attributes()->SetNumPolygroupLayers(MeshDatas.Num());
-		for (int32 Idx = 0; Idx < MeshDatas.Num(); Idx++)
-		{
-			const FName DataName = GetMeshDataName(Idx);
-			EditMesh.Attributes()->GetPolygroupLayer(Idx)->SetName(DataName);
-		}
-		EditMesh.Attributes()->SetNumUVLayers(1);
-		EditMesh.Attributes()->EnablePrimaryColors();
+	UE::Private::AvaShape::EditMesh(
+			this,
+			[this](FDynamicMesh3& EditMesh)
+		{
+			EditMesh.Clear();
+
+			// for sections
+			EditMesh.EnableTriangleGroups();
 
-		CachedMesh = EditMesh;
-	});
+			// for materials, tangents, uv, colors
+			EditMesh.EnableAttributes();
+			EditMesh.Attributes()->EnableMaterialID();
+			EditMesh.Attributes()->GetMaterialID()->SetName(FName(GetMeshName()));
+			EditMesh.Attributes()->EnableTangents();
+
+			// set name for PolyGroup layers
+			EditMesh.Attributes()->SetNumPolygroupLayers(MeshDatas.Num());
+			for (int32 Idx = 0; Idx < MeshDatas.Num(); Idx++)
+			{
+				const FName DataName = GetMeshDataName(Idx);
+				EditMesh.Attributes()->GetPolygroupLayer(Idx)->SetName(DataName);
+			}
+			EditMesh.Attributes()->SetNumUVLayers(1);
+			EditMesh.Attributes()->EnablePrimaryColors();
+
+			CachedMesh = EditMesh;
+		}
+	);
 
 	bZeroScaled = MeshComponent->GetComponentScale().IsNearlyZero();
 	MeshComponent->TransformUpdated.AddUObject(this, &UAvaShapeDynamicMeshBase::OnMeshTransformChanged);
@@ -2358,49 +2399,52 @@ FTransform UAvaShapeDynamicMeshBase::GetTransform() const
 
 bool UAvaShapeDynamicMeshBase::ClearDynamicMeshSection(int32 MeshIndex)
 {
-	UDynamicMeshComponent* ShapeComponent = GetShapeMeshComponent();
 	FAvaShapeMeshData* MeshData = GetMeshData(MeshIndex);
 
-	if (!ShapeComponent || !MeshData)
+	if (!MeshData)
 	{
 		return false;
 	}
 
-	ShapeComponent->GetDynamicMesh()->EditMesh([this, &MeshData](FDynamicMesh3& EditMesh)
-	{
-		for (const int32 TId : MeshData->Mesh.TriangleIds)
+	return UE::Private::AvaShape::EditMesh(
+			this,
+			[this, &MeshData](FDynamicMesh3& EditMesh)
 		{
-			if(EditMesh.IsTriangle(TId))
+			for (const int32 TId : MeshData->Mesh.TriangleIds)
 			{
-				EditMesh.RemoveTriangle(TId);
+				if(EditMesh.IsTriangle(TId))
+				{
+					EditMesh.RemoveTriangle(TId);
+				}
 			}
-		}
 
-		// empty ids arrays
-		MeshData->Mesh.ClearIds();
-
-	}, EDynamicMeshChangeType::GeneralEdit, EDynamicMeshAttributeChangeFlags::Unknown);
-
-	return true;
+			// empty ids arrays
+			MeshData->Mesh.ClearIds();
+		}
+	);
 }
 
 bool UAvaShapeDynamicMeshBase::ClearDynamicMesh()
 {
-	UDynamicMeshComponent* ShapeComponent = GetShapeMeshComponent();
-	if (!ShapeComponent)
-	{
-		return false;
-	}
-	ShapeComponent->GetDynamicMesh()->EditMesh([this](FDynamicMesh3& EditMesh)
-	{
-		for (const int32 TId : EditMesh.TriangleIndicesItr())
+	const bool bResult = UE::Private::AvaShape::EditMesh(
+			this,
+			[this](FDynamicMesh3& EditMesh)
 		{
-			if(EditMesh.IsTriangle(TId))
+			for (const int32 TId : EditMesh.TriangleIndicesItr())
 			{
-				EditMesh.RemoveTriangle(TId);
+				if(EditMesh.IsTriangle(TId))
+				{
+					EditMesh.RemoveTriangle(TId);
+				}
 			}
 		}
-	}, EDynamicMeshChangeType::GeneralEdit, EDynamicMeshAttributeChangeFlags::Unknown);
+	);
+
+	if (!bResult)
+	{
+		return false;
+	}
+
 	for(const int32 MeshIdx : GetMeshesIndexes())
 	{
 		if(FAvaShapeMeshData* MeshData = GetMeshData(MeshIdx))
@@ -2408,6 +2452,7 @@ bool UAvaShapeDynamicMeshBase::ClearDynamicMesh()
 			MeshData->Mesh.ClearIds();
 		}
 	}
+
 	return true;
 }
 
