MotionDesign : GeometryMask
- Fix mask canvas not resolving properly when no primitive components are present under Text3D component (PIE)
- Add better tooltip and display name for mask read and write modifier to avoid confusion

diff --git a/Engine/Plugins/Experimental/ActorModifierCore/Source/ActorModifierCoreEditor/Private/Modifiers/Customizations/ActorModifierCoreEditorStackCustomization.cpp b/Engine/Plugins/Experimental/ActorModifierCore/Source/ActorModifierCoreEditor/Private/Modifiers/Customizations/ActorModifierCoreEditorStackCustomization.cpp
--- a/Engine/Plugins/Experimental/ActorModifierCore/Source/ActorModifierCoreEditor/Private/Modifiers/Customizations/ActorModifierCoreEditorStackCustomization.cpp
+++ b/Engine/Plugins/Experimental/ActorModifierCore/Source/ActorModifierCoreEditor/Private/Modifiers/Customizations/ActorModifierCoreEditorStackCustomization.cpp
@@ -252,6 +252,15 @@ void UActorModifierCoreEditorStackCustomization::CustomizeItemHeader(const FOper
 		TAttribute<FText> MessageText = FText::GetEmpty();
 		FString HeaderLabel = bIsStack ? TEXT("Modifiers") : Modifier->GetModifierName().ToString();
 
+		ModifierSubsystem->ProcessModifierMetadata(
+			Modifier->GetModifierName(),
+			[&HeaderLabel](const FActorModifierCoreMetadata& InMetadata)->bool
+			{
+				HeaderLabel = InMetadata.GetDisplayName().ToString();
+				return true;
+			}
+		);
+
 		if (InItem->GetValueCount() > 1)
 		{
 			HeaderLabel += TEXT(" (") + FString::FromInt(InItem->GetValueCount()) + TEXT(")");
diff --git a/Engine/Plugins/Experimental/Avalanche/Source/AvalancheMask/Private/Mask2D/AvaMask2DBaseModifier.cpp b/Engine/Plugins/Experimental/Avalanche/Source/AvalancheMask/Private/Mask2D/AvaMask2DBaseModifier.cpp
--- a/Engine/Plugins/Experimental/Avalanche/Source/AvalancheMask/Private/Mask2D/AvaMask2DBaseModifier.cpp
+++ b/Engine/Plugins/Experimental/Avalanche/Source/AvalancheMask/Private/Mask2D/AvaMask2DBaseModifier.cpp
@@ -436,6 +436,49 @@ void UAvaMask2DBaseModifier::OnSceneTreeTrackedActorChildrenChanged(int32 InIdx,
 	}
 }
 
+void UAvaMask2DBaseModifier::OnRenderStateUpdated(AActor* InActor, UActorComponent* InComponent)
+{
+	Super::OnRenderStateUpdated(InActor, InComponent);
+
+	if (!IsValid(InActor) || !IsValid(InComponent))
+	{
+		return;
+	}
+
+	const AActor* ActorModified = GetModifiedActor();
+
+	if (!(ActorModified == InActor || InActor->IsAttachedTo(ActorModified)))
+	{
+		return;
+	}
+
+	if (!ActorSupportsMaskReadWrite(InActor))
+	{
+		return;
+	}
+
+	const UPrimitiveComponent* PrimitiveComponent = Cast<UPrimitiveComponent>(InComponent);
+	if (!PrimitiveComponent)
+	{
+		return;
+	}
+	
+	if (CanvasWeak.IsValid())
+	{
+		return;
+	}
+
+	if (const FAvaMask2DActorData* ActorMaskData = ActorData.Find(InActor))
+	{
+		if (ActorMaskData->CanvasTextureWeak.IsValid())
+		{
+			return;
+		}
+	}
+
+	MarkModifierDirty();
+}
+
 FName UAvaMask2DBaseModifier::GenerateUniqueMaskName() const
 {
 	const AActor* ModifierActor = GetModifiedActor();
diff --git a/Engine/Plugins/Experimental/Avalanche/Source/AvalancheMask/Private/Mask2D/AvaMask2DReadModifier.cpp b/Engine/Plugins/Experimental/Avalanche/Source/AvalancheMask/Private/Mask2D/AvaMask2DReadModifier.cpp
--- a/Engine/Plugins/Experimental/Avalanche/Source/AvalancheMask/Private/Mask2D/AvaMask2DReadModifier.cpp
+++ b/Engine/Plugins/Experimental/Avalanche/Source/AvalancheMask/Private/Mask2D/AvaMask2DReadModifier.cpp
@@ -48,8 +48,8 @@ void UAvaMask2DReadModifier::PostEditChangeProperty(FPropertyChangedEvent& InPro
 void UAvaMask2DReadModifier::Apply()
 {
 	Super::Apply();
-	
-	if (GetCurrentCanvas())
+
+	if (const UGeometryMaskCanvas* Canvas = GetCurrentCanvas())
 	{
 		for (TPair<TWeakObjectPtr<AActor>, FAvaMask2DActorData>& ActorDataPair : ActorData)
 		{
@@ -98,21 +98,20 @@ bool UAvaMask2DReadModifier::ApplyRead(AActor* InActor, FAvaMask2DActorData& InA
 			return true;
 		}
 
-		if (const UGeometryMaskCanvas* GeometryCanvas = GetCurrentCanvas())
 		{
 			FAvaMask2DSubjectParameters ApplyParameters;
 			ApplyParameters.MaterialParameters.BlendMode = GetBlendMode();		
-			ApplyParameters.MaterialParameters.CanvasName = GeometryCanvas->GetCanvasName();
-			ApplyParameters.MaterialParameters.Texture = GeometryCanvas->GetTexture();
+			ApplyParameters.MaterialParameters.CanvasName = Canvas->GetCanvasName();
+			ApplyParameters.MaterialParameters.Texture = Canvas->GetTexture();
 			ApplyParameters.MaterialParameters.BaseOpacity = BaseOpacity;
-			ApplyParameters.MaterialParameters.Channel = GeometryCanvas->GetColorChannel();
+			ApplyParameters.MaterialParameters.Channel = Canvas->GetColorChannel();
 			ApplyParameters.MaterialParameters.ChannelAsVector = UE::GeometryMask::MaskChannelEnumToVector[ApplyParameters.MaterialParameters.Channel];
 			ApplyParameters.MaterialParameters.bInvert = bInverted;
 
-			FGeometryMaskDrawingContext DrawingContext(GeometryCanvas->GetCanvasId().Level);
+			FGeometryMaskDrawingContext DrawingContext(Canvas->GetCanvasId().Level);
 
-			const FIntPoint ViewportSize = GeometryCanvas->GetResource()->GetMaxViewportSize();
-			const FVector2f ViewportPadding = FVector2f::One() - FVector2f(ViewportSize + GeometryCanvas->GetResource()->GetViewportPadding(DrawingContext)) / ViewportSize;
+			const FIntPoint ViewportSize = Canvas->GetResource()->GetMaxViewportSize();
+			const FVector2f ViewportPadding = FVector2f::One() - FVector2f(ViewportSize + Canvas->GetResource()->GetViewportPadding(DrawingContext)) / ViewportSize;
 			ApplyParameters.MaterialParameters.Padding = ViewportPadding;
 
 			ApplyParameters.MaterialParameters.bApplyFeathering = bUseFeathering;
@@ -154,8 +153,8 @@ void UAvaMask2DReadModifier::OnModifierCDOSetup(FActorModifierCoreMetadata& InMe
 	InMetadata.DisallowAfter(TEXT("MaskWrite"));
 	InMetadata.DisallowBefore(TEXT("MaskWrite"));
 #if WITH_EDITOR
-	InMetadata.SetDisplayName(FText::FromString(TEXT("Mask (Target)")));
-	InMetadata.SetDescription(LOCTEXT("ModifierDescription", "Allows to use a custom mask texture on attached actors materials"));
+	InMetadata.SetDisplayName(FText::FromString(TEXT("Masked Layer (Output)")));
+	InMetadata.SetDescription(LOCTEXT("ModifierDescription", "Reads from a canvas texture and uses it on materials"));
 #endif
 }
 
diff --git a/Engine/Plugins/Experimental/Avalanche/Source/AvalancheMask/Private/Mask2D/AvaMask2DWriteModifier.cpp b/Engine/Plugins/Experimental/Avalanche/Source/AvalancheMask/Private/Mask2D/AvaMask2DWriteModifier.cpp
--- a/Engine/Plugins/Experimental/Avalanche/Source/AvalancheMask/Private/Mask2D/AvaMask2DWriteModifier.cpp
+++ b/Engine/Plugins/Experimental/Avalanche/Source/AvalancheMask/Private/Mask2D/AvaMask2DWriteModifier.cpp
@@ -112,8 +112,8 @@ void UAvaMask2DWriteModifier::OnModifierCDOSetup(FActorModifierCoreMetadata& InM
 	InMetadata.DisallowAfter(TEXT("MaskRead"));
 	InMetadata.DisallowBefore(TEXT("MaskRead"));
 #if WITH_EDITOR
-	InMetadata.SetDisplayName(FText::FromString(TEXT("Mask (Source)")));
-	InMetadata.SetDescription(LOCTEXT("ModifierDescription", "Allows to use a custom mask texture on attached actors materials"));
+	InMetadata.SetDisplayName(FText::FromString(TEXT("Mask Layer (Input)")));
+	InMetadata.SetDescription(LOCTEXT("ModifierDescription", "Writes to a canvas texture used by the masked layer in materials"));
 #endif
 }
 
diff --git a/Engine/Plugins/Experimental/Avalanche/Source/AvalancheMask/Public/Mask2D/AvaMask2DBaseModifier.h b/Engine/Plugins/Experimental/Avalanche/Source/AvalancheMask/Public/Mask2D/AvaMask2DBaseModifier.h
--- a/Engine/Plugins/Experimental/Avalanche/Source/AvalancheMask/Public/Mask2D/AvaMask2DBaseModifier.h
+++ b/Engine/Plugins/Experimental/Avalanche/Source/AvalancheMask/Public/Mask2D/AvaMask2DBaseModifier.h
@@ -113,14 +113,18 @@ protected:
 	virtual void Apply() override;
 	//~ End UActorModifierCoreBase
 
-	virtual void SaveActorPreState(AActor* InActor, FAvaMask2DActorData& InActorData);
-	virtual void RestoreActorPreState(AActor* InActor, const FAvaMask2DActorData& InActorData);
-
 	//~ Begin IAvaSceneTreeUpdateModifierExtension
 	virtual void OnSceneTreeTrackedActorParentChanged(int32 InIdx, const TArray<TWeakObjectPtr<AActor>>& InPreviousParentActor, const TArray<TWeakObjectPtr<AActor>>& InNewParentActor) override;
 	virtual void OnSceneTreeTrackedActorChildrenChanged(int32 InIdx, const TSet<TWeakObjectPtr<AActor>>& InPreviousChildrenActors, const TSet<TWeakObjectPtr<AActor>>& InNewChildrenActors) override;
 	//~ End IAvaSceneTreeUpdateModifierExtension
 
+	//~ Begin IAvaRenderStateUpdateExtension
+	virtual void OnRenderStateUpdated(AActor* InActor, UActorComponent* InComponent) override;
+	//~ End IAvaRenderStateUpdateExtension
+
+	virtual void SaveActorPreState(AActor* InActor, FAvaMask2DActorData& InActorData);
+	virtual void RestoreActorPreState(AActor* InActor, const FAvaMask2DActorData& InActorData);
+
 	void OnUseParentChannelChanged();
 	void OnChannelChanged();
 	void OnInvertedChanged();
@@ -202,11 +206,6 @@ protected:
 	UPROPERTY(BlueprintReadWrite, Getter, Setter, Category = "Mask2D|Shared", meta = (ClampMin = 0, AllowPrivateAccess = "true", EditCondition = "bUseFeathering", EditConditionHides))
 	int32 InnerFeatherRadius = 16;
 
-#if WITH_EDITOR
-	/** Used for PECP */
-	static const TAvaPropertyChangeDispatcher<UAvaMask2DBaseModifier> PropertyChangeDispatcher;
-#endif
-
 	UPROPERTY(Transient)
 	TWeakObjectPtr<UAvaObjectHandleSubsystem> ObjectHandleSubsystem;
 	
@@ -242,4 +241,10 @@ protected:
 	TMap<TWeakObjectPtr<UMaterialInterface>, FInstancedStruct> MaterialHandleData;
 
 	TMap<TObjectKey<UObject>, TSharedPtr<IAvaMaskMaterialCollectionHandle>> MaterialCollectionHandles;
+
+private:
+#if WITH_EDITOR
+	/** Used for PECP */
+	static const TAvaPropertyChangeDispatcher<UAvaMask2DBaseModifier> PropertyChangeDispatcher;
+#endif
 };
diff --git a/Engine/Plugins/Experimental/Avalanche/Source/AvalancheMask/Public/Mask2D/AvaMask2DReadModifier.h b/Engine/Plugins/Experimental/Avalanche/Source/AvalancheMask/Public/Mask2D/AvaMask2DReadModifier.h
--- a/Engine/Plugins/Experimental/Avalanche/Source/AvalancheMask/Public/Mask2D/AvaMask2DReadModifier.h
+++ b/Engine/Plugins/Experimental/Avalanche/Source/AvalancheMask/Public/Mask2D/AvaMask2DReadModifier.h
@@ -65,11 +65,11 @@ protected:
 
 	bool ApplyRead(AActor* InActor, FAvaMask2DActorData& InActorData);
 
-protected:
 	/** Base opacity/alpha to use in Read mode */
 	UPROPERTY(EditInstanceOnly, BlueprintReadWrite, Getter, Setter, Category = "Mask2D", meta = (ClampMin = 0.0, ClampMax = 1.0, AllowPrivateAccess = "true"))
 	float BaseOpacity = 0.0f;
 
+private:
 #if WITH_EDITOR
 	/** Used for PECP */
 	static const TAvaPropertyChangeDispatcher<UAvaMask2DReadModifier> PropertyChangeDispatcher;
diff --git a/Engine/Plugins/Experimental/Avalanche/Source/AvalancheModifiers/Private/Modifiers/AvaSizeToTextureModifier.cpp b/Engine/Plugins/Experimental/Avalanche/Source/AvalancheModifiers/Private/Modifiers/AvaSizeToTextureModifier.cpp
--- a/Engine/Plugins/Experimental/Avalanche/Source/AvalancheModifiers/Private/Modifiers/AvaSizeToTextureModifier.cpp
+++ b/Engine/Plugins/Experimental/Avalanche/Source/AvalancheModifiers/Private/Modifiers/AvaSizeToTextureModifier.cpp
@@ -175,8 +175,6 @@ void UAvaSizeToTextureModifier::Apply()
 	FVector2D ShapeScale;
 	GetShapeSizeAndScale(ShapeSize, ShapeScale);
 
-	UE_LOG(LogTemp, Warning, TEXT("Apply Modifier : Size %s vs %s Scale %s vs %s"), *ShapeSize.ToString(), *CachedSize.ToString(), *ShapeScale.ToString(), *CachedScale.ToString())
-	
 	if (ShapeSize.IsNearlyZero() || ShapeScale.IsNearlyZero())
 	{
 		Fail(LOCTEXT("InvalidShapeSizeOrScale", "Invalid Shape Size or Scale"));
