MotionDesign : Shapes
- Add auto update option (EDITOR-ONLY) to the parametric material texture style to sync the texture on the modifier when property changes

diff --git a/Engine/Plugins/Experimental/Avalanche/Source/AvalancheEditor/Private/DetailView/Customizations/AvaMeshesDetailCustomization.cpp b/Engine/Plugins/Experimental/Avalanche/Source/AvalancheEditor/Private/DetailView/Customizations/AvaMeshesDetailCustomization.cpp
--- a/Engine/Plugins/Experimental/Avalanche/Source/AvalancheEditor/Private/DetailView/Customizations/AvaMeshesDetailCustomization.cpp
+++ b/Engine/Plugins/Experimental/Avalanche/Source/AvalancheEditor/Private/DetailView/Customizations/AvaMeshesDetailCustomization.cpp
@@ -11,7 +11,6 @@
 #include "DetailWidgetRow.h"
 #include "Dialogs/DlgPickAssetPath.h"
 #include "DynamicMeshes/AvaShape2DDynMeshBase.h"
-#include "DynamicMeshes/AvaShapeDynMeshBase.h"
 #include "DynamicMeshToMeshDescription.h"
 #include "Engine/StaticMesh.h"
 #include "IAssetTools.h"
@@ -20,11 +19,14 @@
 #include "UObject/Object.h"
 #include "Widgets/Layout/SBox.h"
 #include "Widgets/Input/SButton.h"
+#include "Widgets/Input/SCheckBox.h"
 #include "Widgets/Text/STextBlock.h"
 #include "Widgets/SBoxPanel.h"
 
 #define LOCTEXT_NAMESPACE "AvaMeshesDetailCustomization"
 
+const FLazyName FAvaMeshesDetailCustomization::AutoUpdateTextureMetadata = TEXT("AutoUpdateTexture");
+
 void FAvaMeshesDetailCustomization::CustomizeDetails(IDetailLayoutBuilder& InDetailBuilder)
 {
 	TSharedRef<IPropertyHandle> MeshDatasHandle = InDetailBuilder.GetProperty(
@@ -160,7 +162,9 @@ void FAvaMeshesDetailCustomization::CustomizeDetails(IDetailLayoutBuilder& InDet
 						return ParametricMaterialHandle->IsEditable() ? (ParametricChildHandle->IsEditable() ? EVisibility::Visible : EVisibility::Hidden) : EVisibility::Hidden;
 					});
 
-				if (ParametricChildHandle->GetProperty() && ParametricChildHandle->GetProperty()->GetFName() == TEXT("Texture"))
+				if (ParametricChildHandle->GetProperty()
+					&& ParametricChildHandle->GetProperty()->GetFName() == TEXT("Texture")
+					&& CanSizeToTexture())
 				{
 					FDetailWidgetRow& ParametricTextureRow = MaterialCategoryBuilder.AddCustomRow(ParametricChildHandle->GetPropertyDisplayName());
 					ParametricTextureRow.PropertyHandleList({ParametricChildHandle});
@@ -171,7 +175,9 @@ void FAvaMeshesDetailCustomization::CustomizeDetails(IDetailLayoutBuilder& InDet
 						ParametricChildHandle->CreatePropertyNameWidget()
 					];
 
-					const EVisibility SizeToTextureVisibility = CanSizeToTexture() ? EVisibility::Visible : EVisibility::Hidden;
+					ParametricChildHandle->SetOnPropertyValueChanged(FSimpleDelegate::CreateSP(this, &FAvaMeshesDetailCustomization::OnTexturePropertyChanged));
+
+					const FText AutoUpdateTooltip = LOCTEXT("AutoSizeToTexture.Tooltip", "EDITOR-ONLY : Auto update the size to texture modifier when this texture property changes");
 
 					ParametricTextureRow.ValueContent()
 					.HAlign(HAlign_Fill)
@@ -188,17 +194,14 @@ void FAvaMeshesDetailCustomization::CustomizeDetails(IDetailLayoutBuilder& InDet
 						.Padding(2.f, 0.f)
 						.AutoWidth()
 						[
-							SNew(SBox)
-							.HAlign(HAlign_Center)
-							.VAlign(VAlign_Center)
-							.Visibility(SizeToTextureVisibility)
-							.Padding(0.f)
+							SNew(SVerticalBox)
+							+ SVerticalBox::Slot()
 							[
 								SNew(SButton)
-								.ContentPadding(2.f)
+								.ContentPadding(1.f)
 								.VAlign(VAlign_Center)
 								.HAlign(HAlign_Center)
-								.ToolTipText(LOCTEXT("SizeToTextureTooltip", "Find or Add a size to texture modifier and link this texture to it"))
+								.ToolTipText(LOCTEXT("SizeToTextureTooltip", "EDITOR-ONLY : Find or Add a size to texture modifier and link this texture to it"))
 								.OnClicked(this, &FAvaMeshesDetailCustomization::OnSizeToTextureClicked)
 								.Content()
 								[
@@ -207,6 +210,31 @@ void FAvaMeshesDetailCustomization::CustomizeDetails(IDetailLayoutBuilder& InDet
 									.Text(LOCTEXT("SizeToTextureLabel", "Size to Texture"))
 								]
 							]
+							+ SVerticalBox::Slot()
+							[
+								SNew(SHorizontalBox)
+								+ SHorizontalBox::Slot()
+								.AutoWidth()
+								.HAlign(HAlign_Fill)
+								.VAlign(VAlign_Center)
+								[
+									SNew(SCheckBox)
+									.Padding(1.f)
+									.ToolTipText(AutoUpdateTooltip)
+									.IsChecked(this, &FAvaMeshesDetailCustomization::OnIsAutoSizeToTextureChecked)
+									.OnCheckStateChanged(this, &FAvaMeshesDetailCustomization::OnAutoSizeToTextureStateChanged)
+								]
+								+ SHorizontalBox::Slot()
+								.FillWidth(1.f)
+								.HAlign(HAlign_Center)
+								.VAlign(VAlign_Center)
+								[
+									SNew(STextBlock)
+									.Font(IDetailLayoutBuilder::GetDetailFont())
+									.Text(LOCTEXT("AutoSizeToTexture.Label", "Auto Update"))
+									.ToolTipText(AutoUpdateTooltip)
+								]
+							]
 						]
 					];
 				}
@@ -392,50 +420,143 @@ FReply FAvaMeshesDetailCustomization::OnSizeToTextureClicked()
 
 	for (const TWeakObjectPtr<UAvaShapeDynamicMeshBase>& MeshGeneratorWeak : MeshGeneratorsWeak)
 	{
-		UAvaShapeDynamicMeshBase* MeshGenerator = MeshGeneratorWeak.Get();
+		ApplySizeToTexture(MeshGeneratorWeak.Get());
+	}
+
+	return FReply::Handled();
+}
+
+ECheckBoxState FAvaMeshesDetailCustomization::OnIsAutoSizeToTextureChecked() const
+{
+	TOptional<ECheckBoxState> State;
+
+	for (const TWeakObjectPtr<UAvaShapeDynamicMeshBase>& MeshGeneratorWeak : MeshGeneratorsWeak)
+	{
+		const UAvaShapeDynamicMeshBase* MeshGenerator = MeshGeneratorWeak.Get();
 		if (!MeshGenerator)
 		{
 			continue;
 		}
 
-		const FAvaShapeParametricMaterial* PrimaryParametricMaterial = MeshGenerator->GetParametricMaterial(UAvaShapeDynamicMeshBase::MESH_INDEX_PRIMARY);
-		if (!PrimaryParametricMaterial)
+		if (!MeshGenerator->IsA<UAvaShape2DDynMeshBase>())
 		{
 			continue;
 		}
 
-		AActor* ShapeActor = MeshGenerator->GetShapeActor();
-		if (!ShapeActor)
+		const ECheckBoxState AutoState = MeshGenerator->ComponentTags.Contains(AutoUpdateTextureMetadata) ? ECheckBoxState::Checked : ECheckBoxState::Unchecked;
+
+		if (!State.IsSet())
 		{
-			continue;
+			State = AutoState;
 		}
-
-		UActorModifierCoreStack* Stack = nullptr;
-		if (!UActorModifierCoreLibrary::FindModifierStack(ShapeActor, Stack, /** CreateIfNone */true))
+		else if (State.GetValue() != AutoState)
 		{
-			continue;
+			State = ECheckBoxState::Undetermined;
+			break;
 		}
+	}
+
+	return State.Get(ECheckBoxState::Undetermined);
+}
 
-		UActorModifierCoreBase* SizeToTextureModifier = UActorModifierCoreLibrary::FindModifierByClass(Stack, UAvaSizeToTextureModifier::StaticClass());
-		if (!SizeToTextureModifier)
+void FAvaMeshesDetailCustomization::OnTexturePropertyChanged()
+{
+	// Update texture on next tick to give it time to load resources
+	FTSTicker::GetCoreTicker().AddTicker(FTickerDelegate::CreateSPLambda(this, [this](float)
+	{
+		for (const TWeakObjectPtr<UAvaShapeDynamicMeshBase>& MeshGeneratorWeak : MeshGeneratorsWeak)
 		{
-			FActorModifierCoreInsertOperation InsertOp;
-			InsertOp.ModifierClass = UAvaSizeToTextureModifier::StaticClass();
-			InsertOp.InsertPosition = EActorModifierCoreStackPosition::After;
-			InsertOp.InsertPositionContext = nullptr;
+			UAvaShapeDynamicMeshBase* MeshGenerator = MeshGeneratorWeak.Get();
+			if (!MeshGenerator)
+			{
+				continue;
+			}
 
-			if (!UActorModifierCoreLibrary::InsertModifier(Stack, InsertOp, SizeToTextureModifier))
+			if (!MeshGenerator->IsA<UAvaShape2DDynMeshBase>())
 			{
 				continue;
 			}
+
+			if (!MeshGenerator->ComponentTags.Contains(AutoUpdateTextureMetadata))
+			{
+				continue;
+			}
+			
+			ApplySizeToTexture(MeshGenerator);
+		}
+
+		return false; // Stop
+	}));
+}
+
+void FAvaMeshesDetailCustomization::OnAutoSizeToTextureStateChanged(ECheckBoxState InState)
+{
+	for (const TWeakObjectPtr<UAvaShapeDynamicMeshBase>& MeshGeneratorWeak : MeshGeneratorsWeak)
+	{
+		UAvaShapeDynamicMeshBase* MeshGenerator = MeshGeneratorWeak.Get();
+		if (!MeshGenerator)
+		{
+			continue;
+		}
+
+		if (!MeshGenerator->IsA<UAvaShape2DDynMeshBase>())
+		{
+			continue;
+		}
+
+		if (InState == ECheckBoxState::Checked)
+		{
+			MeshGenerator->ComponentTags.AddUnique(AutoUpdateTextureMetadata);
+		}
+		else
+		{
+			MeshGenerator->ComponentTags.Remove(AutoUpdateTextureMetadata);
 		}
+	}
+}
 
-		UAvaSizeToTextureModifier* SizeToTextureModifierCasted = Cast<UAvaSizeToTextureModifier>(SizeToTextureModifier);
-		SizeToTextureModifierCasted->SetTexture(PrimaryParametricMaterial->GetTexture());
-		MeshGenerator->SetMaterialUVMode(UAvaShapeDynamicMeshBase::MESH_INDEX_PRIMARY, EAvaShapeUVMode::Stretch);
+void FAvaMeshesDetailCustomization::ApplySizeToTexture(UAvaShapeDynamicMeshBase* InMeshGenerator)
+{
+	if (!InMeshGenerator)
+	{
+		return;
 	}
 
-	return FReply::Handled();
+	const FAvaShapeParametricMaterial* PrimaryParametricMaterial = InMeshGenerator->GetParametricMaterial(UAvaShapeDynamicMeshBase::MESH_INDEX_PRIMARY);
+	if (!PrimaryParametricMaterial)
+	{
+		return;
+	}
+
+	AActor* ShapeActor = InMeshGenerator->GetShapeActor();
+	if (!ShapeActor)
+	{
+		return;
+	}
+
+	UActorModifierCoreStack* Stack = nullptr;
+	if (!UActorModifierCoreLibrary::FindModifierStack(ShapeActor, Stack, /** CreateIfNone */true))
+	{
+		return;
+	}
+
+	UActorModifierCoreBase* SizeToTextureModifier = UActorModifierCoreLibrary::FindModifierByClass(Stack, UAvaSizeToTextureModifier::StaticClass());
+	if (!SizeToTextureModifier)
+	{
+		FActorModifierCoreInsertOperation InsertOp;
+		InsertOp.ModifierClass = UAvaSizeToTextureModifier::StaticClass();
+		InsertOp.InsertPosition = EActorModifierCoreStackPosition::After;
+		InsertOp.InsertPositionContext = nullptr;
+
+		if (!UActorModifierCoreLibrary::InsertModifier(Stack, InsertOp, SizeToTextureModifier))
+		{
+			return;
+		}
+	}
+
+	UAvaSizeToTextureModifier* SizeToTextureModifierCasted = Cast<UAvaSizeToTextureModifier>(SizeToTextureModifier);
+	SizeToTextureModifierCasted->SetTexture(PrimaryParametricMaterial->GetTexture());
+	InMeshGenerator->SetMaterialUVMode(UAvaShapeDynamicMeshBase::MESH_INDEX_PRIMARY, EAvaShapeUVMode::Stretch);
 }
 
 bool FAvaMeshesDetailCustomization::CanConvertToStaticMesh() const
diff --git a/Engine/Plugins/Experimental/Avalanche/Source/AvalancheEditor/Private/DetailView/Customizations/AvaMeshesDetailCustomization.h b/Engine/Plugins/Experimental/Avalanche/Source/AvalancheEditor/Private/DetailView/Customizations/AvaMeshesDetailCustomization.h
--- a/Engine/Plugins/Experimental/Avalanche/Source/AvalancheEditor/Private/DetailView/Customizations/AvaMeshesDetailCustomization.h
+++ b/Engine/Plugins/Experimental/Avalanche/Source/AvalancheEditor/Private/DetailView/Customizations/AvaMeshesDetailCustomization.h
@@ -2,14 +2,16 @@
 
 #pragma once
 
+#include "Containers/Array.h"
 #include "Delegates/Delegate.h"
 #include "IDetailCustomization.h"
 #include "Templates/SharedPointer.h"
 #include "UObject/WeakObjectPtr.h"
+#include "UObject/WeakObjectPtrTemplates.h"
 
 class FReply;
 class UAvaShapeDynamicMeshBase;
-class UDynamicMaterialModel;
+enum class ECheckBoxState : uint8;
 
 /** Used to create the details materials meshes widget and export to StaticMesh */
 class FAvaMeshesDetailCustomization : public IDetailCustomization
@@ -25,6 +27,8 @@ public:
 	//~ End IDetailCustomization
 
 protected:
+	static const FLazyName AutoUpdateTextureMetadata;
+
 	/** Enable or disable the button when selected object is not compatible */
 	bool CanConvertToStaticMesh() const;
 
@@ -32,8 +36,11 @@ protected:
 	FReply OnConvertToStaticMeshClicked();
 
 	bool CanSizeToTexture() const;
-
 	FReply OnSizeToTextureClicked();
+	ECheckBoxState OnIsAutoSizeToTextureChecked() const;
+	void OnTexturePropertyChanged();
+	void OnAutoSizeToTextureStateChanged(ECheckBoxState InState);
+	void ApplySizeToTexture(UAvaShapeDynamicMeshBase* InMeshGenerator);
 
 	TArray<TWeakObjectPtr<UAvaShapeDynamicMeshBase>> MeshGeneratorsWeak;
 };
\ No newline at end of file
