[Motion Design] Camera system: adding a cvar to disable camera controls for rundown to run in "ICVFX mode" on a display cluster for stage testing.

diff --git a/Engine/Plugins/Experimental/Avalanche/Source/AvalancheCamera/Private/AvaCameraSubsystem.cpp b/Engine/Plugins/Experimental/Avalanche/Source/AvalancheCamera/Private/AvaCameraSubsystem.cpp
--- a/Engine/Plugins/Experimental/Avalanche/Source/AvalancheCamera/Private/AvaCameraSubsystem.cpp
+++ b/Engine/Plugins/Experimental/Avalanche/Source/AvalancheCamera/Private/AvaCameraSubsystem.cpp
@@ -15,6 +15,12 @@
 #include "Subsystems/ActorModifierCoreSubsystem.h"
 #include "Transition/Tasks/AvaCameraBlendTask.h"
 
+static TAutoConsoleVariable<bool> CVarEnableMotionDesignCameraSubsystem(
+	TEXT("MotionDesignCamera.EnableCameraSubsystem"),
+	true,	
+	TEXT("Enable/Disable Motion Design's camera subsystem."),
+	ECVF_Default);
+
 bool FAvaViewTarget::IsValid() const
 {
 	return ::IsValid(Actor);
@@ -131,7 +137,7 @@ bool UAvaCameraSubsystem::IsBlendingToViewTarget(const ULevel* InSceneLevel) con
 
 void UAvaCameraSubsystem::UpdatePlayerControllerViewTarget(const FViewTargetTransitionParams* InOverrideTransitionParams)
 {
-	if (!IsValid(PlayerController) || ViewTargets.IsEmpty())
+	if (!IsValid(PlayerController) || ViewTargets.IsEmpty() || !CVarEnableMotionDesignCameraSubsystem.GetValueOnGameThread())
 	{
 		return;
 	}
