[MotionDesign] Broadcast Editor: Fix a channel status updates during tear down causing check failure.

Repro: with PIE running and a broadcast channel connected to it (and broadcasting), having a page running, then exit the editor.

I don't think it is justified to have an editor widgets cause check failures for backend issues, especially during tear down.
The widget will still be able to behave even if the media output has become invalid during its lifetime. Should it not be because of a tear down, it is still useful for the widget to display something informative.

diff --git a/Engine/Plugins/Experimental/Avalanche/Source/AvalancheMedia/Public/Broadcast/AvaBroadcast.h b/Engine/Plugins/Experimental/Avalanche/Source/AvalancheMedia/Public/Broadcast/AvaBroadcast.h
--- a/Engine/Plugins/Experimental/Avalanche/Source/AvalancheMedia/Public/Broadcast/AvaBroadcast.h
+++ b/Engine/Plugins/Experimental/Avalanche/Source/AvalancheMedia/Public/Broadcast/AvaBroadcast.h
@@ -15,7 +15,7 @@ class FDelegateHandle;
 DECLARE_MULTICAST_DELEGATE_OneParam(FOnAvaBroadcastChanged, EAvaBroadcastChange)
 DECLARE_MULTICAST_DELEGATE_OneParam(FOnAvaBroadcastChannelsListChanged, const FAvaBroadcastProfile& /*InProfile*/);
 
-DECLARE_LOG_CATEGORY_EXTERN(LogAvaBroadcast, Log, All);
+AVALANCHEMEDIA_API DECLARE_LOG_CATEGORY_EXTERN(LogAvaBroadcast, Log, All);
 
 /**
  * Single Instance Class that manages all the Output Channels
diff --git a/Engine/Plugins/Experimental/Avalanche/Source/AvalancheMediaEditor/Private/Broadcast/ChannelGrid/AvaBroadcastOutputTileItem.cpp b/Engine/Plugins/Experimental/Avalanche/Source/AvalancheMediaEditor/Private/Broadcast/ChannelGrid/AvaBroadcastOutputTileItem.cpp
--- a/Engine/Plugins/Experimental/Avalanche/Source/AvalancheMediaEditor/Private/Broadcast/ChannelGrid/AvaBroadcastOutputTileItem.cpp
+++ b/Engine/Plugins/Experimental/Avalanche/Source/AvalancheMediaEditor/Private/Broadcast/ChannelGrid/AvaBroadcastOutputTileItem.cpp
@@ -12,14 +12,12 @@
 #include "Input/Reply.h"
 #include "Layout/Geometry.h"
 #include "MediaOutput.h"
-#include "Rundown/ChannelStatus/SAvaRundownChannelStatus.h"
 #include "Styling/SlateBrush.h"
 #include "Widgets/Images/SImage.h"
 #include "Widgets/Layout/SScaleBox.h"
 #include "Widgets/SBoxPanel.h"
 #include "Widgets/SWidget.h"
 #include "Widgets/Text/STextBlock.h"
-#include "UObject/NoExportTypes.h"
 #include "UObject/UnrealType.h"
 
 #define LOCTEXT_NAMESPACE "AvaBroadcastOutputTileItem"
@@ -57,8 +55,12 @@ FAvaBroadcastOutputChannel& FAvaBroadcastOutputTileItem::GetChannel()
 
 TSharedRef<SWidget> FAvaBroadcastOutputTileItem::GenerateTile() const
 {
-	check(MediaOutput.IsValid());
-	
+	// Although not norminal, we allow the tile to be created. It will show status as "invalid".
+	if (!MediaOutput.IsValid())
+	{
+		UE_LOG(LogAvaBroadcast, Error, TEXT("Broadcast Output Tile generated with invalid output."));
+	}
+
 	return SNew(SHorizontalBox)
 		.ToolTipText(this, &FAvaBroadcastOutputTileItem::GetToolTipText)
 		+ SHorizontalBox::Slot()
@@ -135,8 +137,8 @@ FText FAvaBroadcastOutputTileItem::GetToolTipText() const
 
 const FSlateBrush* FAvaBroadcastOutputTileItem::GetMediaOutputIcon() const
 {
-	check(MediaOutput.IsValid());
-	return FClassIconFinder::FindThumbnailForClass(MediaOutput->GetClass());
+	const UClass* MediaOutputClass = MediaOutput.IsValid() ? MediaOutput->GetClass() : UMediaOutput::StaticClass();   
+	return FClassIconFinder::FindThumbnailForClass(MediaOutputClass);
 }
 
 const FSlateBrush* FAvaBroadcastOutputTileItem::GetMediaOutputStatusBrush() const
@@ -196,6 +198,20 @@ void FAvaBroadcastOutputTileItem::OnBroadcastChanged(EAvaBroadcastChange InChang
 
 void FAvaBroadcastOutputTileItem::UpdateInfo()
 {
+	// The capture backend may fire events during tear down.
+	if (!MediaOutput.IsValid())
+	{
+		MediaOutputStatusText = LOCTEXT("MediaOutputStatus_Invalid", "Invalid");
+		MediaOutputStatusBrush = FAvaMediaEditorStyle::Get().GetBrush("AvaMediaEditor.MediaOutputOffline");
+
+		// Only update if it was uninitialized, otherwise leave as is (may help debugging).
+		if (MediaOutputDisplayText.IsEmpty())
+		{
+			MediaOutputDisplayText = LOCTEXT("MediaOutputDisplayName_Invalid", "Invalid Output");
+		}
+		return;
+	}
+
 	MediaOutputDisplayText = FindLatestDisplayText();
 
 	const FAvaBroadcastOutputChannel& Channel = GetChannel();
