[Motion Design] Rundonw Take Out Layers Command / Page.

- Adding new framework for page commands. The only  implementation for now is a "Stop layers" command.
- Fixing a few issues with page list selection (not refreshing page details when switching focus between different page lists)
- Adding better tooltip for page status to reveal why a page can't be played.

diff --git a/Engine/Plugins/Experimental/Avalanche/Source/AvalancheMedia/Internal/Rundown/AvaRundownMessages.h b/Engine/Plugins/Experimental/Avalanche/Source/AvalancheMedia/Internal/Rundown/AvaRundownMessages.h
--- a/Engine/Plugins/Experimental/Avalanche/Source/AvalancheMedia/Internal/Rundown/AvaRundownMessages.h
+++ b/Engine/Plugins/Experimental/Avalanche/Source/AvalancheMedia/Internal/Rundown/AvaRundownMessages.h
@@ -696,6 +696,10 @@ public:
 	UPROPERTY()
 	bool bTransitionLogicEnabled = false;
 
+	/** Page Commands that can be executed when playing this page. */
+	UPROPERTY()
+	TArray<FAvaRundownPageCommandData> Commands;
+	
 	UPROPERTY()
 	FString OutputChannel;
 
diff --git a/Engine/Plugins/Experimental/Avalanche/Source/AvalancheMedia/Internal/Rundown/AvaRundownPlaybackUtils.h b/Engine/Plugins/Experimental/Avalanche/Source/AvalancheMedia/Internal/Rundown/AvaRundownPlaybackUtils.h
--- a/Engine/Plugins/Experimental/Avalanche/Source/AvalancheMedia/Internal/Rundown/AvaRundownPlaybackUtils.h
+++ b/Engine/Plugins/Experimental/Avalanche/Source/AvalancheMedia/Internal/Rundown/AvaRundownPlaybackUtils.h
@@ -16,4 +16,18 @@ struct AVALANCHEMEDIA_API FAvaRundownPlaybackUtils
 	static int32 GetPageIdToPlayNext(const UAvaRundown* InRundown, const FAvaRundownPageListReference& InPageListReference, bool bInPreview, FName InPreviewChannelName);
 	
 	static TArray<int32> GetPagesToTakeToProgram(const UAvaRundown* InRundown, const TArray<int32>& InSelectedPageIds, FName InPreviewChannel = NAME_None);
+
+	/** Version of AddUnique that will only add the tag handle to the array if it does not have an ExactMatch. */
+	static void AddTagHandleUnique(TArray<FAvaTagHandle>& InTagHandles, const FAvaTagHandle& InTagHandleToAdd)
+	{
+		const bool bTagAlreadyAdded = InTagHandles.ContainsByPredicate([&InTagHandleToAdd](const FAvaTagHandle& InTagHandle)
+		{
+			return InTagHandle.MatchesExact(InTagHandleToAdd); // Exact math: Source and TagId.
+		});
+			
+		if (!bTagAlreadyAdded)
+		{
+			InTagHandles.Add(InTagHandleToAdd);
+		}
+	}
 };
\ No newline at end of file
diff --git a/Engine/Plugins/Experimental/Avalanche/Source/AvalancheMedia/Internal/Rundown/AvaRundownSerializationUtils.h b/Engine/Plugins/Experimental/Avalanche/Source/AvalancheMedia/Internal/Rundown/AvaRundownSerializationUtils.h
--- a/Engine/Plugins/Experimental/Avalanche/Source/AvalancheMedia/Internal/Rundown/AvaRundownSerializationUtils.h
+++ b/Engine/Plugins/Experimental/Avalanche/Source/AvalancheMedia/Internal/Rundown/AvaRundownSerializationUtils.h
@@ -5,8 +5,11 @@
 #include "HAL/Platform.h"
 
 class FArchive;
+class FString;
 class FText;
 class UAvaRundown;
+class UStruct;
+struct FAvaRundownPageCommand;
 
 namespace UE::AvaMedia::RundownSerializationUtils
 {
@@ -15,4 +18,7 @@ namespace UE::AvaMedia::RundownSerializationUtils
 
 	AVALANCHEMEDIA_API bool LoadRundownFromJson(UAvaRundown* InRundown, FArchive& InArchive, FText& OutErrorMessage);
 	AVALANCHEMEDIA_API bool LoadRundownFromJson(UAvaRundown* InRundown, const TCHAR* InFilepath, FText& OutErrorMessage);
+
+	AVALANCHEMEDIA_API bool SerializeRundownPageCommandToJsonString(const FAvaRundownPageCommand* InRundownPageCommand, UStruct& InStruct, FString& OutString);
+	AVALANCHEMEDIA_API bool DeserializeRundownPageCommandFromJson(FAvaRundownPageCommand* InRundownPageCommand, UStruct& InStruct, const FString& InString);
 }
\ No newline at end of file
diff --git a/Engine/Plugins/Experimental/Avalanche/Source/AvalancheMedia/Private/Playable/Transition/AvaPlayableLocalTransition.cpp b/Engine/Plugins/Experimental/Avalanche/Source/AvalancheMedia/Private/Playable/Transition/AvaPlayableLocalTransition.cpp
--- a/Engine/Plugins/Experimental/Avalanche/Source/AvalancheMedia/Private/Playable/Transition/AvaPlayableLocalTransition.cpp
+++ b/Engine/Plugins/Experimental/Avalanche/Source/AvalancheMedia/Private/Playable/Transition/AvaPlayableLocalTransition.cpp
@@ -144,10 +144,8 @@ namespace UE::AvaPlayableTransition::Private
 					ExecutorBuilder.AddExitInstance(BehaviorInstance);
 					++NumExitInstance;
 					
-					// Special case of transition out (no enter instances).
-					// The way this works, we need to make a dummy "Enter" behavior instance
-					// to provide the transition layer for the logic take out.
-					if (NumEnterInstance == 0 && InPlayableRole == EAvaPlayableTransitionEntryRole::Exit)
+					// For exit playables, we need to add a "take out" dummy enter instance.
+					if (InPlayableRole == EAvaPlayableTransitionEntryRole::Exit)
 					{
 						FAvaTransitionBehaviorInstance NullEnterInstance;
 						if (MakeNullTransitionBehaviorInstance(PlayableTransition, GetTransitionLayer(BehaviorInstance), NullEnterInstance))
diff --git a/Engine/Plugins/Experimental/Avalanche/Source/AvalancheMedia/Private/Rundown/AvaRundown.cpp b/Engine/Plugins/Experimental/Avalanche/Source/AvalancheMedia/Private/Rundown/AvaRundown.cpp
--- a/Engine/Plugins/Experimental/Avalanche/Source/AvalancheMedia/Private/Rundown/AvaRundown.cpp
+++ b/Engine/Plugins/Experimental/Avalanche/Source/AvalancheMedia/Private/Rundown/AvaRundown.cpp
@@ -10,6 +10,7 @@
 #include "Playback/AvaPlaybackManager.h"
 #include "Playback/AvaPlaybackUtils.h"
 #include "Rundown/AvaRundownManagedInstanceCache.h"
+#include "Rundown/AvaRundownPageCommand.h"
 #include "Rundown/AvaRundownPageLoadingManager.h"
 #include "Rundown/AvaRundownPagePlayer.h"
 #include "Rundown/AvaRundownPlaybackClientWatcher.h"
@@ -1211,8 +1212,10 @@ TArray<int32> UAvaRundown::PlayPages(const TArray<int32>& InPageIds, EAvaRundown
 		{
 			const bool bIsPreview = UE::AvaRundown::IsPreviewPlayType(InPlayType);
 
-			if (!IsChannelTypeCompatibleForRequest(SelectedPage, bIsPreview, InPreviewChannelName, true) || !CanPlayPage(PageId, bIsPreview))
+			FString FailureReason;
+			if (!IsChannelTypeCompatibleForRequest(SelectedPage, bIsPreview, InPreviewChannelName, &FailureReason) || !CanPlayPage(PageId, bIsPreview))
 			{
+				UE_LOG(LogAvaRundown, Error, TEXT("Page Id:%d failed to play: %s."), PageId, *FailureReason);
 				continue;
 			}
 
@@ -1244,9 +1247,10 @@ bool UAvaRundown::RestorePlaySubPage(int32 InPageId, int32 InSubPageIndex, const
 		return false;
 	}
 
-	if (!IsChannelTypeCompatibleForRequest(Page, bInIsPreview, InPreviewChannelName, /*bInLogFailureReason*/ true))
+	FString FailureReason;
+	if (!IsChannelTypeCompatibleForRequest(Page, bInIsPreview, InPreviewChannelName, &FailureReason))
 	{
-		LogError(TEXT("Channel Type is not compatible"));
+		LogError(FString::Printf(TEXT("Channel Type is not compatible: %s"), *FailureReason));
 		return false;
 	}
 
@@ -1292,20 +1296,35 @@ bool UAvaRundown::RestorePlaySubPage(int32 InPageId, int32 InSubPageIndex, const
 
 bool UAvaRundown::CanPlayPage(int32 InPageId, bool bInPreview) const
 {
-	return CanPlayPage(InPageId, bInPreview, bInPreview ? GetDefaultPreviewChannelName() : NAME_None);
+	return CanPlayPage(InPageId, bInPreview, bInPreview ? GetDefaultPreviewChannelName() : NAME_None, nullptr);
 }
 
-bool UAvaRundown::CanPlayPage(int32 InPageId, bool bInPreview, const FName& InPreviewChannelName) const
+bool UAvaRundown::CanPlayPage(int32 InPageId, bool bInPreview, const FName& InPreviewChannelName, FString* OutFailureReason) const
 {
+	using namespace UE::AvaPlayback::Utils;
+	
 	// Check if the page is valid and enabled.
 	const FAvaRundownPage& SelectedPage = GetPage(InPageId);
-	if (!SelectedPage.IsValidPage() || !SelectedPage.IsEnabled())
+	if (!SelectedPage.IsValidPage())
 	{
+		if (OutFailureReason)
+		{
+			*OutFailureReason = TEXT("Invalid Page Id");
+		}
+		return false;
+	}
+
+	if (!SelectedPage.IsEnabled())
+	{
+		if (OutFailureReason)
+		{
+			*OutFailureReason = TEXT("Page is disabled");
+		}
 		return false;
 	}
 
 	// Check channel validity and type compatibility.
-	if (!IsChannelTypeCompatibleForRequest(SelectedPage, bInPreview, InPreviewChannelName, false))
+	if (!IsChannelTypeCompatibleForRequest(SelectedPage, bInPreview, InPreviewChannelName, OutFailureReason))
 	{
 		return false;
 	}
@@ -1313,17 +1332,47 @@ bool UAvaRundown::CanPlayPage(int32 InPageId, bool bInPreview, const FName& InPr
 	// Check that if it is a template page it is meant to preview
 	if (SelectedPage.IsTemplate() && !bInPreview)
 	{
+		if (OutFailureReason)
+		{
+			*OutFailureReason = TEXT("Template page can't be taken to program");
+		}
 		return false;
 	}
 
 	// Check if the asset path is valid
-	if (SelectedPage.GetAssetPath(this).IsNull())
+	if (!SelectedPage.HasAssets(*this))
 	{
-		return false;
+		if (!SelectedPage.HasCommands(this))
+		{
+			if (OutFailureReason)
+			{
+				*OutFailureReason = TEXT("Page has no assets nor commands");
+			}
+			return false; // no asset and no commands, can't play.
+		}
+
+		// Verify if the commands can be executed.
+		bool bCanRunCommand = false;
+		FName ChannelName = bInPreview ? InPreviewChannelName : SelectedPage.GetChannelName();
+		SelectedPage.ForEachInstancedCommands([&bCanRunCommand, this, &SelectedPage, ChannelName, OutFailureReason](const FAvaRundownPageCommand& InCommand, const FAvaRundownPage& InPage)
+		{
+			FString FailureReason;
+			const bool bCanExec = InCommand.CanExecute(*this, SelectedPage, ChannelName, OutFailureReason ? &FailureReason : nullptr);
+			bCanRunCommand |= bCanExec;
+			if (!bCanExec && OutFailureReason)
+			{
+				*OutFailureReason += FString::Printf(TEXT("%s%s"), OutFailureReason->IsEmpty() ? TEXT("") : TEXT("; "), *FailureReason);
+			}
+		}, this, /*bDirectOnly*/ false);
+
+		if (!bCanRunCommand)
+		{
+			return false;
+		}
 	}
 
 	// For page with TL, need to see if a transition can be started for that page.
-	if (SelectedPage.HasTransitionLogic(this) && !CanStartTransitionForPage(SelectedPage, bInPreview, InPreviewChannelName))
+	if (SelectedPage.HasTransitionLogic(this) && !CanStartTransitionForPage(SelectedPage, bInPreview, InPreviewChannelName, OutFailureReason))
 	{
 		return false;
 	}
@@ -1380,28 +1429,132 @@ TArray<int32> UAvaRundown::StopPages(const TArray<int32>& InPageIds, EAvaRundown
 
 bool UAvaRundown::CanStopPage(int32 InPageId, EAvaRundownPageStopOptions InOptions, bool bInPreview) const
 {
-	return CanStopPage(InPageId, InOptions, bInPreview, bInPreview ? GetDefaultPreviewChannelName() : NAME_None);
+	return CanStopPage(InPageId, InOptions, bInPreview, bInPreview ? GetDefaultPreviewChannelName() : NAME_None, nullptr);
 }
 
-bool UAvaRundown::CanStopPage(int32 InPageId, EAvaRundownPageStopOptions InOptions, bool bInPreview, const FName& InPreviewChannelName) const
+bool UAvaRundown::CanStopPage(int32 InPageId, EAvaRundownPageStopOptions InOptions, bool bInPreview, const FName& InPreviewChannelName, FString* OutFailureReason) const
 {
 	const FAvaRundownPage& SelectedPage = GetPage(InPageId);
 
 	if (!SelectedPage.IsValidPage())
 	{
+		if (OutFailureReason)
+		{
+			*OutFailureReason = TEXT("Invalid Page Id");
+		}
 		return false;
 	}
 	
 	// For page with TL, need to see if a transition can be started for that page.
 	if (!EnumHasAnyFlags(InOptions, EAvaRundownPageStopOptions::ForceNoTransition)
 		&& SelectedPage.HasTransitionLogic(this)
-		&& !CanStartTransitionForPage(SelectedPage, bInPreview, InPreviewChannelName))
+		&& !CanStartTransitionForPage(SelectedPage, bInPreview, InPreviewChannelName, OutFailureReason))
 	{
 		return false;
 	}
 	
 	const UAvaRundownPagePlayer* Player = FindPlayerForPage(InPageId, bInPreview, InPreviewChannelName);
-	return Player && Player->IsPlaying();
+	
+	if (!Player)
+	{
+		if (OutFailureReason)
+		{
+			*OutFailureReason = FString::Printf(TEXT("No page player found on channel \"%s\""), bInPreview ? *InPreviewChannelName.ToString() : *SelectedPage.GetChannelName().ToString());
+		}
+		return false;
+	}
+
+	// Note: for the failure reason, we want to know about missing player vs player not playing.
+	if (OutFailureReason && !Player->IsPlaying())
+	{
+		*OutFailureReason = TEXT("Page Player is not playing.");
+	}
+	
+	return Player->IsPlaying();
+}
+
+TArray<int32> UAvaRundown::StopLayers(FName InChannelName, const TArray<FAvaTagHandle>& InLayers, EAvaRundownPageStopOptions InOptions)
+{
+	TArray<int32> StoppedPageIds;
+	StoppedPageIds.Reserve(PagePlayers.Num());
+
+	// We want to build a transition that is going to kick out the specified layers.
+	FAvaRundownPageTransitionBuilder TransitionBuilder(this);
+	
+	const bool bUseTransition = !EnumHasAnyFlags(InOptions, EAvaRundownPageStopOptions::ForceNoTransition);
+
+	if (bUseTransition)
+	{
+		// We need to make a special transition that kicks out layers
+		if (UAvaRundownPageTransition* PageTransition = TransitionBuilder.FindOrAddTransition(InChannelName))
+		{
+			// Add the layers to kick out.
+			PageTransition->ExitLayers.Append(InLayers);
+		}
+	}
+
+	for (UAvaRundownPagePlayer* PagePlayer : PagePlayers)
+	{
+		if (PagePlayer->ChannelFName != InChannelName)
+		{
+			continue;
+		}
+
+		int32 LayerOverlapCount = 0;
+		PagePlayer->ForEachInstancePlayer([&InLayers, &LayerOverlapCount, bUseTransition](UAvaRundownPlaybackInstancePlayer* InInstancePlayer)
+		{
+			for (const FAvaTagHandle& Layer : InLayers)
+			{
+				if (InInstancePlayer->TransitionLayer.Overlaps(Layer))
+				{
+					if (!bUseTransition)
+					{
+						InInstancePlayer->Stop();
+					}
+					++LayerOverlapCount;
+					break;
+				}
+			}
+		});	
+
+		// Note: if the overlap count is smaller than the number of instances, only part of the page will be taken down.
+		if (LayerOverlapCount)
+		{
+			StoppedPageIds.Add(PagePlayer->PageId);
+		}
+	}
+
+	if (!StoppedPageIds.IsEmpty() && !bUseTransition)
+	{
+		RemoveStoppedPagePlayers();
+	}
+	return StoppedPageIds;
+}
+
+bool UAvaRundown::CanStopLayer(FName InChannelName, const FAvaTagHandle& InLayer) const
+{
+	for (const UAvaRundownPagePlayer* PagePlayer : PagePlayers)
+	{
+		if (PagePlayer->ChannelFName != InChannelName)
+		{
+			continue;
+		}
+
+		int32 LayerOverlapCount = 0;
+		PagePlayer->ForEachInstancePlayer([&InLayer, &LayerOverlapCount](const UAvaRundownPlaybackInstancePlayer* InInstancePlayer)
+		{
+			if (InInstancePlayer->TransitionLayer.Overlaps(InLayer))
+			{
+				++LayerOverlapCount;
+			}
+		});	
+
+		if (LayerOverlapCount)
+		{
+			return true;
+		}
+	}
+	return false;
 }
 
 bool UAvaRundown::StopChannel(const FString& InChannelName)
@@ -1902,11 +2055,12 @@ void UAvaRundown::InitializePage(FAvaRundownPage& InOutPage, int32 InPageId, int
 	InOutPage.PageId = InPageId;
 	InOutPage.TemplateId = InTemplateId;
 	InOutPage.CombinedTemplateIds.Empty();
+	InOutPage.InstancedCommands.Empty();
 	InOutPage.SetPageFriendlyName(FText::GetEmpty());
 	InOutPage.UpdatePageSummary(this);
 }
 
-bool UAvaRundown::IsChannelTypeCompatibleForRequest(const FAvaRundownPage& InSelectedPage, bool bInIsPreview, const FName& InPreviewChannelName, bool bInLogFailureReason) const
+bool UAvaRundown::IsChannelTypeCompatibleForRequest(const FAvaRundownPage& InSelectedPage, bool bInIsPreview, const FName& InPreviewChannelName, FString* OutFailureReason) const
 {
 	// Check channel validity and type compatibility.
 	UAvaBroadcast& Broadcast = UAvaBroadcast::Get();
@@ -1916,9 +2070,9 @@ bool UAvaRundown::IsChannelTypeCompatibleForRequest(const FAvaRundownPage& InSel
 		if (Broadcast.GetCurrentProfile().GetChannel(InPreviewChannelName).IsValidChannel()
 			&& Broadcast.GetChannelType(InPreviewChannelName) != EAvaBroadcastChannelType::Preview)
 		{
-			if (bInLogFailureReason)
+			if (OutFailureReason)
 			{
-				UE_LOG(LogAvaRundown, Error, TEXT("Preview request failed. Channel \"%s\" is not a \"preview\" channel in profile \"%s\"."),
+				*OutFailureReason = FString::Printf(TEXT("Channel \"%s\" is not a \"preview\" channel in profile \"%s\"."),
 					*InPreviewChannelName.ToString(), *Broadcast.GetCurrentProfileName().ToString());
 			}
 			return false;
@@ -1929,18 +2083,18 @@ bool UAvaRundown::IsChannelTypeCompatibleForRequest(const FAvaRundownPage& InSel
 		const FAvaBroadcastOutputChannel& Channel = Broadcast.GetCurrentProfile().GetChannel(InSelectedPage.GetChannelName()); 
 		if (!Channel.IsValidChannel())
 		{
-			if (bInLogFailureReason)
+			if (OutFailureReason)
 			{
-				UE_LOG(LogAvaRundown, Error, TEXT("Playback request failed. Channel \"%s\" is not a valid channel in \"%s\" profile."),
+				*OutFailureReason = FString::Printf(TEXT("Channel \"%s\" is not a valid channel in \"%s\" profile."),
 					*InSelectedPage.GetChannelName().ToString(), *Broadcast.GetCurrentProfileName().ToString());
 			}
 			return false;
 		}
 		if (Broadcast.GetChannelType(InSelectedPage.GetChannelName()) != EAvaBroadcastChannelType::Program)
 		{
-			if (bInLogFailureReason)
+			if (OutFailureReason)
 			{
-				UE_LOG(LogAvaRundown, Error, TEXT("Playback request failed. Channel \"%s\" is not a \"program\" channel in profile \"%s\"."),
+				*OutFailureReason = FString::Printf(TEXT("Channel \"%s\" is not a \"program\" channel in profile \"%s\"."),
 					*InSelectedPage.GetChannelName().ToString(), *Broadcast.GetCurrentProfileName().ToString());
 			}
 			return false;
@@ -1967,9 +2121,9 @@ bool UAvaRundown::IsChannelTypeCompatibleForRequest(const FAvaRundownPage& InSel
 		// the local outputs take priority (for now at least).
 		if (bHasOfflineOutput && !bHasLocalOutput)
 		{
-			if (bInLogFailureReason)
+			if (OutFailureReason)
 			{
-				UE_LOG(LogAvaRundown, Error, TEXT("Playback request failed. Channel \"%s\" is offline."),
+				*OutFailureReason = FString::Printf(TEXT("Channel \"%s\" is offline."),
 					*InSelectedPage.GetChannelName().ToString());
 			}
 			return false;
@@ -2103,6 +2257,27 @@ bool UAvaRundown::PlayPageWithTransition(FAvaRundownPageTransitionBuilder& InBui
 {
 	using namespace UE::AvaMedia::Rundown::Private;
 
+	// Execute page's commands
+	if (InPage.HasCommands(this))
+	{
+		const FName ChannelName = bInIsPreview ? InPreviewChannelName : InPage.GetChannelName();
+
+		constexpr bool bDirectOnly = false;	// Traverse templates.
+
+		InPage.ForEachInstancedCommands(
+			[&InBuilder, this, &ChannelName](const FAvaRundownPageCommand& InCommand, const FAvaRundownPage& InPage)
+			{
+				InCommand.Execute(InBuilder, *this, InPage, ChannelName);
+			},
+			this, bDirectOnly);
+
+		// If this is a page with no assets, stop here to avoid making a page player.
+		if (!InPage.HasAssets(*this))
+		{
+			return true;
+		}
+	}
+
 	// For now, we always start a new page player, loading a new instance.
 	UAvaRundownPagePlayer* NewPagePlayer = NewObject<UAvaRundownPagePlayer>(this);
 
@@ -2374,7 +2549,7 @@ UAvaRundownPageTransition* UAvaRundown::GetPageTransition(const FGuid& InTransit
 	return nullptr;
 }
 
-bool UAvaRundown::CanStartTransitionForPage(const FAvaRundownPage& InPage, bool bInIsPreview, const FName& InPreviewChannelName) const
+bool UAvaRundown::CanStartTransitionForPage(const FAvaRundownPage& InPage, bool bInIsPreview, const FName& InPreviewChannelName, FString* OutFailureReason) const
 {
 	// Current constraint: There can only be one transition (running properly) at a time in a world.
 	// Given that, for now (and the foreseeable future until we support more playable groups per channels),
@@ -2385,6 +2560,10 @@ bool UAvaRundown::CanStartTransitionForPage(const FAvaRundownPage& InPage, bool
 	{
 		if (PageTransition && PageTransition->GetChannelName() == ChannelName)
 		{
+			if (OutFailureReason)
+			{
+				*OutFailureReason = FString::Printf(TEXT("Channel %s already has transition: %s"), *ChannelName.ToString(), *PageTransition->GetBriefTransitionDescription());
+			}
 			return false;
 		}
 	}
diff --git a/Engine/Plugins/Experimental/Avalanche/Source/AvalancheMedia/Private/Rundown/AvaRundownPage.cpp b/Engine/Plugins/Experimental/Avalanche/Source/AvalancheMedia/Private/Rundown/AvaRundownPage.cpp
--- a/Engine/Plugins/Experimental/Avalanche/Source/AvalancheMedia/Private/Rundown/AvaRundownPage.cpp
+++ b/Engine/Plugins/Experimental/Avalanche/Source/AvalancheMedia/Private/Rundown/AvaRundownPage.cpp
@@ -15,7 +15,9 @@
 #include "Rundown/AvaRundownManagedInstance.h"
 #include "Rundown/AvaRundownManagedInstanceCache.h"
 #include "Rundown/AvaRundownPageAssetUtils.h"
+#include "Rundown/AvaRundownPageCommand.h"
 #include "Rundown/AvaRundownPagePlayer.h"
+#include "Rundown/AvaRundownSerializationUtils.h"
 
 #define LOCTEXT_NAMESPACE "AvaRundownPage"
 
@@ -142,13 +144,18 @@ bool FAvaRundownPage::UpdatePageSummary(const UAvaRundown* InRundown)
 		}
 	}
 	
-	return UpdatePageSummary(Presets, false);
+	return UpdatePageSummary(InRundown, Presets, /*bInForceUpdate*/ false);
 }
 
 bool FAvaRundownPage::UpdatePageSummary(const TArray<const URemoteControlPreset*>& InPresets, bool bInIsPresetChanged)
 {
-	const bool bShouldUpdateDescription = bInIsPresetChanged ? true : !HasPageSummary();
-	if (!IsValidPage() || IsTemplate() || InPresets.IsEmpty() || !bShouldUpdateDescription)
+	return UpdatePageSummary(nullptr, InPresets, bInIsPresetChanged);
+}
+
+bool FAvaRundownPage::UpdatePageSummary(const UAvaRundown* InRundown, const TArray<const URemoteControlPreset*>& InPresets, bool bInForceUpdate)
+{
+	const bool bShouldUpdateDescription = bInForceUpdate ? true : !HasPageSummary();
+	if (!IsValidPage() || IsTemplate() || (InPresets.IsEmpty() && !HasCommands(InRundown)) || !bShouldUpdateDescription)
 	{
 		return false;
 	}
@@ -156,6 +163,7 @@ bool FAvaRundownPage::UpdatePageSummary(const TArray<const URemoteControlPreset*
 	TSortedMap<int32, FText> OrderedDescription;
 
 	int32 DisplayIndexOffset = 0;
+
 	for (const URemoteControlPreset* Preset : InPresets)
 	{
 		for (const URCVirtualPropertyBase* VirtualProperty : Preset->GetControllers())
@@ -169,6 +177,16 @@ bool FAvaRundownPage::UpdatePageSummary(const TArray<const URemoteControlPreset*
 		}
 		DisplayIndexOffset += Preset->GetNumControllers();
 	}
+
+	// Also include commands in the page summary.
+	if (InRundown)
+	{
+		ForEachInstancedCommands([&OrderedDescription, &DisplayIndexOffset](const FAvaRundownPageCommand& InCommand, const FAvaRundownPage&)
+		{
+			OrderedDescription.FindOrAdd(DisplayIndexOffset) = FText::FromString(InCommand.GetDescription()); 
+			++DisplayIndexOffset;			
+		}, InRundown, /*bDirectOnly*/ false);
+	}
 	
 	if (!OrderedDescription.IsEmpty())
 	{
@@ -281,6 +299,14 @@ int32 FAvaRundownPage::AppendPageProgramStatuses(const UAvaRundown* InParentRund
 	}
 	
 	const int32 PreviousNumStatuses = OutPageStatuses.Num();
+
+	// Special bypass for command-only page
+	if (!HasAssets(*InParentRundown) && HasCommands(InParentRundown))
+	{
+		// todo: remote status?
+		OutPageStatuses.Add({EAvaBroadcastChannelType::Program, EAvaRundownPageStatus::Available, /*bNeedsSync*/ false});
+		return OutPageStatuses.Num() - PreviousNumStatuses;
+	}
 	
 	IAvaMediaModule& AvaMediaModule = IAvaMediaModule::Get();
 	FAvaPlaybackManager& PlaybackManager = InParentRundown->GetPlaybackManager();
@@ -393,11 +419,18 @@ int32 FAvaRundownPage::AppendPagePreviewStatuses(const UAvaRundown* InParentRund
 		return 0;
 	}
 	
+	const int32 PreviousNumStatuses = OutPageStatuses.Num();
+
+	// Special bypass for command-only page
+	if (!HasAssets(*InParentRundown) && HasCommands(InParentRundown))
+	{
+		OutPageStatuses.Add({EAvaBroadcastChannelType::Program, EAvaRundownPageStatus::Available, /*bNeedsSync*/ false});
+		return OutPageStatuses.Num() - PreviousNumStatuses;
+	}
+
 	FAvaPlaybackManager& PlaybackManager = InParentRundown->GetPlaybackManager();
 	const FSoftObjectPath ResolvedAssetPath = GetAssetPath(InParentRundown);
 
-	const int32 PreviousNumStatuses = OutPageStatuses.Num();
-	
 	// For the preview, we only add the status if it is playing.
 	// We are not really interested if the preview is loaded.
 	if (InParentRundown->IsPagePreviewing(PageId))
@@ -508,6 +541,20 @@ TArray<FSoftObjectPath> FAvaRundownPage::GetAssetPaths(const UAvaRundown* InRund
 	return AssetPaths;
 }
 
+bool FAvaRundownPage::HasAssets(const UAvaRundown& InRundown) const
+{
+	const int32 NumTemplates = GetNumTemplates(&InRundown);
+	for(int32 TemplateIndex = 0; TemplateIndex < NumTemplates; ++TemplateIndex)
+	{
+		const FAvaRundownPage& Template = GetTemplate(&InRundown, TemplateIndex);
+		if (Template.IsValidPage() && !Template.AssetPath.IsNull())
+		{
+			return true;
+		}
+	}
+	return false;
+}
+
 bool FAvaRundownPage::UpdateAsset(const FSoftObjectPath& InAssetPath, bool bInReimportPage)
 {
 	if (IsComboTemplate())
@@ -532,6 +579,52 @@ bool FAvaRundownPage::UpdateAsset(const FSoftObjectPath& InAssetPath, bool bInRe
 	return false;
 }
 
+bool FAvaRundownPage::HasCommands(const UAvaRundown* InRundown) const
+{
+	// Don't count empty commands.
+	auto HasValidCommands = [](const TArray<FInstancedStruct>& InCommands) -> bool
+	{
+		for (const FInstancedStruct& Command : InCommands)
+		{
+			if (Command.GetPtr<FAvaRundownPageCommand>())
+			{
+				return true;
+			}
+		}
+		return false;
+	};
+
+	// Direct test (either instanced page or direct template)
+	if (HasValidCommands(InstancedCommands))
+	{
+		return true;
+	}
+	
+	const FAvaRundownPage& DirectTemplate = ResolveTemplate(InRundown);
+	if (DirectTemplate.IsValidPage())
+	{
+		if (!IsTemplate()	// Already checked if this is a template.
+			&& HasValidCommands(DirectTemplate.InstancedCommands))
+		{
+			return true;
+		}
+
+		// Sub-template commands - Remark: for now, currently only do 1 level of recursion.
+		for (const int32 SubTemplateId : DirectTemplate.GetCombinedTemplateIds())
+		{
+			const FAvaRundownPage& SubTemplate = InRundown->GetPage(SubTemplateId);
+			if (SubTemplate.IsValidPage())
+			{
+				if (HasValidCommands(SubTemplate.InstancedCommands))
+				{
+					return true;
+				}
+			}
+		}
+	}
+	return false;
+}
+
 FName FAvaRundownPage::GetChannelName() const
 {
 	return UAvaBroadcast::Get().GetChannelName(OutputChannel);
@@ -859,4 +952,93 @@ bool FAvaRundownPage::IsTemplateMatchingByValue(const FAvaRundownPage& InTemplat
 	return true;
 }
 
+void FAvaRundownPage::LoadInstancedCommands(const TArray<FAvaRundownPageCommandData>& InCommands)
+{
+	InstancedCommands.Reset(InCommands.Num());
+	
+	for (int32 CommandIndex = 0; CommandIndex < InCommands.Num(); ++CommandIndex)
+	{
+		const FAvaRundownPageCommandData& CommandData = InCommands[CommandIndex];
+		FInstancedStruct CommandInstance;		
+		if (UScriptStruct* CommandStruct = UClass::TryFindTypeSlow<UScriptStruct>(CommandData.Name))
+		{
+			CommandInstance.InitializeAs(CommandStruct);
+			using namespace UE::AvaMedia::RundownSerializationUtils;
+			DeserializeRundownPageCommandFromJson(CommandInstance.GetMutablePtr<FAvaRundownPageCommand>(), *CommandStruct, CommandData.Payload);
+		}
+		InstancedCommands.Add(MoveTemp(CommandInstance));
+	}
+}
+
+TArray<FAvaRundownPageCommandData> FAvaRundownPage::SaveInstancedCommands() const
+{
+	TArray<FAvaRundownPageCommandData> Commands;
+	Commands.SetNum(InstancedCommands.Num());
+	for (int32 CommandIndex = 0; CommandIndex < InstancedCommands.Num(); ++CommandIndex)
+	{
+		if (const FAvaRundownPageCommand* Command = InstancedCommands[CommandIndex].GetPtr<FAvaRundownPageCommand>())
+		{
+			if (UScriptStruct* CommandStruct = const_cast<UScriptStruct*>(InstancedCommands[CommandIndex].GetScriptStruct()))
+			{
+				Commands[CommandIndex].Name = FSoftObjectPath(CommandStruct).ToString();
+				UE::AvaMedia::RundownSerializationUtils::SerializeRundownPageCommandToJsonString(Command, *CommandStruct, Commands[CommandIndex].Payload);
+				continue;
+			}
+		}
+
+		// fallback
+		Commands[CommandIndex].Name.Empty();
+		Commands[CommandIndex].Payload.Empty();
+	}
+	return Commands;
+}
+
+void FAvaRundownPage::SetInstancedCommands(const TArray<FInstancedStruct>& InInstancedCommands)
+{	
+	InstancedCommands = InInstancedCommands;
+}
+
+void FAvaRundownPage::ForEachInstancedCommands(TFunctionRef<void(const FAvaRundownPageCommand&, const FAvaRundownPage&)> InFunction, const UAvaRundown* InRundown, bool bDirectOnly) const
+{
+	auto CallOnCommands = [&InFunction, this](const TArray<FInstancedStruct>& InCommands)
+	{
+		for (const FInstancedStruct& Command : InCommands)
+		{
+			if (const FAvaRundownPageCommand* CommandPtr = Command.GetPtr<FAvaRundownPageCommand>())
+			{
+				InFunction(*CommandPtr, *this);
+			}
+		}
+	};
+	
+	// Direct page commands
+	CallOnCommands(InstancedCommands);
+
+	// Note: if no rundown is provided to resolve templates, we have to stop here.
+	if (bDirectOnly || !InRundown)
+	{
+		return;
+	}
+
+	// Direct Template commands
+	const FAvaRundownPage& DirectTemplate = ResolveTemplate(InRundown);
+	if (DirectTemplate.IsValidPage())
+	{
+		if (!IsTemplate()) // Already done if this is a template.
+		{
+			CallOnCommands(DirectTemplate.InstancedCommands);
+		}
+
+		// Sub-template commands (if any) - Remark: for now, currently only do 1 level of recursion.
+		for (const int32 SubTemplateId : DirectTemplate.GetCombinedTemplateIds())
+		{
+			const FAvaRundownPage& SubTemplate = InRundown->GetPage(SubTemplateId);
+			if (SubTemplate.IsValidPage())
+			{
+				CallOnCommands(SubTemplate.InstancedCommands);
+			}
+		}
+	}
+}
+
 #undef LOCTEXT_NAMESPACE
diff --git a/Engine/Plugins/Experimental/Avalanche/Source/AvalancheMedia/Private/Rundown/AvaRundownPlaybackUtils.cpp b/Engine/Plugins/Experimental/Avalanche/Source/AvalancheMedia/Private/Rundown/AvaRundownPlaybackUtils.cpp
--- a/Engine/Plugins/Experimental/Avalanche/Source/AvalancheMedia/Private/Rundown/AvaRundownPlaybackUtils.cpp
+++ b/Engine/Plugins/Experimental/Avalanche/Source/AvalancheMedia/Private/Rundown/AvaRundownPlaybackUtils.cpp
@@ -34,8 +34,16 @@ int32 FAvaRundownPlaybackUtils::GetPageIdToPlayNext(const UAvaRundown* InRundown
 		{
 			// Check if the current play head page is playing.
 			const bool bIsHeadPlaying = bInPreview ? InRundown->IsPagePreviewing(PlaybackContext->PlayHeadPageId) : InRundown->IsPagePlaying(PlaybackContext->PlayHeadPageId); 
-			
-			if (bIsHeadPlaying)
+
+			// Exception for Command only pages. They don't have a page player.
+			bool bIsCommandOnly = false;
+			if (!bIsHeadPlaying)
+			{
+				const FAvaRundownPage& Page = InRundown->GetPage(PlaybackContext->PlayHeadPageId);
+				bIsCommandOnly = Page.HasCommands(InRundown) && !Page.HasAssets(*InRundown); 
+			}
+
+			if (bIsHeadPlaying || bIsCommandOnly)
 			{
 				const int32 NextPageId = InRundown->GetNextPage(PlaybackContext->PlayHeadPageId, InPageListReference).GetPageId();
 				if (IsPageIdValid(NextPageId) && InRundown->CanPlayPage(NextPageId, bInPreview))
diff --git a/Engine/Plugins/Experimental/Avalanche/Source/AvalancheMedia/Private/Rundown/AvaRundownSerializationUtils.cpp b/Engine/Plugins/Experimental/Avalanche/Source/AvalancheMedia/Private/Rundown/AvaRundownSerializationUtils.cpp
--- a/Engine/Plugins/Experimental/Avalanche/Source/AvalancheMedia/Private/Rundown/AvaRundownSerializationUtils.cpp
+++ b/Engine/Plugins/Experimental/Avalanche/Source/AvalancheMedia/Private/Rundown/AvaRundownSerializationUtils.cpp
@@ -2,15 +2,25 @@
 
 #include "Rundown/AvaRundownSerializationUtils.h"
 
+#include "AvaMediaSerializationUtils.h"
 #include "Backends/JsonStructDeserializerBackend.h"
 #include "Backends/JsonStructSerializerBackend.h"
 #include "HAL/FileManager.h"
+#include "JsonObjectConverter.h"
 #include "Rundown/AvaRundown.h"
 #include "StructDeserializer.h"
 #include "StructSerializer.h"
+#include "StructUtils/InstancedStruct.h"
 
 #define LOCTEXT_NAMESPACE "AvaRundownSerializationUtils"
 
+// Keep the legacy serializer in case there are issues with the new one.
+static TAutoConsoleVariable<bool> CVarRundownNewJsonSerializer(
+	TEXT("MotionDesignRundown.UseNewJsonSerializer"),
+	true,
+	TEXT("Enable/Disable New Rundown Serializer. If disabled, will use the legacy serializer."),
+	ECVF_Default);
+
 namespace UE::AvaMedia::RundownSerializationUtils::Private
 {
 	// Don't serialize the transient properties.
@@ -35,6 +45,157 @@ namespace UE::AvaMedia::RundownSerializationUtils::Private
 			PropertyFilter = TransientPropertyFilter;
 		}
 	};
+
+	const TCHAR* GetJsonTypeName(EJson InType)
+	{
+		switch (InType)
+		{
+		case EJson::None: return TEXT("None");
+		case EJson::Null: return TEXT("Null");
+		case EJson::String: return TEXT("String");
+		case EJson::Number: return TEXT("Number");
+		case EJson::Boolean: return TEXT("Boolean");
+		case EJson::Array: return TEXT("Array");
+		case EJson::Object: return TEXT("Object");
+		default: return TEXT("Invalid");
+		}
+	}
+	
+	TSharedPtr<FJsonObject> SerializeToJson(const UStruct* InStruct, const void* InObject);
+
+	static TSharedPtr<FJsonValue> ExportPropertyCallback(FProperty* InProperty, const void* InValue)
+	{
+		TSharedPtr<FJsonValue> Result = nullptr;
+
+		if (const FStructProperty* StructProperty = CastField<FStructProperty>(InProperty))
+		{
+			if (StructProperty->Struct == FInstancedStruct::StaticStruct())
+			{
+				const FInstancedStruct& InstancedStruct = *static_cast<const FInstancedStruct*>(InValue);
+				if (const TSharedPtr<FJsonObject> LoadedStruct = SerializeToJson(InstancedStruct.GetScriptStruct(), InstancedStruct.GetMemory()))
+				{
+					const TSharedPtr<FJsonObject> InstancedStructObj = MakeShared<FJsonObject>();
+					const FSoftObjectPath StructPath(InstancedStruct.GetScriptStruct());
+					InstancedStructObj->SetField(StructPath.ToString(), MakeShared<FJsonValueObject>(LoadedStruct));
+					
+					Result = MakeShared<FJsonValueObject>(InstancedStructObj);
+				}
+			}
+		}
+		return Result;
+	}
+
+	TSharedPtr<FJsonObject> SerializeToJson(const UStruct* InStruct, const void* InObject)
+	{
+		FJsonObjectConverter::CustomExportCallback CustomCB;
+		CustomCB.BindStatic(&ExportPropertyCallback);
+
+		const TSharedPtr<FJsonObject> JsonObject = MakeShared<FJsonObject>();
+		constexpr int64 CheckFlags = 0;
+		constexpr int64 SkipFlags = 0;
+		const bool bLoaded = FJsonObjectConverter::UStructToJsonObject(InStruct, InObject, JsonObject.ToSharedRef(), CheckFlags, SkipFlags, &CustomCB, EJsonObjectConversionFlags::SkipStandardizeCase);
+		return bLoaded ? JsonObject : nullptr;
+	}
+
+	TSharedPtr<FJsonObject> SerializeToJson(const UObject* InObject)
+	{
+		return SerializeToJson(InObject->GetClass(), InObject);
+	}
+
+	bool DeserializeFromJson(const TSharedRef<FJsonObject>& InJsonObject, const UStruct* InStruct, void* InNativeObject, FText& OutErrorMessage);
+
+	// The deserialization error will get lost, so we log instead.
+	//
+	// Note: 
+	//	Returning false means the callback doesn't handle the value, so fallback code does.
+	//	However, if this is our case we handle, even if there is an error, we return true to indicate it is handled.
+	bool ImportPropertyCallback(const TSharedPtr<FJsonValue>& InJsonValue, FProperty* InProperty, void* InValue)
+	{
+		if (const FStructProperty* StructProperty = CastField<FStructProperty>(InProperty))
+		{
+			if (StructProperty->Struct == FInstancedStruct::StaticStruct())
+			{
+				constexpr bool bHandledButError = true;	// There is no way to indicate an error or return the error message.
+				
+				FInstancedStruct& InstancedStruct = *static_cast<FInstancedStruct*>(InValue);
+
+				if (InJsonValue->Type !=  EJson::Object)
+				{
+					UE_LOG(LogAvaRundown, Error, TEXT("Rundown Json Parsing FInstancedStruct Property: Json value should be of type Object, is of type: %s"), GetJsonTypeName(InJsonValue->Type));
+					return bHandledButError;
+				}
+				
+				// json value should be an object with 1 field
+				const TSharedPtr<FJsonObject>* InstancedStructObject;
+				if (InJsonValue->TryGetObject(InstancedStructObject))
+				{
+					if ((*InstancedStructObject)->Values.IsEmpty())
+					{
+						UE_LOG(LogAvaRundown, Error, TEXT("Rundown Json Parsing FInstancedStruct Property: Json Object should have at least 1 field (has %d)"), (*InstancedStructObject)->Values.Num());
+						return bHandledButError;
+					}
+
+					if ((*InstancedStructObject)->Values.Num() > 1)
+					{
+						UE_LOG(LogAvaRundown, Warning, TEXT("Rundown Json Parsing FInstancedStruct Property: Json Object should have only 1 field (has %d)"), (*InstancedStructObject)->Values.Num());
+					}
+						
+					for (const TPair<FString, TSharedPtr<FJsonValue>>& Value : (*InstancedStructObject)->Values)
+					{
+						if (!Value.Value.IsValid())
+						{
+							UE_LOG(LogAvaRundown, Error, TEXT("Rundown Json Parsing FInstancedStruct Property: Json value is invalid"));
+							continue;
+						}
+
+						FSoftObjectPath ScriptPath(Value.Key);
+						if (const UScriptStruct* Struct = Cast<UScriptStruct>(ScriptPath.ResolveObject()))	// so we got a valid struct
+						{
+							if (Value.Value->Type !=  EJson::Object)
+							{
+								UE_LOG(LogAvaRundown, Error, TEXT("Rundown Json Parsing FInstancedStruct Property: Json value should be of type Object, is of type: %s"), GetJsonTypeName(Value.Value->Type));
+							}
+							
+							const TSharedPtr<FJsonObject>* StructObject;
+							if (Value.Value->TryGetObject(StructObject))
+							{
+								InstancedStruct.InitializeAs(Struct);
+								FText DeserializationError;
+								if (DeserializeFromJson((*StructObject).ToSharedRef(), Struct, InstancedStruct.GetMutableMemory(), DeserializationError))
+								{
+									constexpr bool bHandledSuccess = true;
+									return bHandledSuccess;
+								}
+								
+								UE_LOG(LogAvaRundown, Error, TEXT("Rundown Json Parsing FInstancedStruct Property: Failed to deserialize: \"%s\""), *DeserializationError.ToString());
+							}
+						}
+						else
+						{
+							UE_LOG(LogAvaRundown, Error, TEXT("Rundown Json Parsing FInstancedStruct Property: Unknown Script Type: \"%s\""), *Value.Key);
+						}
+					}
+
+					return bHandledButError;
+				}
+			}
+		}
+		
+		return false;
+	}
+
+	bool DeserializeFromJson(const TSharedRef<FJsonObject>& InJsonObject, const UStruct* InStruct, void* InNativeObject, FText& OutErrorMessage)
+	{
+		FJsonObjectConverter::CustomImportCallback CustomCB;
+		CustomCB.BindStatic(&ImportPropertyCallback);
+
+		constexpr int64 CheckFlags = 0;
+		constexpr int64 SkipFlags = 0;
+		constexpr bool bStrictMode = false;
+		return FJsonObjectConverter::JsonObjectToUStruct(InJsonObject, InStruct, InNativeObject, CheckFlags, SkipFlags, bStrictMode, &OutErrorMessage, &CustomCB);
+	}
+
+	
 }
 
 namespace UE::AvaMedia::RundownSerializationUtils
@@ -47,9 +208,27 @@ namespace UE::AvaMedia::RundownSerializationUtils
 			return false;
 		}
 
-		// Remark: this is hardcoded to encode in utf16-le.
-		FJsonStructSerializerBackend Backend(InArchive, EStructSerializerBackendFlags::Default);
-		FStructSerializer::Serialize(InRundown, *InRundown->GetClass(), Backend, Private::FRundownSerializerPolicies());
+		if(CVarRundownNewJsonSerializer.GetValueOnGameThread())
+		{
+			const TSharedPtr<FJsonObject> JsonRundown = Private::SerializeToJson(InRundown);
+
+			const TSharedPtr<FJsonObject> JsonRoot = MakeShared<FJsonObject>();
+			JsonRoot->SetStringField(TEXT("RundownFileVersion"), TEXT("1.0"));
+			JsonRoot->SetObjectField(TEXT("Rundown"), JsonRundown);
+
+			const TSharedRef<TJsonWriter<>> Writer = TJsonWriterFactory<>::Create(&InArchive);
+			if (!FJsonSerializer::Serialize(JsonRoot.ToSharedRef(), Writer))
+			{
+				return false;
+			}
+		}
+		else
+		{
+			// Legacy serializer
+			// Remark: this is hardcoded to encode in utf16-le.
+			FJsonStructSerializerBackend Backend(InArchive, EStructSerializerBackendFlags::Default);
+			FStructSerializer::Serialize(InRundown, *InRundown->GetClass(), Backend, Private::FRundownSerializerPolicies());
+		}
 		return true;
 	}
 
@@ -90,9 +269,39 @@ namespace UE::AvaMedia::RundownSerializationUtils
 			}
 			return false;
 		}
+
+		const int64 ArchivePosition = InArchive.Tell();
+		
+		TSharedPtr<FJsonObject> JsonRoot = MakeShared<FJsonObject>();
+		const TSharedRef<TJsonReader<TCHAR>> JsonReader = TJsonReaderFactory<TCHAR>::Create(&InArchive);
+		if (!FJsonSerializer::Deserialize(JsonReader, JsonRoot))
+		{
+			OutErrorMessage = LOCTEXT("LoadRundownJson_FailedParseJson", "Couldn't parse json data.");
+			return false;
+		}
+
+		bool bTryLoadLegacy = true;
+		bool bLoaded = false;
+		FText DeserializationError;
 		
-		FJsonStructDeserializerBackend Backend(InArchive);
-		const bool bLoaded = FStructDeserializer::Deserialize(InRundown, *InRundown->GetClass(), Backend, Private::FRundownDeserializerPolicies());
+		// If we can find the rundown field, we got the new version.
+		if (const TSharedPtr<FJsonObject> RundownJson = JsonRoot->GetObjectField(TEXT("Rundown")))
+		{
+			bTryLoadLegacy = false;
+			bLoaded = Private::DeserializeFromJson(RundownJson.ToSharedRef(), InRundown->GetClass(), InRundown, DeserializationError);
+		}
+
+		if (bTryLoadLegacy)
+		{
+			InArchive.Seek(ArchivePosition); // Return to start.
+			FJsonStructDeserializerBackend Backend(InArchive);
+			bLoaded = FStructDeserializer::Deserialize(InRundown, *InRundown->GetClass(), Backend, Private::FRundownDeserializerPolicies());
+			if (!bLoaded)
+			{
+				DeserializationError = FText::FromString(Backend.GetLastErrorMessage()); 
+			}
+		}
+			
 		if (bLoaded)
 		{
 			InRundown->PostLoad();
@@ -101,8 +310,7 @@ namespace UE::AvaMedia::RundownSerializationUtils
 		else
 		{
 			OutErrorMessage =
-				FText::Format(LOCTEXT("LoadRundownJson_DeserializerError", "Json Deserializer error: {0}"),
-					FText::FromString(Backend.GetLastErrorMessage()));
+				FText::Format(LOCTEXT("LoadRundownJson_DeserializerError", "Json Deserializer error: {0}"), DeserializationError);
 		}
 		return bLoaded;
 	}
@@ -120,6 +328,24 @@ namespace UE::AvaMedia::RundownSerializationUtils
 		OutErrorMessage = LOCTEXT("FileNotFound", "File not found");
 		return false;
 	}
+
+	bool SerializeRundownPageCommandToJsonString(const FAvaRundownPageCommand* InRundownPageCommand, UStruct& InStruct, FString& OutString)
+	{
+		TArray<uint8> ValueAsBytes;
+		FMemoryWriter Writer = FMemoryWriter(ValueAsBytes);
+		FJsonStructSerializerBackend Backend(Writer, EStructSerializerBackendFlags::Default);
+		FStructSerializer::Serialize(InRundownPageCommand, InStruct, Backend, Private::FRundownSerializerPolicies());
+		Writer.Close();
+		AvaMediaSerializationUtils::JsonValueConversion::BytesToString(ValueAsBytes, OutString);
+		return true;
+	}
+	
+	bool DeserializeRundownPageCommandFromJson(FAvaRundownPageCommand* InRundownPageCommand, UStruct& InStruct, const FString& InString)
+	{
+		FMemoryReaderView Reader(AvaMediaSerializationUtils::JsonValueConversion::ValueToConstBytesView(InString));
+		FJsonStructDeserializerBackend Backend(Reader);
+		return FStructDeserializer::Deserialize(InRundownPageCommand, InStruct, Backend, Private::FRundownDeserializerPolicies());
+	}
 }
 
 #undef LOCTEXT_NAMESPACE
diff --git a/Engine/Plugins/Experimental/Avalanche/Source/AvalancheMedia/Private/Rundown/AvaRundownServer.cpp b/Engine/Plugins/Experimental/Avalanche/Source/AvalancheMedia/Private/Rundown/AvaRundownServer.cpp
--- a/Engine/Plugins/Experimental/Avalanche/Source/AvalancheMedia/Private/Rundown/AvaRundownServer.cpp
+++ b/Engine/Plugins/Experimental/Avalanche/Source/AvalancheMedia/Private/Rundown/AvaRundownServer.cpp
@@ -170,6 +170,7 @@ namespace UE::AvaRundownServer::Private
 		PageInfo.Statuses = InPage.GetPageStatuses(InRundown);
 		PageInfo.TransitionLayerName = InPage.GetTransitionLayer(InRundown).ToString(); 	// Todo: combo templates
 		PageInfo.bTransitionLogicEnabled = InPage.HasTransitionLogic(InRundown);
+		PageInfo.Commands = InPage.SaveInstancedCommands();
 		PageInfo.OutputChannel = InPage.GetChannelName().ToString();
 		PageInfo.bIsEnabled = InPage.IsEnabled();
 		PageInfo.bIsPlaying = InRundown->IsPagePlaying(InPage);
diff --git a/Engine/Plugins/Experimental/Avalanche/Source/AvalancheMedia/Private/Rundown/Commands/AvaRundownPageCommandStopLayers.cpp b/Engine/Plugins/Experimental/Avalanche/Source/AvalancheMedia/Private/Rundown/Commands/AvaRundownPageCommandStopLayers.cpp
new file mode 100644
--- /dev/null
+++ b/Engine/Plugins/Experimental/Avalanche/Source/AvalancheMedia/Private/Rundown/Commands/AvaRundownPageCommandStopLayers.cpp
@@ -0,0 +1,127 @@
+// Copyright Epic Games, Inc. All Rights Reserved.
+
+#include "AvaRundownPageCommandStopLayers.h"
+
+#include "AvaTagCollection.h"
+#include "Rundown/AvaRundown.h"
+#include "Rundown/AvaRundownPage.h"
+#include "Rundown/AvaRundownPagePlayer.h"
+#include "Rundown/AvaRundownPlaybackUtils.h"
+#include "Rundown/Transition/AvaRundownPageTransition.h"
+#include "Rundown/Transition/AvaRundownPageTransitionBuilder.h"
+
+namespace UE::AvaMedia::RundownPageCommandStopLayers::Private
+{
+	// todo: might be worth having this in FAvaTagHandleContainer. 
+	bool IsEmpty(const FAvaTagHandleContainer& InLayers)
+	{
+		return InLayers.GetTagIds().IsEmpty() || InLayers.Source == nullptr;	
+	}
+
+	FString GetLayersString(const FAvaTagHandleContainer& InLayers, const FString& InLayerPrefix, const FString& InSeparator)
+	{
+		FString LayersString;
+		LayersString.Reserve(64);
+		for (const FAvaTagId& TagId : InLayers.GetTagIds())
+		{
+			if (!LayersString.IsEmpty())
+			{
+				LayersString += InSeparator;
+			}
+			LayersString += InLayerPrefix;
+			LayersString += InLayers.Source->GetTagName(TagId).ToString();		
+		}
+		return LayersString;
+	}
+
+	bool Overlaps(const FAvaTagHandleContainer& InLayers, const FAvaTagHandle& InOther)
+	{
+		for (const FAvaTagId& TagId : InLayers.GetTagIds())
+		{
+			if (FAvaTagHandle(InLayers.Source, TagId).Overlaps(InOther))
+			{
+				return true;
+			}
+		}
+		return false;
+	}
+}
+
+FString FAvaRundownPageCommandStopLayers::GetDescription() const
+{
+	using namespace UE::AvaMedia::RundownPageCommandStopLayers::Private;
+	return FString::Printf(TEXT("Stop Layers: %s"), *GetLayersString(Layers, TEXT(""), TEXT(", ")));
+}
+
+bool FAvaRundownPageCommandStopLayers::HasTransitionLogic() const
+{
+	using namespace UE::AvaMedia::RundownPageCommandStopLayers::Private;
+	return !IsEmpty(Layers);
+}
+
+FString FAvaRundownPageCommandStopLayers::GetTransitionLayerString(const FString& InSeparator) const
+{
+	// For the layers text of a stop layer command, we'll add a '-' sign in front of the layer name.
+	using namespace UE::AvaMedia::RundownPageCommandStopLayers::Private;
+	return GetLayersString(Layers, TEXT("-"), InSeparator);
+}
+
+bool FAvaRundownPageCommandStopLayers::CanExecute(const UAvaRundown& InRundown, const FAvaRundownPage& InPage, FName InChannelName, FString* OutFailureReason) const
+{
+	using namespace UE::AvaMedia::RundownPageCommandStopLayers::Private;
+	if (IsEmpty(Layers))
+	{
+		if (OutFailureReason)
+		{
+			*OutFailureReason = TEXT("Stop Layers: no layers defined in command.");
+		}
+		return false;
+	}
+
+	// Check if there are currently playing pages that have layers than can be stopped.
+	for (const UAvaRundownPagePlayer* PagePlayer : InRundown.GetPagePlayers())
+	{
+		if (!PagePlayer || PagePlayer->ChannelFName != InChannelName)	// Not this channel
+		{
+			continue;
+		}
+
+		bool bLayerOverlap = false;
+		PagePlayer->ForEachInstancePlayer([this, &bLayerOverlap](const UAvaRundownPlaybackInstancePlayer* InInstancePlayer)
+		{
+			using namespace UE::AvaMedia::RundownPageCommandStopLayers::Private;
+			bLayerOverlap |= Overlaps(Layers, InInstancePlayer->TransitionLayer);
+		});
+
+		// Found an instance of a layer that we can kick out.
+		if (bLayerOverlap)
+		{
+			return true;
+		}
+	}
+
+	if (OutFailureReason)
+	{
+		*OutFailureReason = TEXT("Stop Layers: no currently playing pages overlaps with defined layers.");
+	}
+	return false;
+}
+
+bool FAvaRundownPageCommandStopLayers::Execute(FAvaRundownPageTransitionBuilder& InTransitionBuilder, const UAvaRundown& InRundown, const FAvaRundownPage& InPage, FName InChannelName) const
+{
+	using namespace UE::AvaMedia::RundownPageCommandStopLayers::Private;
+	if (IsEmpty(Layers))
+	{
+		return false;
+	}
+
+	if (UAvaRundownPageTransition* Transition = InTransitionBuilder.FindOrAddTransition(InChannelName))
+	{
+		for (const FAvaTagId& TagId : Layers.GetTagIds())
+		{
+			FAvaRundownPlaybackUtils::AddTagHandleUnique(Transition->ExitLayers, FAvaTagHandle(Layers.Source, TagId));
+		}
+		return true;
+	}
+	return false;
+}
diff --git a/Engine/Plugins/Experimental/Avalanche/Source/AvalancheMedia/Private/Rundown/Commands/AvaRundownPageCommandStopLayers.h b/Engine/Plugins/Experimental/Avalanche/Source/AvalancheMedia/Private/Rundown/Commands/AvaRundownPageCommandStopLayers.h
new file mode 100644
--- /dev/null
+++ b/Engine/Plugins/Experimental/Avalanche/Source/AvalancheMedia/Private/Rundown/Commands/AvaRundownPageCommandStopLayers.h
@@ -0,0 +1,26 @@
+// Copyright Epic Games, Inc. All Rights Reserved.
+
+#pragma once
+
+#include "AvaTagHandleContainer.h"
+#include "Rundown/AvaRundownPageCommand.h"
+
+#include "AvaRundownPageCommandStopLayers.generated.h"
+
+USTRUCT(BlueprintType, meta = (DisplayName = "Stop Layers"))
+struct FAvaRundownPageCommandStopLayers : public FAvaRundownPageCommand
+{
+	GENERATED_BODY()
+
+	//~ Begin FAvaRundownPageCommand
+	virtual FString GetDescription() const override;
+	virtual bool HasTransitionLogic() const override;
+	virtual FString GetTransitionLayerString(const FString& InSeparator) const override;
+	virtual bool CanExecute(const UAvaRundown& InRundown, const FAvaRundownPage& InPage, FName InChannelName, FString* OutFailureReason) const override;
+	virtual bool Execute(FAvaRundownPageTransitionBuilder& InTransitionBuilder, const UAvaRundown& InRundown, const FAvaRundownPage& InPage, FName InChannelName) const override;
+	//~ End FAvaRundownPageCommand
+	
+	/** Layers to stop. */
+	UPROPERTY(EditAnywhere, BlueprintReadWrite, Category = "Page Data")
+	FAvaTagHandleContainer Layers;
+};
diff --git a/Engine/Plugins/Experimental/Avalanche/Source/AvalancheMedia/Private/Rundown/Transition/AvaRundownPageTransition.cpp b/Engine/Plugins/Experimental/Avalanche/Source/AvalancheMedia/Private/Rundown/Transition/AvaRundownPageTransition.cpp
--- a/Engine/Plugins/Experimental/Avalanche/Source/AvalancheMedia/Private/Rundown/Transition/AvaRundownPageTransition.cpp
+++ b/Engine/Plugins/Experimental/Avalanche/Source/AvalancheMedia/Private/Rundown/Transition/AvaRundownPageTransition.cpp
@@ -519,6 +519,22 @@ void UAvaRundownPageTransition::AddPlayablesToBuilder(FAvaPlayableTransitionBuil
 	
 		if (UAvaPlayable* Playable = GetPlayable(InstancePlayer))
 		{
+			// Begin -- Special Exit Layers Logic -- 
+			if (EntryRole == EAvaPlayableTransitionEntryRole::Playing)
+			{
+				// Kick out playing instances that overlap with exit layers.
+				for (const FAvaTagHandle& ExitLayer : ExitLayers)
+				{
+					// Note: relies on the instance player having the correct transition layer (from the page).
+					if (ExitLayer.Overlaps(InstancePlayer->TransitionLayer))
+					{
+						EntryRole = EAvaPlayableTransitionEntryRole::Exit;
+						break;
+					}
+				}
+			}
+			// End -- Special Exit Layers Logic --
+
 			const bool bPlayableAdded = InOutBuilder.AddPlayable(Playable, EntryRole);
 			if (EntryRole == EAvaPlayableTransitionEntryRole::Enter && bPlayableAdded)
 			{
diff --git a/Engine/Plugins/Experimental/Avalanche/Source/AvalancheMedia/Private/Rundown/Transition/AvaRundownPageTransition.h b/Engine/Plugins/Experimental/Avalanche/Source/AvalancheMedia/Private/Rundown/Transition/AvaRundownPageTransition.h
--- a/Engine/Plugins/Experimental/Avalanche/Source/AvalancheMedia/Private/Rundown/Transition/AvaRundownPageTransition.h
+++ b/Engine/Plugins/Experimental/Avalanche/Source/AvalancheMedia/Private/Rundown/Transition/AvaRundownPageTransition.h
@@ -2,7 +2,7 @@
 
 #pragma once
 
-#include "AvaTagId.h"
+#include "AvaTagHandle.h"
 #include "Playable/AvaPlayable.h"
 #include "Playback/Transition/AvaPlaybackTransition.h"
 #include "UObject/Object.h"
@@ -66,6 +66,8 @@ public:
 
 	UAvaRundown* GetRundown() const;
 
+	FString GetBriefTransitionDescription() const;
+
 	/** Instances will not be added to the playable transition. */
 	UPROPERTY(Transient)
 	TSet<FGuid> InstancesBypassingTransition;
@@ -74,6 +76,9 @@ public:
 	UPROPERTY(Transient)
 	TSet<FGuid> ReusedInstances;
 
+	/** Layers to kick out as part of this transition command. */
+	TArray<FAvaTagHandle> ExitLayers;
+
 	/**
 	 * Special mark if any of the enter playables where issued with a PreviewFrame play type.
 	 * Current logic will mark the whole transition.
@@ -94,7 +99,6 @@ protected:
 
 	FString GetInstanceName() const;
 	void LogDetailedTransitionInfo() const;
-	FString GetBriefTransitionDescription() const;
 
 	void RegisterToPlayableTransitionEvent();
 	void UnregisterFromPlayableTransitionEvent() const;
@@ -120,4 +124,6 @@ protected:
 	
 	UPROPERTY(Transient)
 	TObjectPtr<UAvaPlayableTransition> PlayableTransition;
+
+	friend class FAvaRundownPageTransitionBuilder;
 };
\ No newline at end of file
diff --git a/Engine/Plugins/Experimental/Avalanche/Source/AvalancheMedia/Private/Rundown/Transition/AvaRundownPageTransitionBuilder.cpp b/Engine/Plugins/Experimental/Avalanche/Source/AvalancheMedia/Private/Rundown/Transition/AvaRundownPageTransitionBuilder.cpp
--- a/Engine/Plugins/Experimental/Avalanche/Source/AvalancheMedia/Private/Rundown/Transition/AvaRundownPageTransitionBuilder.cpp
+++ b/Engine/Plugins/Experimental/Avalanche/Source/AvalancheMedia/Private/Rundown/Transition/AvaRundownPageTransitionBuilder.cpp
@@ -40,12 +40,11 @@ FAvaRundownPageTransitionBuilder::~FAvaRundownPageTransitionBuilder()
 	}
 }
 
-UAvaRundownPageTransition* FAvaRundownPageTransitionBuilder::FindTransition(const UAvaRundownPagePlayer* InPlayer) const
+UAvaRundownPageTransition* FAvaRundownPageTransitionBuilder::FindTransition(FName InChannelName) const
 {
-	// Currently, the only batching criteria is the channel.
 	for (UAvaRundownPageTransition* PageTransition : PageTransitions)
 	{
-		if (PageTransition->GetChannelName() == InPlayer->ChannelFName)
+		if (PageTransition->GetChannelName() == InChannelName)
 		{
 			return PageTransition;
 		}
@@ -53,17 +52,31 @@ UAvaRundownPageTransition* FAvaRundownPageTransitionBuilder::FindTransition(cons
 	return nullptr;	
 }
 
+UAvaRundownPageTransition* FAvaRundownPageTransitionBuilder::FindTransition(const UAvaRundownPagePlayer* InPlayer) const
+{
+	// Currently, the only batching criteria is the channel.
+	return InPlayer ? FindTransition(InPlayer->ChannelFName) : nullptr;
+}
+
+UAvaRundownPageTransition* FAvaRundownPageTransitionBuilder::FindOrAddTransition(FName InChannelName)
+{
+	UAvaRundownPageTransition* PageTransition = FindTransition(InChannelName);
+	return PageTransition ? PageTransition : AddTransition(InChannelName);
+}
+
 UAvaRundownPageTransition* FAvaRundownPageTransitionBuilder::FindOrAddTransition(const UAvaRundownPagePlayer* InPlayer)
 {
-	if (UAvaRundownPageTransition* PageTransition = FindTransition(InPlayer))
-	{
-		return PageTransition;
-	}
-		
+	UAvaRundownPageTransition* PageTransition = FindTransition(InPlayer);
+	return PageTransition ? PageTransition : AddTransition(InPlayer->ChannelFName);
+}
+
+UAvaRundownPageTransition* FAvaRundownPageTransitionBuilder::AddTransition(FName InChannelName)
+{
 	UAvaRundownPageTransition* PageTransition = UAvaRundownPageTransition::MakeNew(Rundown);
 	if (PageTransition)
 	{
+		PageTransition->ChannelName = InChannelName;
 		PageTransitions.Add(PageTransition);
 	}
-	return PageTransition;	
+	return PageTransition;
 }
\ No newline at end of file
diff --git a/Engine/Plugins/Experimental/Avalanche/Source/AvalancheMedia/Private/Rundown/Transition/AvaRundownPageTransitionBuilder.h b/Engine/Plugins/Experimental/Avalanche/Source/AvalancheMedia/Private/Rundown/Transition/AvaRundownPageTransitionBuilder.h
--- a/Engine/Plugins/Experimental/Avalanche/Source/AvalancheMedia/Private/Rundown/Transition/AvaRundownPageTransitionBuilder.h
+++ b/Engine/Plugins/Experimental/Avalanche/Source/AvalancheMedia/Private/Rundown/Transition/AvaRundownPageTransitionBuilder.h
@@ -37,6 +37,13 @@ public:
 
 	~FAvaRundownPageTransitionBuilder();
 
+	/**
+	 * @brief Finds a suitable page transition object for the channel.
+	 * @param InChannel channel we are requesting the transition object for.
+	 * @return the requested page transition object or null if not found.
+	 */
+	UAvaRundownPageTransition* FindTransition(FName InChannelName) const;
+
 	/**
 	 * @brief Finds a suitable page transition object for the given page player.
 	 * @param InPlayer Page player we are requesting the transition object for.
@@ -44,6 +51,15 @@ public:
 	 */
 	UAvaRundownPageTransition* FindTransition(const UAvaRundownPagePlayer* InPlayer) const;
 
+	/**
+	 * @brief Retrieves a suitable page transition object for the given channel.
+	 * @param InChannel Channel we are requesting the transition object for.
+	 * @return the requested page transition object.
+	 *
+	 * This is where the batching occurs according to criteria.
+	 */
+	UAvaRundownPageTransition* FindOrAddTransition(FName InChannelName);
+	
 	/**
 	 * @brief Retrieves a suitable page transition object for the given page player.
 	 * @param InPlayer Page player we are requesting the transition object for.
@@ -53,6 +69,9 @@ public:
 	 */
 	UAvaRundownPageTransition* FindOrAddTransition(const UAvaRundownPagePlayer* InPlayer);
 
+private:
+	UAvaRundownPageTransition* AddTransition(FName InChannelName);
+	
 private:
 	UAvaRundown* Rundown;
 
diff --git a/Engine/Plugins/Experimental/Avalanche/Source/AvalancheMedia/Public/AvaMediaDefines.h b/Engine/Plugins/Experimental/Avalanche/Source/AvalancheMedia/Public/AvaMediaDefines.h
--- a/Engine/Plugins/Experimental/Avalanche/Source/AvalancheMedia/Public/AvaMediaDefines.h
+++ b/Engine/Plugins/Experimental/Avalanche/Source/AvalancheMedia/Public/AvaMediaDefines.h
@@ -254,6 +254,7 @@ enum class EAvaRundownPageChanges : uint8
 	Status					= 1 << 4,
 	Name					= 1 << 5,
 	FriendlyName			= 1 << 6,
+	Commands				= 1 << 7,
 
 	All              = 0xFF,
 };
diff --git a/Engine/Plugins/Experimental/Avalanche/Source/AvalancheMedia/Public/Rundown/AvaRundown.h b/Engine/Plugins/Experimental/Avalanche/Source/AvalancheMedia/Public/Rundown/AvaRundown.h
--- a/Engine/Plugins/Experimental/Avalanche/Source/AvalancheMedia/Public/Rundown/AvaRundown.h
+++ b/Engine/Plugins/Experimental/Avalanche/Source/AvalancheMedia/Public/Rundown/AvaRundown.h
@@ -532,7 +532,7 @@ public:
 	bool PlayPage(int32 InPageId, EAvaRundownPagePlayType InPlayType, const FName& InPreviewChannelName) { return !PlayPages({InPageId}, InPlayType, InPreviewChannelName).IsEmpty(); }
 
 	bool CanPlayPage(int32 InPageId, bool bInPreview) const;
-	bool CanPlayPage(int32 InPageId, bool bInPreview, const FName& InPreviewChannelName) const;
+	bool CanPlayPage(int32 InPageId, bool bInPreview, const FName& InPreviewChannelName, FString* OutFailureReason) const;
 
 	TArray<int32> StopPages(const TArray<int32>& InPageIds, EAvaRundownPageStopOptions InOptions, bool bInPreview);
 	TArray<int32> StopPages(const TArray<int32>& InPageIds, EAvaRundownPageStopOptions InOptions, bool bInPreview, const FName& InPreviewChannelName);
@@ -541,8 +541,19 @@ public:
 	bool StopPage(int32 InPageId, EAvaRundownPageStopOptions InOptions, bool bInPreview, const FName& InPreviewChannelName) { return !StopPages({InPageId}, InOptions, bInPreview, InPreviewChannelName).IsEmpty(); }
 	
 	bool CanStopPage(int32 InPageId, EAvaRundownPageStopOptions InOptions, bool bInPreview) const;
-	bool CanStopPage(int32 InPageId, EAvaRundownPageStopOptions InOptions, bool bInPreview, const FName& InPreviewChannelName) const;
+	bool CanStopPage(int32 InPageId, EAvaRundownPageStopOptions InOptions, bool bInPreview, const FName& InPreviewChannelName, FString* OutFailureReason) const;
 
+	/**
+	 * Stop the layers of a given channel.
+	 * If using transition, it will produce a special transition that has only kick out layers.
+	 *
+	 * returns the pages that will be affected (totally or partially stopped).
+	 */
+	TArray<int32> StopLayers(FName InChannelName, const TArray<FAvaTagHandle>& InLayers, EAvaRundownPageStopOptions InOptions);
+
+	/** Determines if a given layer has something playing that can be stopped. */
+	bool CanStopLayer(FName InChannelName, const FAvaTagHandle& InLayer) const;
+	
 	bool StopChannel(const FString& InChannelName);
 	bool CanStopChannel(const FString& InChannelName) const;
 
@@ -661,7 +672,7 @@ protected:
 	 * Returns true if the selected page can play on the given channel.
 	 * This enforces that the channel type (program or preview) is compatible with the requested operation.
 	 */
-	bool IsChannelTypeCompatibleForRequest(const FAvaRundownPage& InSelectedPage, bool bInIsPreview, const FName& InPreviewChannelName, bool bInLogFailureReason) const;
+	bool IsChannelTypeCompatibleForRequest(const FAvaRundownPage& InSelectedPage, bool bInIsPreview, const FName& InPreviewChannelName, FString* OutFailureReason) const;
 
 private:
 	IAvaRundownPageLoadingManager& MakePageLoadingManager();
@@ -764,7 +775,7 @@ public:
 
 	UAvaRundownPageTransition* GetPageTransition(const FGuid& InTransitionId) const;
 
-	bool CanStartTransitionForPage(const FAvaRundownPage& InPage, bool bInIsPreview, const FName& InPreviewChannelName) const;
+	bool CanStartTransitionForPage(const FAvaRundownPage& InPage, bool bInIsPreview, const FName& InPreviewChannelName, FString* OutFailureReason) const;
 
 	void StopPageTransitionsForPage(const FAvaRundownPage& InPage, bool bInIsPreview, const FName& InPreviewChannelName);
 
diff --git a/Engine/Plugins/Experimental/Avalanche/Source/AvalancheMedia/Public/Rundown/AvaRundownPage.h b/Engine/Plugins/Experimental/Avalanche/Source/AvalancheMedia/Public/Rundown/AvaRundownPage.h
--- a/Engine/Plugins/Experimental/Avalanche/Source/AvalancheMedia/Public/Rundown/AvaRundownPage.h
+++ b/Engine/Plugins/Experimental/Avalanche/Source/AvalancheMedia/Public/Rundown/AvaRundownPage.h
@@ -6,10 +6,13 @@
 #include "AvaTagHandle.h"
 #include "AvaTransitionEnums.h"
 #include "Playable/AvaPlayableRemoteControlValues.h"
+#include "StructUtils/InstancedStruct.h"
+
 #include "AvaRundownPage.generated.h"
 
 class UAvaRundown;
 class URemoteControlPreset;
+struct FAvaRundownPageCommand;
 
 UENUM()
 enum class EAvaRundownPageStatus : uint8
@@ -53,6 +56,21 @@ struct FAvaRundownChannelPageStatus
 	bool bNeedsSync = false;
 };
 
+/** Page command data is stored in json serialized string in the page to be compatible with external apps. */
+USTRUCT()
+struct FAvaRundownPageCommandData
+{
+	GENERATED_BODY()
+
+	/** Command name: string. */
+	UPROPERTY(VisibleAnywhere, Category = "Motion Design")
+	FString Name;
+
+	/** Command payload: json formatted string. */
+	UPROPERTY(VisibleAnywhere, Category = "Motion Design")
+	FString Payload;
+};
+
 USTRUCT(BlueprintType, DisplayName = "Motion Design Rundown Page")
 struct AVALANCHEMEDIA_API FAvaRundownPage
 {
@@ -109,9 +127,25 @@ public:
 
 	FText GetPageDescription() const;
 
+	/**
+	 * @brief Update the page summary (only if summary is empty)
+	 * @param InRundown Owning rundown.
+	 * @return true if the summary was generated, false otherwise (either because of an error or it was not needed).
+	 */
 	bool UpdatePageSummary(const UAvaRundown* InRundown);
+
+	UE_DEPRECATED(5.6, "Use the version that provides the rundown." )
 	bool UpdatePageSummary(const TArray<const URemoteControlPreset*>& InPresets, bool bInIsPresetChanged = false);
 
+	/**
+	 * @brief Update the page summary from controller values or commands.
+	 * @param InRundown Owning rundown. Necessary to resolve templates.
+	 * @param InPresets Input presets to get the value from.
+	 * @param bInForceUpdate Forces the update of the summary even if there is already one. If false, summary is generated only if empty. 
+	 * @return true if the summary was generated, false otherwise (either because of an error or it was not needed).
+	 */
+	bool UpdatePageSummary(const UAvaRundown* InRundown, const TArray<const URemoteControlPreset*>& InPresets, bool bInForceUpdate = false);
+
 	bool UpdateTransitionLogic();
 	bool HasTransitionLogic(const UAvaRundown* InRundown) const;
 	
@@ -172,8 +206,14 @@ public:
 	TArray<FSoftObjectPath> GetAssetPaths(const UAvaRundown* InRundown) const;
 	FSoftObjectPath GetAssetPathDirect() const { return AssetPath;}
 
+	/** Returns true if the page has some valid assets. */
+	bool HasAssets(const UAvaRundown& InRundown) const;
+
 	bool UpdateAsset(const FSoftObjectPath& InAssetPath, bool bInReimportPage = false);
 
+	/** Returns true if the page (or it's template(s)) has commands. */
+	bool HasCommands(const UAvaRundown* InRundown) const;
+
 	FName GetChannelName() const;
 	int32 GetChannelIndex() const { return OutputChannel; }
 	void SetChannelName(FName InChannelName);
@@ -235,9 +275,43 @@ public:
 	 * @return true if the templates are matching (ignores ids).
 	 */
 	bool IsTemplateMatchingByValue(const FAvaRundownPage& InTemplatePage) const;
-	
-protected:
 
+	/**
+	 * Load the instanced commands from the given command data (from web api).
+	 * @todo Needed until all json serializer can support FInstancedStruct.
+	 */
+	void LoadInstancedCommands(const TArray<FAvaRundownPageCommandData>& InCommands);
+
+	/**
+	 * Save the instanced commands to command data (suitable for web api).
+	 * @todo Needed until all json serializer can support FInstancedStruct.
+	 */
+	TArray<FAvaRundownPageCommandData> SaveInstancedCommands() const;
+
+	/**
+	 * @brief Get the instanced page commands.
+	 * @return array of instanced page commands that can be executed.
+	 */
+	const TArray<FInstancedStruct>& GetInstancedCommands() const
+	{
+		return InstancedCommands;
+	}
+
+	/**
+	 * @brief Set the instanced commands.
+	 * @param InInstancedCommands new instance commands to copy in the page.
+	 */
+	void SetInstancedCommands(const TArray<FInstancedStruct>& InInstancedCommands);
+
+	/**
+	 * @brief Calls the given function for each commands in the page.
+	 * @param InFunction Function to call for each command. The Page parameter is the direct page the command is on, i.e. may be the template.
+	 * @param InRundown Owning rundown (used to resolve templates)
+	 * @param bDirectOnly If true, only the current page's command is traversed, otherwise, the templates commands are included as well.
+	 */
+	void ForEachInstancedCommands(TFunctionRef<void(const FAvaRundownPageCommand&, const FAvaRundownPage&)> InFunction, const UAvaRundown* InRundown, bool bDirectOnly) const;
+
+protected:
 	UPROPERTY(VisibleAnywhere, BlueprintReadOnly, Category = "Motion Design")
 	bool bEnabled = true;
 
@@ -289,4 +363,10 @@ protected:
 
 	UPROPERTY(VisibleAnywhere, BlueprintReadOnly, Category = "Motion Design")
 	EAvaTransitionInstancingMode TransitionMode = EAvaTransitionInstancingMode::New;
+
+	/**
+	 * Page command that can be executed when playing this page.
+	 */
+	UPROPERTY(EditAnywhere, BlueprintReadWrite, Category = "Motion Design", meta = (ExcludeBaseStruct, BaseStruct = "/Script/AvalancheMedia.AvaRundownPageCommand"))
+	TArray<FInstancedStruct> InstancedCommands;
 };
diff --git a/Engine/Plugins/Experimental/Avalanche/Source/AvalancheMedia/Public/Rundown/AvaRundownPageCommand.h b/Engine/Plugins/Experimental/Avalanche/Source/AvalancheMedia/Public/Rundown/AvaRundownPageCommand.h
new file mode 100644
--- /dev/null
+++ b/Engine/Plugins/Experimental/Avalanche/Source/AvalancheMedia/Public/Rundown/AvaRundownPageCommand.h
@@ -0,0 +1,35 @@
+// Copyright Epic Games, Inc. All Rights Reserved.
+
+#pragma once
+
+#include "Containers/UnrealString.h"
+
+#include "AvaRundownPageCommand.generated.h"
+
+class FAvaRundownPageTransitionBuilder;
+class UAvaRundown;
+struct FAvaRundownPage;
+
+/** Base class for page commands */
+USTRUCT(BlueprintType)
+struct FAvaRundownPageCommand
+{
+	GENERATED_BODY()
+
+	virtual ~FAvaRundownPageCommand() = default;
+
+	/** Returns this command's description. Used for page summary. */
+	virtual FString GetDescription() const { return FString();}
+	
+	/** Returns true if the command will act as transition logic. */
+	virtual bool HasTransitionLogic() const { return false; }
+
+	/** Returns a string representing the transition layers. This is used for the page list "layer" column. */
+	virtual FString GetTransitionLayerString(const FString& InSeparator) const { return FString(); }
+	
+	/** Returns true if the command can be executed. Used by UI for button status. */
+	virtual bool CanExecute(const UAvaRundown& InRundown, const FAvaRundownPage& InPage, FName InChannelName, FString* OutFailureReason) const { return false; }
+
+	/** Execute the command. */
+	virtual bool Execute(FAvaRundownPageTransitionBuilder& InTransitionBuilder, const UAvaRundown& InRundown, const FAvaRundownPage& InPage, FName InChannelName) const { return false; }
+};
diff --git a/Engine/Plugins/Experimental/Avalanche/Source/AvalancheMediaEditor/Private/Rundown/AvaRundownEditor.cpp b/Engine/Plugins/Experimental/Avalanche/Source/AvalancheMediaEditor/Private/Rundown/AvaRundownEditor.cpp
--- a/Engine/Plugins/Experimental/Avalanche/Source/AvalancheMediaEditor/Private/Rundown/AvaRundownEditor.cpp
+++ b/Engine/Plugins/Experimental/Avalanche/Source/AvalancheMediaEditor/Private/Rundown/AvaRundownEditor.cpp
@@ -29,7 +29,6 @@
 #include "ScopedTransaction.h"
 #include "TabFactories/AvaRundownInstancedPageListTabFactory.h"
 #include "TabFactories/AvaRundownSubListDocumentTabFactory.h"
-#include "TabFactories/AvaRundownSubListTabFactory.h"
 #include "TabFactories/AvaRundownTemplatePageListTabFactory.h"
 #include "Widgets/Docking/SDockTab.h"
 #include "WorkflowOrientedApp/ApplicationMode.h"
@@ -175,6 +174,27 @@ private:
 		InvokeOnActiveEditor([&InArgs](const TSharedPtr<FAvaRundownEditor>& InRundownEditor){InRundownEditor->UnloadPageCommand(InArgs);});
 	}
 
+	void StopLayersPreviewCommand(const TArray<FString>& InArgs) const
+	{
+		InvokeOnActiveEditor([&InArgs](const TSharedPtr<FAvaRundownEditor>& InRundownEditor){InRundownEditor->StopLayersCommand(InArgs, /*bInPreview*/ true, /*bInForce*/ false);});
+	}
+
+	void ForceStopLayersPreviewCommand(const TArray<FString>& InArgs) const
+	{
+		InvokeOnActiveEditor([&InArgs](const TSharedPtr<FAvaRundownEditor>& InRundownEditor){InRundownEditor->StopLayersCommand(InArgs, /*bInPreview*/ true, /*bInForce*/ true);});
+	}
+
+	void StopLayersProgramCommand(const TArray<FString>& InArgs) const
+	{
+		InvokeOnActiveEditor([&InArgs](const TSharedPtr<FAvaRundownEditor>& InRundownEditor){InRundownEditor->StopLayersCommand(InArgs, /*bInPreview*/ false, /*bInForce*/ false);});
+	}
+
+	void ForceStopLayersProgramCommand(const TArray<FString>& InArgs) const
+	{
+		InvokeOnActiveEditor([&InArgs](const TSharedPtr<FAvaRundownEditor>& InRundownEditor){InRundownEditor->StopLayersCommand(InArgs, /*bInPreview*/ false, /*bInForce*/ true);});
+	}
+	
+
 private:
 	TArray<IConsoleObject*> ConsoleCommands;
 	TArray<TWeakPtr<FAvaRundownEditor>> RundownEditors;
@@ -1359,6 +1379,69 @@ void FAvaRundownEditor::StopPageCommand(const TArray<FString>& InArgs, bool bInP
 	}
 }
 
+void FAvaRundownEditor::StopLayersCommand(const TArray<FString>& InArgs, bool bInPreview, bool bInForce)
+{
+	// Validate arguments
+	if (InArgs.Num() < 2)
+	{
+		UE_LOG(LogAvaRundown, Error, TEXT("Errors in \"Stop Layers\" command: Too few arguments, need at least ChannelName and one layer name"));
+	}
+
+	const FName ChannelName(*InArgs[0]);
+
+	bool bAnyChannel = false;
+	if (InArgs[0] == TEXT("all") || InArgs[0] == TEXT("any"))
+	{
+		bAnyChannel = true;
+	}
+	
+	UAvaRundown* Rundown = AvaRundown.Get();
+	if (!Rundown)
+	{
+		return;
+	}
+
+	TArray<FName> LayerNames;
+	LayerNames.Reserve(InArgs.Num());
+	for (const FString& Arg : InArgs)
+	{
+		LayerNames.Add(FName(*Arg));
+	}
+	LayerNames.RemoveAtSwap(0);
+
+	// Using a map to gather layers per channel.
+	TMap<FName, TArray<FAvaTagHandle>> AllLayers;
+	
+	// Search for the layers in specified channel
+	for (const UAvaRundownPagePlayer* PagePlayer : Rundown->GetPagePlayers())
+	{
+		if (!PagePlayer
+			|| PagePlayer->bIsPreview != bInPreview
+			|| (!bAnyChannel && PagePlayer->ChannelFName != ChannelName))
+		{
+			continue;
+		}
+
+		TArray<FAvaTagHandle>& Layers = AllLayers.FindOrAdd(PagePlayer->ChannelFName);
+
+		PagePlayer->ForEachInstancePlayer([&Layers, &LayerNames](const UAvaRundownPlaybackInstancePlayer* InInstancePlayer)
+		{
+			if (LayerNames.Contains(InInstancePlayer->TransitionLayer.ToName()))
+			{
+				FAvaRundownPlaybackUtils::AddTagHandleUnique(Layers, InInstancePlayer->TransitionLayer);
+			}
+		});
+	}
+
+	const EAvaRundownPageStopOptions StopOptions = bInForce ? EAvaRundownPageStopOptions::ForceNoTransition : EAvaRundownPageStopOptions::Default;
+
+	for (TPair<FName, TArray<FAvaTagHandle>>& Layers : AllLayers)
+	{
+		UE_LOG(LogAvaRundown, Log, TEXT("Stopping layers: %s in channel %s."), *FString::JoinBy(Layers.Value, TEXT(", "), [](const FAvaTagHandle& InTag){return InTag.ToString();}), *Layers.Key.ToString());
+		Rundown->StopLayers(Layers.Key, Layers.Value, StopOptions);
+	}
+}
+
 void FAvaRundownEditor::TakeToProgramCommand(const TArray<FString>& InArgs)
 {
 	FString Errors;
@@ -1735,6 +1818,30 @@ void FAvaRundownEditor::FSharedConsoleCommands::RegisterConsoleCommands()
 		FConsoleCommandWithArgsDelegate::CreateRaw(this, &FSharedConsoleCommands::UnloadPageCommand),
 		ECVF_Default
 	));
+	ConsoleCommands.Add(IConsoleManager::Get().RegisterConsoleCommand(
+		TEXT("MotionDesignRundown.StopLayersPreview"),
+		TEXT("Stops currently previewing layers."),
+		FConsoleCommandWithArgsDelegate::CreateRaw(this, &FSharedConsoleCommands::StopLayersPreviewCommand),
+		ECVF_Default
+	));
+	ConsoleCommands.Add(IConsoleManager::Get().RegisterConsoleCommand(
+		TEXT("MotionDesignRundown.ForceStopLayersPreview"),
+		TEXT("Stops currently previewing layers."),
+		FConsoleCommandWithArgsDelegate::CreateRaw(this, &FSharedConsoleCommands::ForceStopLayersPreviewCommand),
+		ECVF_Default
+	));
+	ConsoleCommands.Add(IConsoleManager::Get().RegisterConsoleCommand(
+		TEXT("MotionDesignRundown.StopLayersProgram"),
+		TEXT("Stops currently playing layers."),
+		FConsoleCommandWithArgsDelegate::CreateRaw(this, &FSharedConsoleCommands::StopLayersProgramCommand),
+		ECVF_Default
+	));
+	ConsoleCommands.Add(IConsoleManager::Get().RegisterConsoleCommand(
+		TEXT("MotionDesignRundown.ForceStopLayersProgram"),
+		TEXT("Stops currently playing layers."),
+		FConsoleCommandWithArgsDelegate::CreateRaw(this, &FSharedConsoleCommands::ForceStopLayersProgramCommand),
+		ECVF_Default
+	));
 }
 
 void FAvaRundownEditor::FSharedConsoleCommands::UnregisterConsoleCommands()
diff --git a/Engine/Plugins/Experimental/Avalanche/Source/AvalancheMediaEditor/Private/Rundown/AvaRundownEditorUtils.cpp b/Engine/Plugins/Experimental/Avalanche/Source/AvalancheMediaEditor/Private/Rundown/AvaRundownEditorUtils.cpp
--- a/Engine/Plugins/Experimental/Avalanche/Source/AvalancheMediaEditor/Private/Rundown/AvaRundownEditorUtils.cpp
+++ b/Engine/Plugins/Experimental/Avalanche/Source/AvalancheMediaEditor/Private/Rundown/AvaRundownEditorUtils.cpp
@@ -675,7 +675,7 @@ bool UE::AvaRundownEditor::Utils::ForceUpdatePageSummary(int32 InPageId, UAvaRun
 		}
 	}
 
-	const bool bSummaryUpdated = Page.UpdatePageSummary(Presets, /*bInIsPresetChanged*/ true);
+	const bool bSummaryUpdated = Page.UpdatePageSummary(InRundown, Presets, /*bInForceUpdate*/ true);
 
 	// Restore managed instances controllers
 	for (const TSharedPtr<FAvaRundownManagedInstance>& Instance : ManagedInstances)
diff --git a/Engine/Plugins/Experimental/Avalanche/Source/AvalancheMediaEditor/Private/Rundown/DetailsView/AvaRundownPageDetailView.h b/Engine/Plugins/Experimental/Avalanche/Source/AvalancheMediaEditor/Private/Rundown/DetailsView/AvaRundownPageDetailView.h
new file mode 100644
--- /dev/null
+++ b/Engine/Plugins/Experimental/Avalanche/Source/AvalancheMediaEditor/Private/Rundown/DetailsView/AvaRundownPageDetailView.h
@@ -0,0 +1,21 @@
+// Copyright Epic Games, Inc. All Rights Reserved.
+
+#pragma once
+
+#include "UObject/Object.h"
+#include "StructUtils/InstancedStruct.h"
+
+#include "AvaRundownPageDetailView.generated.h"
+
+
+/** Class used to expose part of the rundown page data to the detail panel. */
+UCLASS()
+class UAvaRundownPageDetailView : public UObject
+{
+	GENERATED_BODY()
+
+public:
+	/** Page Commands */
+	UPROPERTY(EditAnywhere, BlueprintReadWrite, Category = "Page Data", meta = (ExcludeBaseStruct, BaseStruct = "/Script/AvalancheMedia.AvaRundownPageCommand"))
+	TArray<FInstancedStruct> Commands;
+};
diff --git a/Engine/Plugins/Experimental/Avalanche/Source/AvalancheMediaEditor/Private/Rundown/DetailsView/RemoteControl/Controllers/AvaRundownPageControllerContextMenu.cpp b/Engine/Plugins/Experimental/Avalanche/Source/AvalancheMediaEditor/Private/Rundown/DetailsView/RemoteControl/Controllers/AvaRundownPageControllerContextMenu.cpp
--- a/Engine/Plugins/Experimental/Avalanche/Source/AvalancheMediaEditor/Private/Rundown/DetailsView/RemoteControl/Controllers/AvaRundownPageControllerContextMenu.cpp
+++ b/Engine/Plugins/Experimental/Avalanche/Source/AvalancheMediaEditor/Private/Rundown/DetailsView/RemoteControl/Controllers/AvaRundownPageControllerContextMenu.cpp
@@ -230,7 +230,7 @@ void FAvaRundownPageControllerContextMenu::ResetValuesToDefaults(const bool bInU
 	{
 		if (const TSharedPtr<SAvaRundownRCControllerPanel> ControllerListWidget = CurrentContext->GetControllerListWidget())
 		{
-			ControllerListWidget->UpdatePageSummary(true);
+			ControllerListWidget->UpdatePageSummary(/*bInForceUpdate*/ true);
 		}
 	}
 }
diff --git a/Engine/Plugins/Experimental/Avalanche/Source/AvalancheMediaEditor/Private/Rundown/DetailsView/RemoteControl/Controllers/SAvaRundownRCControllerPanel.cpp b/Engine/Plugins/Experimental/Avalanche/Source/AvalancheMediaEditor/Private/Rundown/DetailsView/RemoteControl/Controllers/SAvaRundownRCControllerPanel.cpp
--- a/Engine/Plugins/Experimental/Avalanche/Source/AvalancheMediaEditor/Private/Rundown/DetailsView/RemoteControl/Controllers/SAvaRundownRCControllerPanel.cpp
+++ b/Engine/Plugins/Experimental/Avalanche/Source/AvalancheMediaEditor/Private/Rundown/DetailsView/RemoteControl/Controllers/SAvaRundownRCControllerPanel.cpp
@@ -245,7 +245,7 @@ void SAvaRundownRCControllerPanel::Refresh(const TArray<int32>& InSelectedPageId
 	ControllerItems.Empty(NumItems);
 	
 	UpdatePropertyRowGenerators(ManagedInstances.Num());
-	UpdatePageSummary(false);
+	UpdatePageSummary(/*bInForceUpdate*/ false); // Generate only if missing.
 
 	for (int32 Index = 0; Index < ManagedInstances.Num(); ++Index)
 	{
@@ -389,7 +389,7 @@ bool SAvaRundownRCControllerPanel::SetSelectedPageControllerValue(const URCVirtu
 	return false;
 }
 
-void SAvaRundownRCControllerPanel::UpdatePageSummary(bool bInIsPresetChanged)
+void SAvaRundownRCControllerPanel::UpdatePageSummary(bool bInForceUpdate)
 {
 	FAvaRundownPage& Page = GetActivePageMutable();
 	if (Page.IsValidPage())
@@ -403,7 +403,7 @@ void SAvaRundownRCControllerPanel::UpdatePageSummary(bool bInIsPresetChanged)
 				Presets.Add(ManagedInstance->GetRemoteControlPreset());
 			}
 		}					
-		Page.UpdatePageSummary(Presets, bInIsPresetChanged);
+		Page.UpdatePageSummary(GetRundown(), Presets, bInForceUpdate);
 	}
 }
 
@@ -485,7 +485,7 @@ void SAvaRundownRCControllerPanel::FPropertyRowGeneratorWrapper::NotifyPostChang
 		Preset->OnModifyController(InPropertyChangedEvent);
 		if (ParentPanel && InPropertyChangedEvent.ChangeType & EPropertyChangeType::ValueSet)
 		{
-			ParentPanel->UpdatePageSummary(true);
+			ParentPanel->UpdatePageSummary(/*bInForceUpdate*/ true);
 			OngoingPropertyChanges.Remove(InPropertyThatChanged);
 		}
 	}
diff --git a/Engine/Plugins/Experimental/Avalanche/Source/AvalancheMediaEditor/Private/Rundown/DetailsView/RemoteControl/Controllers/SAvaRundownRCControllerPanel.h b/Engine/Plugins/Experimental/Avalanche/Source/AvalancheMediaEditor/Private/Rundown/DetailsView/RemoteControl/Controllers/SAvaRundownRCControllerPanel.h
--- a/Engine/Plugins/Experimental/Avalanche/Source/AvalancheMediaEditor/Private/Rundown/DetailsView/RemoteControl/Controllers/SAvaRundownRCControllerPanel.h
+++ b/Engine/Plugins/Experimental/Avalanche/Source/AvalancheMediaEditor/Private/Rundown/DetailsView/RemoteControl/Controllers/SAvaRundownRCControllerPanel.h
@@ -53,7 +53,7 @@ public:
 
 	const TArray<FAvaRundownRCControllerItemPtr> GetSelectedControllerItems() const;
 
-	void UpdatePageSummary(bool bInIsPresetChanged);
+	void UpdatePageSummary(bool bInForceUpdate);
 
 private:
 	static FAvaRundownRCControllerHeaderRowExtensionDelegate HeaderRowExtensionDelegate;
diff --git a/Engine/Plugins/Experimental/Avalanche/Source/AvalancheMediaEditor/Private/Rundown/DetailsView/SAvaRundownPageDetails.cpp b/Engine/Plugins/Experimental/Avalanche/Source/AvalancheMediaEditor/Private/Rundown/DetailsView/SAvaRundownPageDetails.cpp
--- a/Engine/Plugins/Experimental/Avalanche/Source/AvalancheMediaEditor/Private/Rundown/DetailsView/SAvaRundownPageDetails.cpp
+++ b/Engine/Plugins/Experimental/Avalanche/Source/AvalancheMediaEditor/Private/Rundown/DetailsView/SAvaRundownPageDetails.cpp
@@ -2,6 +2,7 @@
 
 #include "SAvaRundownPageDetails.h"
 
+#include "AvaRundownPageDetailView.h"
 #include "Async/Async.h"
 #include "Framework/Application/SlateApplication.h"
 #include "IAvaMediaModule.h"
@@ -50,6 +51,18 @@ void SAvaRundownPageDetails::Construct(const FArguments& InArgs, const TSharedPt
 		Rundown->GetOnPageListChanged().AddSP(this, &SAvaRundownPageDetails::OnPageListChanged);
 	}
 
+	FPropertyEditorModule& EditModule = FModuleManager::Get().GetModuleChecked<FPropertyEditorModule>("PropertyEditor");
+
+	FDetailsViewArgs DetailsViewArgs;
+	DetailsViewArgs.NameAreaSettings = FDetailsViewArgs::HideNameArea;
+	DetailsViewArgs.bHideSelectionTip = true;
+	DetailsViewArgs.bShowOptions = false;
+	DetailsViewArgs.bAllowSearch = false;
+	DetailsViewArgs.bShowPropertyMatrixButton = false;
+
+	PageDetailView = EditModule.CreateDetailView(DetailsViewArgs);
+	PageDetailView->OnFinishedChangingProperties().AddSP(this, &SAvaRundownPageDetails::OnFinishedChangingPageDetailProperties);
+
 	ChildSlot
 	[
 		SNew(SScrollBox)
@@ -165,6 +178,12 @@ void SAvaRundownPageDetails::Construct(const FArguments& InArgs, const TSharedPt
 			[
 				SAssignNew(RemoteControlProps, SAvaRundownPageRemoteControlProps, SharedThis(this), InRundownEditor)
 			]
+			// Page Data View
+			+ SVerticalBox::Slot()
+			.AutoHeight()
+			[
+				PageDetailView.ToSharedRef()
+			]
 		]
 	];
 }
@@ -181,6 +200,7 @@ SAvaRundownPageDetails::~SAvaRundownPageDetails()
 			Rundown->GetOnPageListChanged().RemoveAll(this);
 		}
 	}
+
 	if (IAvaMediaModule::IsModuleLoaded() && IAvaMediaModule::Get().IsManagedInstanceCacheAvailable())
 	{
 		IAvaMediaModule::Get().GetManagedInstanceCache().OnEntryInvalidated.RemoveAll(this);
@@ -203,9 +223,47 @@ void SAvaRundownPageDetails::OnPageEvent(const TArray<int32>& InSelectedPageIds,
 	{
 		RemoteControlProps->Refresh(InSelectedPageIds);
 		RCControllerPanel->Refresh(InSelectedPageIds);
+		RefreshPageDetailView();
 	}
 }
 
+void SAvaRundownPageDetails::RefreshPageDetailView()
+{
+	const FRundownContext RundownContext = GetRundownContext();
+	if (RundownContext.IsValid())
+	{
+		// Save previously edited page -- todo: this may not be necessary.
+		if (ActivePageId != CurrentPageDetailViewPageId
+			&& CurrentPageDetailViewPageId != FAvaRundownPage::InvalidPageId
+			&& CurrentPageDetailViewData.IsValid())
+		{
+			FAvaRundownPage& Page = RundownContext.GetRundown()->GetPage(CurrentPageDetailViewPageId);
+			if (Page.IsValidPage())
+			{
+				RundownContext.GetRundown()->Modify();
+				Page.SetInstancedCommands(CurrentPageDetailViewData->Commands);
+				RundownContext.GetRundown()->GetOnPagesChanged().Broadcast(RundownContext.GetRundown(), Page, EAvaRundownPageChanges::Commands);
+			}
+		}
+		
+		const FAvaRundownPage& ActivePage = RundownContext.GetRundown()->GetPage(ActivePageId);
+		if (ActivePage.IsValidPage())
+		{
+			if (!CurrentPageDetailViewData.IsValid())
+			{
+				CurrentPageDetailViewData.Reset(NewObject<UAvaRundownPageDetailView>());
+			}
+			CurrentPageDetailViewData->Commands = ActivePage.GetInstancedCommands();
+			PageDetailView->SetObject(CurrentPageDetailViewData.Get());
+			CurrentPageDetailViewPageId = ActivePageId;
+			return;
+		}
+	}
+
+	const TArray<UObject*> SelectedObjects;
+	PageDetailView->SetObjects(SelectedObjects);
+}
+
 void SAvaRundownPageDetails::OnManagedInstanceCacheEntryInvalidated(const FSoftObjectPath& InAssetPath)
 {
 	if (!bRefreshSelectedPageQueued)
@@ -225,6 +283,27 @@ void SAvaRundownPageDetails::OnManagedInstanceCacheEntryInvalidated(const FSoftO
 	}
 }
 
+void SAvaRundownPageDetails::OnFinishedChangingPageDetailProperties(const FPropertyChangedEvent& InChangedProperty)
+{
+	// Ansme: is this part of the transaction?
+	const FRundownContext RundownContext = GetRundownContext();
+	if (RundownContext.IsValid() && CurrentPageDetailViewPageId != FAvaRundownPage::InvalidPageId && CurrentPageDetailViewData.IsValid())
+	{
+		FAvaRundownPage& Page = RundownContext.GetRundown()->GetPage(CurrentPageDetailViewPageId);
+		if (Page.IsValidPage())
+		{
+			RundownContext.GetRundown()->Modify();
+			Page.SetInstancedCommands(CurrentPageDetailViewData->Commands);
+			RundownContext.GetRundown()->GetOnPagesChanged().Broadcast(RundownContext.GetRundown(), Page, EAvaRundownPageChanges::Commands);
+			if (RCControllerPanel)
+			{
+				// Commands are also part of the page summary.
+				RCControllerPanel->UpdatePageSummary(/*bInForceUpdate*/ true);
+			}
+		}
+	}
+}
+
 FReply SAvaRundownPageDetails::ToggleExposedPropertiesVisibility()
 {
 	if (UAvaRundownEditorSettings* RundownEditorSettings = UAvaRundownEditorSettings::GetMutable())
@@ -301,6 +380,7 @@ void SAvaRundownPageDetails::QueueRefreshSelectedPage()
 		{
 			RemoteControlProps->Refresh({SelectedPage.GetPageId()});
 			RCControllerPanel->Refresh({SelectedPage.GetPageId()});
+			RefreshPageDetailView();
 		}
 		bRefreshSelectedPageQueued = false;
 		return false;
@@ -322,7 +402,8 @@ void SAvaRundownPageDetails::QueueUpdateAndRefreshSelectedPage()
 		if (SelectedPage.IsValidPage())
 		{
 			RemoteControlProps->UpdateDefaultValuesAndRefresh({SelectedPage.GetPageId()});
-			RemoteControlProps->Refresh({SelectedPage.GetPageId()});
+			RCControllerPanel->Refresh({SelectedPage.GetPageId()});
+			RefreshPageDetailView();
 		}
 		bUpdateAndRefreshSelectedPageQueued = true;
 		return false;
diff --git a/Engine/Plugins/Experimental/Avalanche/Source/AvalancheMediaEditor/Private/Rundown/DetailsView/SAvaRundownPageDetails.h b/Engine/Plugins/Experimental/Avalanche/Source/AvalancheMediaEditor/Private/Rundown/DetailsView/SAvaRundownPageDetails.h
--- a/Engine/Plugins/Experimental/Avalanche/Source/AvalancheMediaEditor/Private/Rundown/DetailsView/SAvaRundownPageDetails.h
+++ b/Engine/Plugins/Experimental/Avalanche/Source/AvalancheMediaEditor/Private/Rundown/DetailsView/SAvaRundownPageDetails.h
@@ -2,8 +2,10 @@
 
 #pragma once
 
+#include "AvaRundownPageDetailView.h"
 #include "Containers/Array.h"
 #include "Rundown/AvaRundownEditorDefines.h"
+#include "Rundown/AvaRundownPage.h"
 #include "Templates/SharedPointer.h"
 #include "Types/SlateEnums.h"
 #include "Widgets/SCompoundWidget.h"
@@ -40,6 +42,8 @@ public:
 	}
 
 protected:
+	void OnFinishedChangingPageDetailProperties(const FPropertyChangedEvent& InChangedProperty);
+
 	FReply ToggleExposedPropertiesVisibility();
 
 	const FSlateBrush* GetExposedPropertiesVisibilityBrush() const;
@@ -59,6 +63,7 @@ protected:
 
 	void QueueRefreshSelectedPage();
 	void QueueUpdateAndRefreshSelectedPage();
+	void RefreshPageDetailView();
 
 	bool HasSelectedPage() const;
 
@@ -81,11 +86,19 @@ private:
 	TWeakPtr<FAvaRundownEditor> RundownEditorWeak;
 
 	TSharedPtr<SAvaRundownPageRemoteControlProps> RemoteControlProps;
-
 	TSharedPtr<SAvaRundownRCControllerPanel> RCControllerPanel;
 
+	/** Detail view of the page detail model. */
+	TSharedPtr<IDetailsView> PageDetailView;
+
+	/** Page view data of the current page (being edited). */
+	TStrongObjectPtr<UAvaRundownPageDetailView> CurrentPageDetailViewData;
+
+	/** This is separate from ActivePageId because of order of update events. ActivePageId update first. */ 
+	int32 CurrentPageDetailViewPageId = FAvaRundownPage::InvalidPageId;
+
 	bool bRefreshSelectedPageQueued = false;
 	bool bUpdateAndRefreshSelectedPageQueued = false;
 
-	int32 ActivePageId;
+	int32 ActivePageId = FAvaRundownPage::InvalidPageId;
 };
diff --git a/Engine/Plugins/Experimental/Avalanche/Source/AvalancheMediaEditor/Private/Rundown/Pages/Columns/AvaRundownPageAssetNameColumn.cpp b/Engine/Plugins/Experimental/Avalanche/Source/AvalancheMediaEditor/Private/Rundown/Pages/Columns/AvaRundownPageAssetNameColumn.cpp
--- a/Engine/Plugins/Experimental/Avalanche/Source/AvalancheMediaEditor/Private/Rundown/Pages/Columns/AvaRundownPageAssetNameColumn.cpp
+++ b/Engine/Plugins/Experimental/Avalanche/Source/AvalancheMediaEditor/Private/Rundown/Pages/Columns/AvaRundownPageAssetNameColumn.cpp
@@ -39,7 +39,7 @@ TSharedRef<SWidget> FAvaRundownPageAssetNameColumn::ConstructRowWidget(const FAv
 		.Padding(UE::AvaRundown::FEditorMetrics::ColumnLeftOffset, 0.f, 0.f, 0.f)
 		[
 			SNew(STextBlock)
-			.Text(InstancedPageView, &FAvaRundownPageViewImpl::GetObjectName, Rundown)
+			.Text(InstancedPageView, &FAvaRundownPageViewImpl::GetObjectNames, Rundown)
 		];
 }
 
diff --git a/Engine/Plugins/Experimental/Avalanche/Source/AvalancheMediaEditor/Private/Rundown/Pages/Columns/AvaRundownPageAssetSelectorColumn.cpp b/Engine/Plugins/Experimental/Avalanche/Source/AvalancheMediaEditor/Private/Rundown/Pages/Columns/AvaRundownPageAssetSelectorColumn.cpp
--- a/Engine/Plugins/Experimental/Avalanche/Source/AvalancheMediaEditor/Private/Rundown/Pages/Columns/AvaRundownPageAssetSelectorColumn.cpp
+++ b/Engine/Plugins/Experimental/Avalanche/Source/AvalancheMediaEditor/Private/Rundown/Pages/Columns/AvaRundownPageAssetSelectorColumn.cpp
@@ -64,10 +64,10 @@ TSharedRef<SWidget> FAvaRundownPageAssetSelectorColumn::ConstructRowWidget(const
 
 	// Combo templates don't have an asset selector.
 	// Put a place holder instead.
-	if (!InPageView->HasObjectPath(Rundown))
+	if (InPageView->IsComboTemplate())
 	{
 		return SNew(STextBlock)
-		.Text(LOCTEXT("AssetSelectorColumn_ComboPage", "N/A"))
+		.Text(InPageView, &IAvaRundownPageView::GetObjectNames, Rundown)
 		.Font(FAppStyle::GetFontStyle(TEXT("PropertyWindow.NormalFont")));
 	}
 	
diff --git a/Engine/Plugins/Experimental/Avalanche/Source/AvalancheMediaEditor/Private/Rundown/Pages/Columns/AvaRundownPageStatusColumn.cpp b/Engine/Plugins/Experimental/Avalanche/Source/AvalancheMediaEditor/Private/Rundown/Pages/Columns/AvaRundownPageStatusColumn.cpp
--- a/Engine/Plugins/Experimental/Avalanche/Source/AvalancheMediaEditor/Private/Rundown/Pages/Columns/AvaRundownPageStatusColumn.cpp
+++ b/Engine/Plugins/Experimental/Avalanche/Source/AvalancheMediaEditor/Private/Rundown/Pages/Columns/AvaRundownPageStatusColumn.cpp
@@ -54,7 +54,7 @@ TSharedRef<SWidget> FAvaRundownPageStatusColumn::ConstructRowWidget(const FAvaRu
 			.ButtonStyle(FAvaMediaEditorStyle::Get(), "AvaMediaEditor.BorderlessButton")
 			.OnClicked(InPageView, &IAvaRundownPageView::OnPreviewButtonClicked)
 			.IsEnabled(InPageView, &IAvaRundownPageView::CanPreview)
-			.ToolTipText(LOCTEXT("Preview", "Preview\n\n- Click: Preview from start\n- +Shift: Use Preview Frame\n- +Control: Continue"))
+			.ToolTipText(InPageView, &IAvaRundownPageView::GetPreviewInTooltip)
 			[
 				SNew(SBorder)
 				.Padding(FMargin(0.f, 0.f))
@@ -166,14 +166,11 @@ const FSlateBrush* FAvaRundownPageStatusColumn::GetPreviewBorderBackgroundImage(
 
 		if (IsValid(Rundown))
 		{
-			if (Rundown->IsPagePreviewing(PageView->GetPageId()))
+			if (const FAvaRundownPageListPlaybackContextCollection* ContextCollection = Rundown->GetPageListPlaybackContextCollection())
 			{
-				if (const FAvaRundownPageListPlaybackContextCollection* ContextCollection = Rundown->GetPageListPlaybackContextCollection())
+				if (const TSharedPtr<FAvaRundownPageListPlaybackContext> Context = ContextCollection->GetContext(/*bInIsPreview*/ true, Rundown->GetDefaultPreviewChannelName()))
 				{
-					if (const TSharedPtr<FAvaRundownPageListPlaybackContext> Context = ContextCollection->GetContext(true, Rundown->GetDefaultPreviewChannelName()))
-					{
-						bIsPlayHead = Context->PlayHeadPageId == PageView->GetPageId();
-					}
+					bIsPlayHead = Context->PlayHeadPageId == PageView->GetPageId();
 				}
 			}
 		}
@@ -233,14 +230,11 @@ const FSlateBrush* FAvaRundownPageStatusColumn::GetProgramBorderBackgroundImage(
 
 		if (IsValid(Rundown))
 		{
-			if (Rundown->IsPagePlaying(PageView->GetPageId()))
+			if (const FAvaRundownPageListPlaybackContextCollection* ContextCollection = Rundown->GetPageListPlaybackContextCollection())
 			{
-				if (const FAvaRundownPageListPlaybackContextCollection* ContextCollection = Rundown->GetPageListPlaybackContextCollection())
+				if (const TSharedPtr<FAvaRundownPageListPlaybackContext> Context = ContextCollection->GetContext(/*bInIsPreview*/ false, NAME_None))
 				{
-					if (const TSharedPtr<FAvaRundownPageListPlaybackContext> Context = ContextCollection->GetContext(false, NAME_None))
-					{
-						bIsPlayHead = Context->PlayHeadPageId == PageView->GetPageId();
-					}
+					bIsPlayHead = Context->PlayHeadPageId == PageView->GetPageId();
 				}
 			}
 		}
@@ -256,11 +250,12 @@ const FSlateBrush* FAvaRundownPageStatusColumn::GetProgramBorderBackgroundImage(
 FText FAvaRundownPageStatusColumn::GetTakeInTooltip(const TWeakPtr<FAvaRundownInstancedPageViewImpl> InPageViewWeak)
 {
 	// static const FText BaseTooltip = LOCTEXT("TakeIn", "Take In\n\n- Click: Play\n- Control+Click: Continue");
-	static const FText BaseTooltip = LOCTEXT("Status", "Page Status: ");
-	static const FText PlayedTooltip = LOCTEXT("Playing", "Playing");
-	static const FText NotPlayedTooltip = LOCTEXT("Stopped", "Stopped");
-	static const FText ErrorTooltip = LOCTEXT("CantPlay", "**Cannot Take In**");
-	static const FText NoOutputs = LOCTEXT("NoOutputs", "No Outputs Selected");
+	static const FText BaseTooltip = LOCTEXT("Program_Status", "Program Status: ");
+	static const FText PlayedTooltip = LOCTEXT("Program_Playing", "Playing");
+	static const FText NotPlayedTooltip = LOCTEXT("Program_Stopped", "Stopped");
+	static const FText ErrorTooltip = LOCTEXT("Program_CantPlay", "**Cannot Take In**");
+	static const FText NoOutputs = LOCTEXT("Program_NoOutputs", "No Outputs Selected");
+	static const FText NewLines = LOCTEXT("Program_NewLines", "\n\n");
 
 	if (const FAvaRundownPageViewPtr PageView = InPageViewWeak.Pin())
 	{
@@ -273,9 +268,11 @@ FText FAvaRundownPageStatusColumn::GetTakeInTooltip(const TWeakPtr<FAvaRundownIn
 			bool bAddedErrorTooltip = false;
 
 			// Checks whether it can play based on situation, not status
-			if (!Rundown->CanPlayPage(PageView->GetPageId(), false))
+			FString FailureReason;
+			if (!Rundown->CanPlayPage(PageView->GetPageId(), /*bInPreview*/ false, NAME_None, &FailureReason))
 			{
 				Texts.Add(ErrorTooltip);
+				Texts.Add(FText::Format(LOCTEXT("Program_CantPlayReason", "Reason: {0}"), FText::FromString(FailureReason)));
 				bAddedErrorTooltip = true;
 			}
 
@@ -298,18 +295,11 @@ FText FAvaRundownPageStatusColumn::GetTakeInTooltip(const TWeakPtr<FAvaRundownIn
 				}
 				else
 				{
-					if (Rundown->IsPagePlaying(Page))
-					{
-						Texts.Add(PlayedTooltip);
-					}
-					else
-					{
-						Texts.Add(NotPlayedTooltip);
-					}
+					Texts.Add(Rundown->IsPagePlaying(Page) ? PlayedTooltip : NotPlayedTooltip);
 				}
 			}
 
-			return FText::Join(LOCTEXT("NewLines", "\n\n"), Texts);
+			return FText::Join(NewLines, Texts);
 		}
 	}
 
diff --git a/Engine/Plugins/Experimental/Avalanche/Source/AvalancheMediaEditor/Private/Rundown/Pages/Columns/AvaRundownPageTemplateStatusColumn.cpp b/Engine/Plugins/Experimental/Avalanche/Source/AvalancheMediaEditor/Private/Rundown/Pages/Columns/AvaRundownPageTemplateStatusColumn.cpp
--- a/Engine/Plugins/Experimental/Avalanche/Source/AvalancheMediaEditor/Private/Rundown/Pages/Columns/AvaRundownPageTemplateStatusColumn.cpp
+++ b/Engine/Plugins/Experimental/Avalanche/Source/AvalancheMediaEditor/Private/Rundown/Pages/Columns/AvaRundownPageTemplateStatusColumn.cpp
@@ -53,7 +53,7 @@ TSharedRef<SWidget> FAvaRundownPageTemplateStatusColumn::ConstructRowWidget(cons
 			.ButtonStyle(FAvaMediaEditorStyle::Get(), "AvaMediaEditor.BorderlessButton")
 			.OnClicked(TemplatePageView, &FAvaRundownTemplatePageViewImpl::OnPreviewButtonClicked)
 			.IsEnabled(TemplatePageView, &FAvaRundownTemplatePageViewImpl::CanPreview)
-			.ToolTipText(LOCTEXT("Preview", "Preview\n\n- Click: Preview from start\n- +Shift: Use Preview Frame\n- +Control: Continue"))
+			.ToolTipText(InPageView, &IAvaRundownPageView::GetPreviewInTooltip)
 			[
 				SNew(SImage)
 				.ColorAndOpacity_Static(&FAvaRundownPageTemplateStatusColumn::GetIsPreviewingButtonColor, TemplatePageViewWeak)
@@ -104,15 +104,15 @@ FSlateColor FAvaRundownPageTemplateStatusColumn::GetIsPreviewingButtonColor(cons
 
 	if (const TSharedPtr<FAvaRundownTemplatePageViewImpl> PageView = InPageViewWeak.Pin())
 	{
-		UAvaRundown* Previewlist = PageView->GetRundown();
+		const UAvaRundown* Rundown = PageView->GetRundown();
 
-		if (IsValid(Previewlist))
+		if (IsValid(Rundown))
 		{
-			FAvaRundownPage& Page = Previewlist->GetPage(PageView->GetPageId());
+			const FAvaRundownPage& Page = Rundown->GetPage(PageView->GetPageId());
 
 			if (Page.IsValidPage())
 			{
-				const TArray<FAvaRundownChannelPageStatus> Statuses = Page.GetPagePreviewStatuses(Previewlist);
+				const TArray<FAvaRundownChannelPageStatus> Statuses = Page.GetPagePreviewStatuses(Rundown);
 
 				if (FAvaRundownPage::StatusesContainsStatus(Statuses, {EAvaRundownPageStatus::Unknown, 
 					EAvaRundownPageStatus::Error, 
@@ -162,7 +162,7 @@ FSlateColor FAvaRundownPageTemplateStatusColumn::GetAssetStatusButtonColor(const
 
 		if (IsValid(Rundown))
 		{
-			FAvaRundownPage& Page = Rundown->GetPage(PageView->GetPageId());
+			const FAvaRundownPage& Page = Rundown->GetPage(PageView->GetPageId());
 
 			if (Page.IsValidPage())
 			{
@@ -200,15 +200,15 @@ FText FAvaRundownPageTemplateStatusColumn::GetAssetStatusButtonTooltip(const TWe
 {
 	if (const TSharedPtr<FAvaRundownTemplatePageViewImpl> PageView = InPageViewWeak.Pin())
 	{
-		UAvaRundown* Rundown = PageView->GetRundown();
+		const UAvaRundown* Rundown = PageView->GetRundown();
 
 		if (IsValid(Rundown))
 		{
-			FAvaRundownPage& Page = Rundown->GetPage(PageView->GetPageId());
+			const FAvaRundownPage& Page = Rundown->GetPage(PageView->GetPageId());
 
 			if (Page.IsValidPage())
 			{
-				UEnum* StatusEnum = StaticEnum<EAvaRundownPageStatus>();
+				const UEnum* StatusEnum = StaticEnum<EAvaRundownPageStatus>();
 				const TArray<FAvaRundownChannelPageStatus> Statuses = Page.GetPagePreviewStatuses(Rundown);
 				TArray<FText> StatusTexts;
 
@@ -240,11 +240,11 @@ FText FAvaRundownPageTemplateStatusColumn::GetSyncStatusButtonTooltip(const TWea
 
 	if (const TSharedPtr<FAvaRundownTemplatePageViewImpl> PageView = InPageViewWeak.Pin())
 	{
-		UAvaRundown* Rundown = PageView->GetRundown();
+		const UAvaRundown* Rundown = PageView->GetRundown();
 
 		if (IsValid(Rundown))
 		{
-			FAvaRundownPage& Page = Rundown->GetPage(PageView->GetPageId());
+			const FAvaRundownPage& Page = Rundown->GetPage(PageView->GetPageId());
 
 			if (Page.IsValidPage())
 			{
@@ -281,11 +281,11 @@ FSlateColor FAvaRundownPageTemplateStatusColumn::GetSyncStatusButtonColor(
 
 	if (const TSharedPtr<FAvaRundownTemplatePageViewImpl> PageView = InPageViewWeak.Pin())
 	{
-		UAvaRundown* Rundown = PageView->GetRundown();
+		const UAvaRundown* Rundown = PageView->GetRundown();
 
 		if (IsValid(Rundown))
 		{
-			FAvaRundownPage& Page = Rundown->GetPage(PageView->GetPageId());
+			const FAvaRundownPage& Page = Rundown->GetPage(PageView->GetPageId());
 
 			if (Page.IsValidPage())
 			{
diff --git a/Engine/Plugins/Experimental/Avalanche/Source/AvalancheMediaEditor/Private/Rundown/Pages/Columns/Slate/SAvaRundownPageThumbnail.cpp b/Engine/Plugins/Experimental/Avalanche/Source/AvalancheMediaEditor/Private/Rundown/Pages/Columns/Slate/SAvaRundownPageThumbnail.cpp
--- a/Engine/Plugins/Experimental/Avalanche/Source/AvalancheMediaEditor/Private/Rundown/Pages/Columns/Slate/SAvaRundownPageThumbnail.cpp
+++ b/Engine/Plugins/Experimental/Avalanche/Source/AvalancheMediaEditor/Private/Rundown/Pages/Columns/Slate/SAvaRundownPageThumbnail.cpp
@@ -48,7 +48,7 @@ void SAvaRundownPageThumbnail::Construct(const FArguments& InArgs, const FAvaRun
 
 void SAvaRundownPageThumbnail::OnAssetUpdate(const UAvaRundown* InRundown, const FAvaRundownPage& InPage, EAvaRundownPageChanges InPageChangeType)
 {
-	if (!EnumHasAnyFlags(InPageChangeType, EAvaRundownPageChanges::Blueprint) || !InRundown)
+	if (!EnumHasAnyFlags(InPageChangeType, EAvaRundownPageChanges::Blueprint|EAvaRundownPageChanges::Commands) || !InRundown)
 	{
 		return;
 	}
@@ -63,7 +63,13 @@ void SAvaRundownPageThumbnail::OnAssetUpdate(const UAvaRundown* InRundown, const
 
 		if (AssetThumbnail.IsValid() && InPage.IsValidPage() && InPage.IsTemplate() && PageIdToCheck == InPage.GetPageId())
 		{
-			const FSoftObjectPath AssetPath = InPage.GetAssetPath(InRundown);
+			FSoftObjectPath AssetPath = InPage.GetAssetPath(InRundown);
+			
+			if (AssetPath.IsNull() && InPage.HasCommands(InRundown))
+			{
+				AssetPath = FSoftObjectPath(InRundown);	// Using rundown itself as asset thumbnail.
+			}
+
 			const FAssetData Data = IAssetRegistry::Get()->GetAssetByObjectPath(AssetPath);
 
 			AssetThumbnail->SetAsset(Data);
@@ -87,7 +93,13 @@ void SAvaRundownPageThumbnail::InitThumbnailWidget()
 			{
 				Page = Rundown->GetPage(Page.GetTemplateId());
 			}
-			const FSoftObjectPath AssetPath = Page.GetAssetPath(Rundown);
+			
+			FSoftObjectPath AssetPath = Page.GetAssetPath(Rundown);
+			if (AssetPath.IsNull() && Page.HasCommands(Rundown))
+			{
+				AssetPath = FSoftObjectPath(Rundown);	// Using rundown itself as asset thumbnail. avoid the "None".
+			}
+
 			const FAssetData Data = IAssetRegistry::Get()->GetAssetByObjectPath(AssetPath);
 
 
diff --git a/Engine/Plugins/Experimental/Avalanche/Source/AvalancheMediaEditor/Private/Rundown/Pages/PageViews/AvaRundownInstancedPageViewImpl.cpp b/Engine/Plugins/Experimental/Avalanche/Source/AvalancheMediaEditor/Private/Rundown/Pages/PageViews/AvaRundownInstancedPageViewImpl.cpp
--- a/Engine/Plugins/Experimental/Avalanche/Source/AvalancheMediaEditor/Private/Rundown/Pages/PageViews/AvaRundownInstancedPageViewImpl.cpp
+++ b/Engine/Plugins/Experimental/Avalanche/Source/AvalancheMediaEditor/Private/Rundown/Pages/PageViews/AvaRundownInstancedPageViewImpl.cpp
@@ -186,4 +186,9 @@ bool FAvaRundownInstancedPageViewImpl::IsTemplate() const
 	return false;
 }
 
+bool FAvaRundownInstancedPageViewImpl::IsComboTemplate() const 
+{ 
+	return false;
+}
+
 #undef LOCTEXT_NAMESPACE
diff --git a/Engine/Plugins/Experimental/Avalanche/Source/AvalancheMediaEditor/Private/Rundown/Pages/PageViews/AvaRundownInstancedPageViewImpl.h b/Engine/Plugins/Experimental/Avalanche/Source/AvalancheMediaEditor/Private/Rundown/Pages/PageViews/AvaRundownInstancedPageViewImpl.h
--- a/Engine/Plugins/Experimental/Avalanche/Source/AvalancheMediaEditor/Private/Rundown/Pages/PageViews/AvaRundownInstancedPageViewImpl.h
+++ b/Engine/Plugins/Experimental/Avalanche/Source/AvalancheMediaEditor/Private/Rundown/Pages/PageViews/AvaRundownInstancedPageViewImpl.h
@@ -13,6 +13,7 @@ public:
 	FAvaRundownInstancedPageViewImpl(int32 InPageId, UAvaRundown* InRundown, const TSharedPtr<SAvaRundownPageList>& InPageList);
 
 	virtual bool IsTemplate() const override;
+	virtual bool IsComboTemplate() const override;
 
 	virtual FReply OnPlayButtonClicked() override;
 	virtual bool CanPlay() const override;
diff --git a/Engine/Plugins/Experimental/Avalanche/Source/AvalancheMediaEditor/Private/Rundown/Pages/PageViews/AvaRundownPageViewImpl.cpp b/Engine/Plugins/Experimental/Avalanche/Source/AvalancheMediaEditor/Private/Rundown/Pages/PageViews/AvaRundownPageViewImpl.cpp
--- a/Engine/Plugins/Experimental/Avalanche/Source/AvalancheMediaEditor/Private/Rundown/Pages/PageViews/AvaRundownPageViewImpl.cpp
+++ b/Engine/Plugins/Experimental/Avalanche/Source/AvalancheMediaEditor/Private/Rundown/Pages/PageViews/AvaRundownPageViewImpl.cpp
@@ -7,11 +7,26 @@
 #include "Rundown/AvaRundown.h"
 #include "Rundown/AvaRundownEditorDefines.h"
 #include "Rundown/AvaRundownPage.h"
+#include "Rundown/AvaRundownPageCommand.h"
 #include "Rundown/Pages/Slate/SAvaRundownPageList.h"
 #include "ScopedTransaction.h"
 
 #define LOCTEXT_NAMESPACE "AvaRundownPageViewImpl"
 
+namespace UE::AvaMediaEditor::RundownPageViewImpl
+{
+	static FString ListSeparator(TEXT(", "));
+
+	void AppendWithSeparator(FString& InDestinationString, const FString& InStringToAdd, const FString& InSeparator = ListSeparator)
+	{
+		if (!InDestinationString.IsEmpty())
+		{
+			InDestinationString += InSeparator;
+		}
+		InDestinationString += InStringToAdd;
+	}
+}
+
 FAvaRundownPageViewImpl::FAvaRundownPageViewImpl(int32 InPageId, UAvaRundown* InRundown, const TSharedPtr<SAvaRundownPageList>& InPageList)
 	: PageId(InPageId)
 	, RundownWeak(InRundown)
@@ -51,29 +66,39 @@ FText FAvaRundownPageViewImpl::GetPageNameText() const
 
 FText FAvaRundownPageViewImpl::GetPageTransitionLayerNameText() const
 {
+	using namespace UE::AvaMediaEditor::RundownPageViewImpl;
+
 	if (const UAvaRundown* Rundown = RundownWeak.Get())
 	{
 		const FAvaRundownPage& Page = GetPage();
 		if (Page.IsValidPage())
 		{
+			FString TransitionLayers;
+
+			// Collect Transition Layers from asset template
 			if (Page.HasTransitionLogic(Rundown))
 			{
-				FString TransitionLayers;
 				const int32 NumTemplates = Page.GetNumTemplates(Rundown);
 				for (int32 TemplateIndex = 0; TemplateIndex < NumTemplates; ++TemplateIndex)
 				{
-					if (!TransitionLayers.IsEmpty())
-					{
-						TransitionLayers += TEXT(", ");
-					}
-					TransitionLayers += Page.GetTransitionLayer(Rundown, TemplateIndex).ToString();
+					AppendWithSeparator(TransitionLayers, Page.GetTransitionLayer(Rundown, TemplateIndex).ToString());
 				}
-				return FText::FromString(TransitionLayers);
 			}
 
-			return LOCTEXT("PageTransitionLayerText_NA", "N/A");
+			// Collect Transition Layers from commands.
+			Page.ForEachInstancedCommands([&TransitionLayers](const FAvaRundownPageCommand& InCommand, const FAvaRundownPage&)
+			{
+				const FString CommandLayersString = InCommand.GetTransitionLayerString(ListSeparator);
+				if (!CommandLayersString.IsEmpty())
+				{
+					AppendWithSeparator(TransitionLayers, CommandLayersString);
+				}
+			}, Rundown, /*bDirectOnly*/ false);
+
+			return !TransitionLayers.IsEmpty() ? FText::FromString(TransitionLayers) : LOCTEXT("PageTransitionLayerText_NA", "N/A");
 		}
 	}
+
 	// Either invalid rundown or invalid page.
 	return LOCTEXT("PageTransitionLayerText_Invalid", "(invalid)");
 }
@@ -128,6 +153,37 @@ FText FAvaRundownPageViewImpl::GetObjectName(const UAvaRundown* InRundown) const
 	return FText::GetEmpty();
 }
 
+FText FAvaRundownPageViewImpl::GetObjectNames(const UAvaRundown* InRundown) const
+{
+	using namespace UE::AvaMediaEditor::RundownPageViewImpl;
+
+	const FAvaRundownPage& Page = GetPage();
+	if (Page.IsValidPage())
+	{
+		FString AssetNames;
+		const int32 NumTemplates = Page.GetNumTemplates(InRundown);
+		for (int32 TemplateIndex = 0; TemplateIndex < NumTemplates; ++TemplateIndex)
+		{
+			const FSoftObjectPath AssetPath = Page.GetAssetPath(InRundown);
+			if (!AssetPath.IsNull())
+			{
+				AppendWithSeparator(AssetNames, AssetPath.GetAssetName());
+			}
+		}
+
+		if (!AssetNames.IsEmpty())
+		{
+			return FText::FromString(AssetNames);
+		}
+		
+		if (Page.HasCommands(InRundown))
+		{
+			return LOCTEXT("AssetName_CommandPage", "Commands");
+		}
+	}
+	return FText::GetEmpty();
+}
+
 void FAvaRundownPageViewImpl::OnObjectChanged(const FAssetData& InAssetData)
 {
 	if (!IsPageSelected())
@@ -147,6 +203,15 @@ void FAvaRundownPageViewImpl::OnObjectChanged(const FAssetData& InAssetData)
 		});
 }
 
+bool FAvaRundownPageViewImpl::HasCommands(const UAvaRundown* InRundown) const
+{
+	if (const FAvaRundownPage& Page = GetPage(); Page.IsValidPage())
+	{
+		return Page.HasCommands(InRundown);
+	}
+	return false;
+}
+
 bool FAvaRundownPageViewImpl::Rename(const FText& InNewName)
 {
 	UAvaRundown* const Rundown = RundownWeak.Get();
@@ -298,15 +363,43 @@ bool FAvaRundownPageViewImpl::CanPreview() const
 	return false;
 }
 
-bool FAvaRundownPageViewImpl::IsPageSelected() const
+FText FAvaRundownPageViewImpl::GetPreviewInTooltip() const
 {
-	TSharedPtr<SAvaRundownPageList> PageList = PageListWeak.Pin();
-	if (!PageList.IsValid())
+	static const FText BaseTooltip = LOCTEXT("Preview_BaseTooltip", "Preview\n\n- Click: Preview from start\n- +Shift: Use Preview Frame\n- +Control: Continue");
+	static const FText StatusTooltip = LOCTEXT("Preview_Status", "Preview Status: ");
+	static const FText PlayedTooltip = LOCTEXT("Preview_Playing", "Playing");
+	static const FText NotPlayedTooltip = LOCTEXT("Preview_Stopped", "Stopped");
+	static const FText ErrorTooltip = LOCTEXT("Preview_CantPlay", "**Cannot Preview**");
+	static const FText NoOutputs = LOCTEXT("Preview_NoOutputs", "No Outputs Selected");
+	static const FText NewLines = LOCTEXT("Preview_NewLines", "\n\n");
+	static const FText ExtraTooltip = LOCTEXT("Preview_ExtraTooltip", "Click: Preview from start\n- +Shift: Use Preview Frame\n- +Control: Continue");
+	
+	if (const UAvaRundown* Rundown = GetRundown(); IsValid(Rundown))
 	{
-		return false;
+		TArray<FText> Texts;
+		Texts.Add(StatusTooltip);
+
+		// Checks whether it can play based on situation, not status
+		FString FailureReason;
+		if (!Rundown->CanPlayPage(PageId, /*bInPreview*/ true, Rundown->GetDefaultPreviewChannelName(), &FailureReason))
+		{
+			Texts.Add(ErrorTooltip);
+			Texts.Add(FText::Format(LOCTEXT("Preview_CantPlayReason", "Reason: {0}"), FText::FromString(FailureReason)));
+		}
+		
+		// Check actual status
+		Texts.Add(Rundown->IsPagePreviewing(PageId) ? PlayedTooltip : NotPlayedTooltip);
+		Texts.Add(ExtraTooltip);
+
+		return FText::Join(NewLines, Texts);
 	}
+	return BaseTooltip;
+}
 
-	return PageList->GetSelectedPageIds().Contains(PageId);
+bool FAvaRundownPageViewImpl::IsPageSelected() const
+{
+	const TSharedPtr<SAvaRundownPageList> PageList = PageListWeak.Pin();
+	return PageList.IsValid() ? PageList->GetSelectedPageIds().Contains(PageId) : false;;
 }
 
 bool FAvaRundownPageViewImpl::SetPageSelection(EAvaRundownPageViewSelectionChangeType InSelectionChangeType)
diff --git a/Engine/Plugins/Experimental/Avalanche/Source/AvalancheMediaEditor/Private/Rundown/Pages/PageViews/AvaRundownPageViewImpl.h b/Engine/Plugins/Experimental/Avalanche/Source/AvalancheMediaEditor/Private/Rundown/Pages/PageViews/AvaRundownPageViewImpl.h
--- a/Engine/Plugins/Experimental/Avalanche/Source/AvalancheMediaEditor/Private/Rundown/Pages/PageViews/AvaRundownPageViewImpl.h
+++ b/Engine/Plugins/Experimental/Avalanche/Source/AvalancheMediaEditor/Private/Rundown/Pages/PageViews/AvaRundownPageViewImpl.h
@@ -37,8 +37,11 @@ public:
 	virtual bool HasObjectPath(const UAvaRundown* InRundown) const override;
 	virtual FSoftObjectPath GetObjectPath(const UAvaRundown* InRundown) const override;
 	virtual FText GetObjectName(const UAvaRundown* InRundown) const override;
+	virtual FText GetObjectNames(const UAvaRundown* InRundown) const override;
 	virtual void OnObjectChanged(const FAssetData& InAssetData) override;
 
+	virtual bool HasCommands(const UAvaRundown* InRundown) const override;
+
 	virtual bool Rename(const FText& InNewName) override;
 	virtual bool RenameFriendlyName(const FText& InNewName) override;
 	virtual FOnPageAction& GetOnRename() override { return OnRename; }
@@ -49,6 +52,7 @@ public:
 	
 	virtual FReply OnPreviewButtonClicked() override;
 	virtual bool CanPreview() const override;
+	virtual FText GetPreviewInTooltip() const override;
 
 	virtual bool IsPageSelected() const override;
 	virtual bool SetPageSelection(EAvaRundownPageViewSelectionChangeType InSelectionChangeType) override;
diff --git a/Engine/Plugins/Experimental/Avalanche/Source/AvalancheMediaEditor/Private/Rundown/Pages/PageViews/AvaRundownTemplatePageViewImpl.cpp b/Engine/Plugins/Experimental/Avalanche/Source/AvalancheMediaEditor/Private/Rundown/Pages/PageViews/AvaRundownTemplatePageViewImpl.cpp
--- a/Engine/Plugins/Experimental/Avalanche/Source/AvalancheMediaEditor/Private/Rundown/Pages/PageViews/AvaRundownTemplatePageViewImpl.cpp
+++ b/Engine/Plugins/Experimental/Avalanche/Source/AvalancheMediaEditor/Private/Rundown/Pages/PageViews/AvaRundownTemplatePageViewImpl.cpp
@@ -83,3 +83,9 @@ bool FAvaRundownTemplatePageViewImpl::IsTemplate() const
 { 
 	return true;
 }
+
+bool FAvaRundownTemplatePageViewImpl::IsComboTemplate() const 
+{ 
+	const FAvaRundownPage& Page = GetPage();
+	return Page.IsValidPage() && Page.IsComboTemplate();
+}
diff --git a/Engine/Plugins/Experimental/Avalanche/Source/AvalancheMediaEditor/Private/Rundown/Pages/PageViews/AvaRundownTemplatePageViewImpl.h b/Engine/Plugins/Experimental/Avalanche/Source/AvalancheMediaEditor/Private/Rundown/Pages/PageViews/AvaRundownTemplatePageViewImpl.h
--- a/Engine/Plugins/Experimental/Avalanche/Source/AvalancheMediaEditor/Private/Rundown/Pages/PageViews/AvaRundownTemplatePageViewImpl.h
+++ b/Engine/Plugins/Experimental/Avalanche/Source/AvalancheMediaEditor/Private/Rundown/Pages/PageViews/AvaRundownTemplatePageViewImpl.h
@@ -18,6 +18,7 @@ public:
 	virtual UAvaRundown* GetRundown() const override;
 
 	virtual bool IsTemplate() const override;
+	virtual bool IsComboTemplate() const override;
 	virtual FReply OnPlayButtonClicked() override;
 	virtual bool CanPlay() const override;
 
diff --git a/Engine/Plugins/Experimental/Avalanche/Source/AvalancheMediaEditor/Private/Rundown/Pages/PageViews/IAvaRundownPageView.h b/Engine/Plugins/Experimental/Avalanche/Source/AvalancheMediaEditor/Private/Rundown/Pages/PageViews/IAvaRundownPageView.h
--- a/Engine/Plugins/Experimental/Avalanche/Source/AvalancheMediaEditor/Private/Rundown/Pages/PageViews/IAvaRundownPageView.h
+++ b/Engine/Plugins/Experimental/Avalanche/Source/AvalancheMediaEditor/Private/Rundown/Pages/PageViews/IAvaRundownPageView.h
@@ -43,12 +43,21 @@ public:
 	virtual FText GetPageDescription() const = 0;
 
 	virtual bool IsTemplate() const = 0;
+	
+	/** Returns true if the page is a template, and that template is a combo. */
+	virtual bool IsComboTemplate() const = 0;
 
 	virtual bool HasObjectPath(const UAvaRundown* InRundown) const = 0;
 	virtual FSoftObjectPath GetObjectPath(const UAvaRundown* InRundown) const = 0;
+	/** Returns asset name if single asset template. */
 	virtual FText GetObjectName(const UAvaRundown* InRundown) const = 0;
+	/** Returns list of all assets: single, combo or commands. */
+	virtual FText GetObjectNames(const UAvaRundown* InRundown) const = 0;
 	virtual void OnObjectChanged(const FAssetData& InAssetData) = 0;
 
+	/** Returns true if the page has commands, either directly or from it's templates. */
+	virtual bool HasCommands(const UAvaRundown* InRundown) const = 0;
+
 	virtual bool Rename(const FText& InNewName) = 0;
 	virtual bool RenameFriendlyName(const FText& InNewName) = 0;
 	
@@ -61,6 +70,7 @@ public:
 
 	virtual FReply OnPreviewButtonClicked() = 0;
 	virtual bool CanPreview() const = 0;
+	virtual FText GetPreviewInTooltip() const = 0;
 
 	virtual FReply OnPlayButtonClicked() = 0;
 	virtual bool CanPlay() const = 0;
diff --git a/Engine/Plugins/Experimental/Avalanche/Source/AvalancheMediaEditor/Private/Rundown/Pages/Slate/SAvaRundownPageList.cpp b/Engine/Plugins/Experimental/Avalanche/Source/AvalancheMediaEditor/Private/Rundown/Pages/Slate/SAvaRundownPageList.cpp
--- a/Engine/Plugins/Experimental/Avalanche/Source/AvalancheMediaEditor/Private/Rundown/Pages/Slate/SAvaRundownPageList.cpp
+++ b/Engine/Plugins/Experimental/Avalanche/Source/AvalancheMediaEditor/Private/Rundown/Pages/Slate/SAvaRundownPageList.cpp
@@ -27,6 +27,7 @@
 #include "Rundown/Pages/Slate/SAvaRundownPageViewRow.h"
 #include "SAvaRundownRenumberPages.h"
 #include "ScopedTransaction.h"
+#include "Widgets/Docking/SDockTab.h"
 #include "Widgets/SBoxPanel.h"
 #include "Widgets/SNullWidget.h"
 #include "Widgets/Views/SListView.h"
@@ -154,7 +155,9 @@ void SAvaRundownPageList::Construct(const FArguments& InArgs, TSharedPtr<FAvaRun
 {
 	RundownEditorWeak = InRundownEditor;
 	check(InRundownEditor.IsValid());
-
+	
+	OnActiveTabChangedDelegateHandle = FGlobalTabmanager::Get()->OnActiveTabChanged_Subscribe(FOnActiveTabChanged::FDelegate::CreateSP(this, &SAvaRundownPageList::OnActiveTabChanged));
+	
 	PageContextMenu = MakeShared<FAvaRundownPageContextMenu>();
 
 	PageListReference = InPageListReference;
@@ -165,7 +168,7 @@ void SAvaRundownPageList::Construct(const FArguments& InArgs, TSharedPtr<FAvaRun
 	PageTextFilter->OnChanged().AddSP(this, &SAvaRundownPageList::OnPageTextFilterChanged);
 	RefreshPagesVisibility();
 	
-	UAvaRundown* const Rundown = InRundownEditor->GetRundown();
+	const UAvaRundown* const Rundown = InRundownEditor->GetRundown();
 	check(Rundown);
 
 	InRundownEditor->GetOnPageEvent().AddSP(this, &SAvaRundownPageList::OnPageEvent);
@@ -223,14 +226,12 @@ void SAvaRundownPageList::Construct(const FArguments& InArgs, TSharedPtr<FAvaRun
 
 SAvaRundownPageList::~SAvaRundownPageList()
 {
-	if (RundownEditorWeak.IsValid())
-	{
-		TSharedPtr<FAvaRundownEditor> RundownEditor = RundownEditorWeak.Pin();
+	FGlobalTabmanager::Get()->OnActiveTabChanged_Unsubscribe(OnActiveTabChangedDelegateHandle);
+	OnActiveTabChangedDelegateHandle.Reset();
 
-		if (RundownEditor.IsValid())
-		{
-			RundownEditor->GetOnPageEvent().RemoveAll(this);
-		}
+	if (const TSharedPtr<FAvaRundownEditor> RundownEditor = RundownEditorWeak.Pin())
+	{
+		RundownEditor->GetOnPageEvent().RemoveAll(this);
 	}
 }
 
@@ -251,6 +252,8 @@ void SAvaRundownPageList::OnPageEvent(const TArray<int32>& InSelectedPageIds, UE
 		case EPageEvent::RenumberRequest:
 			OnPageItemActionRequested(InSelectedPageIds, &IAvaRundownPageView::GetOnRenumber);
 			break;
+		default:
+			break; // skip everything else
 	}
 }
 
@@ -1272,6 +1275,22 @@ void SAvaRundownPageList::OnPageTextFilterChanged()
 	Refresh();
 }
 
+void SAvaRundownPageList::OnActiveTabChanged(TSharedPtr<SDockTab> InPreviouslyActive, TSharedPtr<SDockTab> InNewlyActivated)
+{
+	if (const TSharedPtr<FAvaRundownEditor> RundownEditor = RundownEditorWeak.Pin())
+	{
+		if (RundownEditor->GetAssociatedTabManager() == InNewlyActivated->GetTabManagerPtr())
+		{
+			const TSharedPtr<SWidget> Content = InNewlyActivated->GetContent();
+			if (Content.Get() == this)
+			{
+				// The set of selected pages of this page list becomes the current set edited. Relevant for page detail.
+				RundownEditor->GetOnPageEvent().Broadcast(SelectedPageIds, UE::AvaRundown::EPageEvent::SelectionChanged);
+			}
+		}
+	}
+}
+
 TArray<int32> SAvaRundownPageList::FilterSelectedPages(FFilterPageFunctionRef InFilterPageFunction) const
 {
 	if (const UAvaRundown* Rundown = GetRundown())
diff --git a/Engine/Plugins/Experimental/Avalanche/Source/AvalancheMediaEditor/Private/Rundown/Pages/Slate/SAvaRundownPageList.h b/Engine/Plugins/Experimental/Avalanche/Source/AvalancheMediaEditor/Private/Rundown/Pages/Slate/SAvaRundownPageList.h
--- a/Engine/Plugins/Experimental/Avalanche/Source/AvalancheMediaEditor/Private/Rundown/Pages/Slate/SAvaRundownPageList.h
+++ b/Engine/Plugins/Experimental/Avalanche/Source/AvalancheMediaEditor/Private/Rundown/Pages/Slate/SAvaRundownPageList.h
@@ -15,6 +15,7 @@ class FUICommandList;
 class IAvaRundownPageViewColumn;
 class ITableRow;
 class SAssetSearchBox;
+class SDockTab;
 class SHeaderRow;
 class SHorizontalBox;
 class STableViewBase;
@@ -201,6 +202,8 @@ private:
 
 	void OnPageTextFilterChanged();
 
+	void OnActiveTabChanged(TSharedPtr<SDockTab> InPreviouslyActive, TSharedPtr<SDockTab> InNewlyActivated);
+
 protected:
 	TWeakPtr<FAvaRundownEditor> RundownEditorWeak;
 
@@ -235,4 +238,7 @@ protected:
 
 private:
 	TSharedPtr<SAssetSearchBox> AssetSearchBoxPtr;
+	
+	/** Delegate handle bound to the FGlobalTabmanager::OnActiveTabChanged delegate */
+	FDelegateHandle OnActiveTabChangedDelegateHandle;
 };
diff --git a/Engine/Plugins/Experimental/Avalanche/Source/AvalancheMediaEditor/Public/Rundown/AvaRundownEditor.h b/Engine/Plugins/Experimental/Avalanche/Source/AvalancheMediaEditor/Public/Rundown/AvaRundownEditor.h
--- a/Engine/Plugins/Experimental/Avalanche/Source/AvalancheMediaEditor/Public/Rundown/AvaRundownEditor.h
+++ b/Engine/Plugins/Experimental/Avalanche/Source/AvalancheMediaEditor/Public/Rundown/AvaRundownEditor.h
@@ -133,6 +133,7 @@ protected:
 	void ContinuePageCommand(const TArray<FString>& InArgs, bool bInPreview);
 	void PlayNextPageCommand(const TArray<FString>& InArgs, bool bInPreview);
 	void StopPageCommand(const TArray<FString>& InArgs, bool bInPreview, bool bInForce);
+	void StopLayersCommand(const TArray<FString>& InArgs, bool bInPreview, bool bInForce);
 	void TakeToProgramCommand(const TArray<FString>& InArgs);
 	void StartChannelCommand(const TArray<FString>& InArgs);
 	void StopChannelCommand(const TArray<FString>& InArgs);
