MotionDesign : InteractiveTools
- Fix viewport changed event called every frame causing viewport toolbar to regenerate
- Only updating viewport toolbar when a related property is changed in the settings

diff --git a/Engine/Plugins/Experimental/Avalanche/Source/AvalancheInteractiveTools/Private/EdMode/AvaInteractiveToolsEdModeToolkit.cpp b/Engine/Plugins/Experimental/Avalanche/Source/AvalancheInteractiveTools/Private/EdMode/AvaInteractiveToolsEdModeToolkit.cpp
--- a/Engine/Plugins/Experimental/Avalanche/Source/AvalancheInteractiveTools/Private/EdMode/AvaInteractiveToolsEdModeToolkit.cpp
+++ b/Engine/Plugins/Experimental/Avalanche/Source/AvalancheInteractiveTools/Private/EdMode/AvaInteractiveToolsEdModeToolkit.cpp
@@ -349,12 +349,23 @@ FText FAvaInteractiveToolsEdModeToolkit::GetToolWarningText() const
 
 void FAvaInteractiveToolsEdModeToolkit::OnSettingsChanged(UObject* InSettings, FPropertyChangedEvent& InEvent)
 {
-	MakeViewportOverlayToolbar();
+	if (InEvent.GetMemberPropertyName().ToString().Contains(TEXT("ViewportToolbar")))
+	{
+		MakeViewportOverlayToolbar();
+	}
 }
 
 void FAvaInteractiveToolsEdModeToolkit::OnViewportChanged(TSharedPtr<IAssetViewport> InOldViewport, TSharedPtr<IAssetViewport> InNewViewport)
 {
-	MakeViewportOverlayToolbar();
+	if (!InNewViewport.IsValid())
+	{
+		return;
+	}
+
+	if (InOldViewport != InNewViewport)
+	{
+		MakeViewportOverlayToolbar();
+	}
 }
 
 #undef LOCTEXT_NAMESPACE
