MotionDesign : InteractiveTools
- Only updating viewport toolbar widget when a related property is changed in the settings

diff --git a/Engine/Plugins/Experimental/Avalanche/Source/AvalancheInteractiveTools/Private/AvaInteractiveToolsSettings.cpp b/Engine/Plugins/Experimental/Avalanche/Source/AvalancheInteractiveTools/Private/AvaInteractiveToolsSettings.cpp
--- a/Engine/Plugins/Experimental/Avalanche/Source/AvalancheInteractiveTools/Private/AvaInteractiveToolsSettings.cpp
+++ b/Engine/Plugins/Experimental/Avalanche/Source/AvalancheInteractiveTools/Private/AvaInteractiveToolsSettings.cpp
@@ -8,6 +8,12 @@
 #include "ISettingsModule.h"
 #include "Modules/ModuleManager.h"
 
+bool UAvaInteractiveToolsSettings::IsViewportToolbarProperty(FName InPropertyName)
+{
+	return InPropertyName == GET_MEMBER_NAME_CHECKED(UAvaInteractiveToolsSettings, ViewportToolbarPosition)
+		|| InPropertyName == GET_MEMBER_NAME_CHECKED(UAvaInteractiveToolsSettings, bViewportToolbarLabelEnabled);
+}
+
 UAvaInteractiveToolsSettings* UAvaInteractiveToolsSettings::Get()
 {
 	return GetMutableDefault<UAvaInteractiveToolsSettings>();
diff --git a/Engine/Plugins/Experimental/Avalanche/Source/AvalancheInteractiveTools/Private/EdMode/AvaInteractiveToolsEdModeToolkit.cpp b/Engine/Plugins/Experimental/Avalanche/Source/AvalancheInteractiveTools/Private/EdMode/AvaInteractiveToolsEdModeToolkit.cpp
--- a/Engine/Plugins/Experimental/Avalanche/Source/AvalancheInteractiveTools/Private/EdMode/AvaInteractiveToolsEdModeToolkit.cpp
+++ b/Engine/Plugins/Experimental/Avalanche/Source/AvalancheInteractiveTools/Private/EdMode/AvaInteractiveToolsEdModeToolkit.cpp
@@ -349,7 +349,7 @@ FText FAvaInteractiveToolsEdModeToolkit::GetToolWarningText() const
 
 void FAvaInteractiveToolsEdModeToolkit::OnSettingsChanged(UObject* InSettings, FPropertyChangedEvent& InEvent)
 {
-	if (InEvent.GetMemberPropertyName().ToString().Contains(TEXT("ViewportToolbar")))
+	if (UAvaInteractiveToolsSettings::IsViewportToolbarProperty(InEvent.GetMemberPropertyName()))
 	{
 		MakeViewportOverlayToolbar();
 	}
diff --git a/Engine/Plugins/Experimental/Avalanche/Source/AvalancheInteractiveTools/Public/AvaInteractiveToolsSettings.h b/Engine/Plugins/Experimental/Avalanche/Source/AvalancheInteractiveTools/Public/AvaInteractiveToolsSettings.h
--- a/Engine/Plugins/Experimental/Avalanche/Source/AvalancheInteractiveTools/Public/AvaInteractiveToolsSettings.h
+++ b/Engine/Plugins/Experimental/Avalanche/Source/AvalancheInteractiveTools/Public/AvaInteractiveToolsSettings.h
@@ -33,6 +33,8 @@ class AVALANCHEINTERACTIVETOOLS_API UAvaInteractiveToolsSettings : public UDevel
 	GENERATED_BODY()
 
 public:
+	static bool IsViewportToolbarProperty(FName InPropertyName);
+	
 	static UAvaInteractiveToolsSettings* Get();
 
 	UAvaInteractiveToolsSettings();
