[Motion Design] Rundown Editor Template Page List: fixing the asset selector to properly refresh when reordering templates by drag and drop.

The ConstructRowWidget is called once and widgets are not reconstructed again. What changes is the PageView is just refreshed with RefreshPageId and existing widgets will reflect the modified view.

The approach of constructing different widgets doesn't work with this system. We need to create all the possible widgets and have them invisible if not supposed to be seen.

Code originally introduced in CL 30466870 for combo templates.

diff --git a/Engine/Plugins/Experimental/Avalanche/Source/AvalancheMediaEditor/Private/Rundown/Pages/Columns/AvaRundownPageAssetSelectorColumn.cpp b/Engine/Plugins/Experimental/Avalanche/Source/AvalancheMediaEditor/Private/Rundown/Pages/Columns/AvaRundownPageAssetSelectorColumn.cpp
--- a/Engine/Plugins/Experimental/Avalanche/Source/AvalancheMediaEditor/Private/Rundown/Pages/Columns/AvaRundownPageAssetSelectorColumn.cpp
+++ b/Engine/Plugins/Experimental/Avalanche/Source/AvalancheMediaEditor/Private/Rundown/Pages/Columns/AvaRundownPageAssetSelectorColumn.cpp
@@ -62,24 +62,46 @@ TSharedRef<SWidget> FAvaRundownPageAssetSelectorColumn::ConstructRowWidget(const
 {
 	const UAvaRundown* Rundown = InPageView->GetRundown();
 
-	// Combo templates don't have an asset selector.
-	// Put a place holder instead.
-	if (InPageView->IsComboTemplate())
-	{
-		return SNew(STextBlock)
-		.Text(InPageView, &IAvaRundownPageView::GetObjectNames, Rundown)
-		.Font(FAppStyle::GetFontStyle(TEXT("PropertyWindow.NormalFont")));
-	}
-	
-	return SNew(SComboButton)
-			.OnGetMenuContent_Lambda([InPageView]{ return UE::AvaPageAssetSelectorColumn::Private::GetAssetPicker(InPageView);})
-			.ContentPadding(FMargin(2.0f, 2.0f))
-			.ButtonContent()
-			[
-				SNew(STextBlock)
-				.Text(InPageView, &IAvaRundownPageView::GetObjectName, Rundown)
-				.Font(FAppStyle::GetFontStyle(TEXT("PropertyWindow.NormalFont")))
-			];
+	const FAvaRundownPageViewWeak PageViewWeak = InPageView;
+
+	// The row widget is not reconstructed when page view changes.
+	// So we need to create all the possible widgets and have them switch visibility.
+	return SNew(SHorizontalBox)
+		+SHorizontalBox::Slot()
+		.AutoWidth()
+		[
+			// Combo templates don't have an asset selector, use a text bock with asset names instead.
+			SNew(STextBlock)
+			.Visibility_Static(&FAvaRundownPageAssetSelectorColumn::GetAssetNameVisibility, PageViewWeak)
+			.Text(InPageView, &IAvaRundownPageView::GetObjectNames, Rundown)
+			.Font(FAppStyle::GetFontStyle(TEXT("PropertyWindow.NormalFont")))
+		]
+		+SHorizontalBox::Slot()
+		.FillWidth(1.0f)
+		[
+			SNew(SComboButton)
+				.Visibility_Static(&FAvaRundownPageAssetSelectorColumn::GetAssetSelectorVisibility, PageViewWeak)
+				.OnGetMenuContent_Lambda([InPageView]{ return UE::AvaPageAssetSelectorColumn::Private::GetAssetPicker(InPageView);})
+				.ContentPadding(FMargin(2.0f, 2.0f))
+				.ButtonContent()
+				[
+					SNew(STextBlock)
+					.Text(InPageView, &IAvaRundownPageView::GetObjectName, Rundown)
+					.Font(FAppStyle::GetFontStyle(TEXT("PropertyWindow.NormalFont")))
+				]
+		];
+}
+
+EVisibility FAvaRundownPageAssetSelectorColumn::GetAssetSelectorVisibility(const FAvaRundownPageViewWeak InPageViewWeak)
+{
+	const TSharedPtr<IAvaRundownPageView> PageView = InPageViewWeak.Pin();
+	return PageView->IsComboTemplate() ? EVisibility::Collapsed : EVisibility::Visible;
+}
+
+EVisibility FAvaRundownPageAssetSelectorColumn::GetAssetNameVisibility(const FAvaRundownPageViewWeak InPageViewWeak)
+{
+	const TSharedPtr<IAvaRundownPageView> PageView = InPageViewWeak.Pin();
+	return !PageView->IsComboTemplate() ? EVisibility::Collapsed : EVisibility::Visible;
 }
 
 #undef LOCTEXT_NAMESPACE
diff --git a/Engine/Plugins/Experimental/Avalanche/Source/AvalancheMediaEditor/Private/Rundown/Pages/Columns/AvaRundownPageAssetSelectorColumn.h b/Engine/Plugins/Experimental/Avalanche/Source/AvalancheMediaEditor/Private/Rundown/Pages/Columns/AvaRundownPageAssetSelectorColumn.h
--- a/Engine/Plugins/Experimental/Avalanche/Source/AvalancheMediaEditor/Private/Rundown/Pages/Columns/AvaRundownPageAssetSelectorColumn.h
+++ b/Engine/Plugins/Experimental/Avalanche/Source/AvalancheMediaEditor/Private/Rundown/Pages/Columns/AvaRundownPageAssetSelectorColumn.h
@@ -15,4 +15,7 @@ public:
 	virtual SHeaderRow::FColumn::FArguments ConstructHeaderRowColumn() override;
 	virtual TSharedRef<SWidget> ConstructRowWidget(const FAvaRundownPageViewRef& InPageView, const TSharedPtr<SAvaRundownPageViewRow>& InRow) override;
 	//~ End IAvaRundownPageViewColumn
+
+	static EVisibility GetAssetSelectorVisibility(const FAvaRundownPageViewWeak InPageViewWeak);
+	static EVisibility GetAssetNameVisibility(const FAvaRundownPageViewWeak InPageViewWeak);
 };
