MotionDesign : Mask
- Fix LOCTEXT issue 836617 with mask read and write modifier using name namespace

diff --git a/Engine/Plugins/Experimental/Avalanche/Source/AvalancheMask/Private/Mask2D/AvaMask2DReadModifier.cpp b/Engine/Plugins/Experimental/Avalanche/Source/AvalancheMask/Private/Mask2D/AvaMask2DReadModifier.cpp
--- a/Engine/Plugins/Experimental/Avalanche/Source/AvalancheMask/Private/Mask2D/AvaMask2DReadModifier.cpp
+++ b/Engine/Plugins/Experimental/Avalanche/Source/AvalancheMask/Private/Mask2D/AvaMask2DReadModifier.cpp
@@ -17,7 +17,7 @@
 #include "Materials/MaterialInstanceDynamic.h"
 #include "Modifiers/ActorModifierCoreStack.h"
 
-#define LOCTEXT_NAMESPACE "AvaMask2DModifier"
+#define LOCTEXT_NAMESPACE "AvaMask2DReadModifier"
 
 #if WITH_EDITOR
 const TAvaPropertyChangeDispatcher<UAvaMask2DReadModifier> UAvaMask2DReadModifier::PropertyChangeDispatcher =
diff --git a/Engine/Plugins/Experimental/Avalanche/Source/AvalancheMask/Private/Mask2D/AvaMask2DWriteModifier.cpp b/Engine/Plugins/Experimental/Avalanche/Source/AvalancheMask/Private/Mask2D/AvaMask2DWriteModifier.cpp
--- a/Engine/Plugins/Experimental/Avalanche/Source/AvalancheMask/Private/Mask2D/AvaMask2DWriteModifier.cpp
+++ b/Engine/Plugins/Experimental/Avalanche/Source/AvalancheMask/Private/Mask2D/AvaMask2DWriteModifier.cpp
@@ -13,7 +13,7 @@
 #include "Materials/MaterialInstanceDynamic.h"
 #include "Modifiers/ActorModifierCoreStack.h"
 
-#define LOCTEXT_NAMESPACE "AvaMask2DModifier"
+#define LOCTEXT_NAMESPACE "AvaMask2DWriteModifier"
 
 #if WITH_EDITOR
 const TAvaPropertyChangeDispatcher<UAvaMask2DWriteModifier> UAvaMask2DWriteModifier::PropertyChangeDispatcher =
