Motion Design: fix crashes (check fail) when calling gizmo toggle ufunction on components/objects that are being destroyed (unreachable).

diff --git a/Engine/Plugins/Experimental/Avalanche/Source/Avalanche/Private/Framework/AvaGizmoComponent.cpp b/Engine/Plugins/Experimental/Avalanche/Source/Avalanche/Private/Framework/AvaGizmoComponent.cpp
--- a/Engine/Plugins/Experimental/Avalanche/Source/Avalanche/Private/Framework/AvaGizmoComponent.cpp
+++ b/Engine/Plugins/Experimental/Avalanche/Source/Avalanche/Private/Framework/AvaGizmoComponent.cpp
@@ -208,8 +208,9 @@ void UAvaGizmoComponent::RestoreComponentValues()
 		bIsRestoringValues = false;
 	};
 
-	if (AActor* Owner = GetOwner();
-		IsValid(Owner))
+	AActor* Owner = GetOwner();
+
+	if (IsValid(Owner) && !Owner->IsActorBeingDestroyed() && !Owner->IsUnreachable())
 	{
 		bAreGizmoValuesApplied = false;
 
@@ -218,7 +219,7 @@ void UAvaGizmoComponent::RestoreComponentValues()
 
 		for (UPrimitiveComponent* Component : PrimitiveComponents)
 		{
-			if (!IsValid(Component))
+			if (!IsValid(Component) || Component->IsBeingDestroyed() || Component->IsUnreachable())
 			{
 				continue;
 			}
@@ -275,7 +276,7 @@ void UAvaGizmoComponent::RestoreComponentValues()
 			TArray<UActorComponent*> GizmoObjectComponents = Owner->GetComponentsByInterface(UAvaGizmoObjectInterface::StaticClass());
 			for (UActorComponent* Component : GizmoObjectComponents)
 			{
-				if (IsValid(Component))
+				if (IsValid(Component) && !Component->IsBeingDestroyed() && !Component->IsUnreachable())
 				{
 					IAvaGizmoObjectInterface::Execute_ToggleGizmo(Component, this, false);	
 				}
