[Motion Design] Fix Rundown Editor crash when switching out of tabs.

Regression introduced in CL 41008205.

diff --git a/Engine/Plugins/Experimental/Avalanche/Source/AvalancheMediaEditor/Private/Rundown/Pages/Slate/SAvaRundownPageList.cpp b/Engine/Plugins/Experimental/Avalanche/Source/AvalancheMediaEditor/Private/Rundown/Pages/Slate/SAvaRundownPageList.cpp
--- a/Engine/Plugins/Experimental/Avalanche/Source/AvalancheMediaEditor/Private/Rundown/Pages/Slate/SAvaRundownPageList.cpp
+++ b/Engine/Plugins/Experimental/Avalanche/Source/AvalancheMediaEditor/Private/Rundown/Pages/Slate/SAvaRundownPageList.cpp
@@ -1279,7 +1279,7 @@ void SAvaRundownPageList::OnActiveTabChanged(TSharedPtr<SDockTab> InPreviouslyAc
 {
 	if (const TSharedPtr<FAvaRundownEditor> RundownEditor = RundownEditorWeak.Pin())
 	{
-		if (RundownEditor->GetAssociatedTabManager() == InNewlyActivated->GetTabManagerPtr())
+		if (InNewlyActivated && RundownEditor->GetAssociatedTabManager() == InNewlyActivated->GetTabManagerPtr())
 		{
 			const TSharedPtr<SWidget> Content = InNewlyActivated->GetContent();
 			if (Content.Get() == this)
