[Motion Design] Fix crash when adding broadcast profile.

Repro:
- Add new profile
- Select  previously existing profile
- Start channels in the previous profile.

The channels where keeping a pointer to the parent profile. When a new profile is added and the Map is resized, the pointers become invalid. It is not a terribly good idea to keep pointers to array or map members like that. So instead, we are goign to index them by name, like we do for the channels.

diff --git a/Engine/Plugins/Experimental/Avalanche/Source/AvalancheMedia/Private/Broadcast/AvaBroadcast.cpp b/Engine/Plugins/Experimental/Avalanche/Source/AvalancheMedia/Private/Broadcast/AvaBroadcast.cpp
--- a/Engine/Plugins/Experimental/Avalanche/Source/AvalancheMedia/Private/Broadcast/AvaBroadcast.cpp
+++ b/Engine/Plugins/Experimental/Avalanche/Source/AvalancheMedia/Private/Broadcast/AvaBroadcast.cpp
@@ -265,7 +265,7 @@ bool UAvaBroadcast::RenameProfile(FName InProfileName, FName InNewProfileName)
 	{
 		FAvaBroadcastProfile Profile = MoveTemp(Profiles[InProfileName]);
 		Profiles.Remove(InProfileName);
-		Profile.ProfileName = InNewProfileName;
+		Profile.SetProfileName(InNewProfileName);
 		Profiles.Add(InNewProfileName, MoveTemp(Profile));
 		
 		if (CurrentProfile == InProfileName)
@@ -700,7 +700,7 @@ void UAvaBroadcast::UpdateProfileNames()
 {
 	for (TPair<FName, FAvaBroadcastProfile>& Pair : Profiles)
 	{
-		Pair.Value.ProfileName = Pair.Key;
+		Pair.Value.SetProfileName(Pair.Key);
 	}
 }
 
diff --git a/Engine/Plugins/Experimental/Avalanche/Source/AvalancheMedia/Private/Broadcast/AvaBroadcastProfile.cpp b/Engine/Plugins/Experimental/Avalanche/Source/AvalancheMedia/Private/Broadcast/AvaBroadcastProfile.cpp
--- a/Engine/Plugins/Experimental/Avalanche/Source/AvalancheMedia/Private/Broadcast/AvaBroadcastProfile.cpp
+++ b/Engine/Plugins/Experimental/Avalanche/Source/AvalancheMedia/Private/Broadcast/AvaBroadcastProfile.cpp
@@ -38,12 +38,24 @@ UAvaBroadcast& FAvaBroadcastProfile::GetBroadcast() const
 	return ensure(ParentBroadcast) ? *ParentBroadcast : UAvaBroadcast::Get();
 }
 
+void FAvaBroadcastProfile::SetProfileName(FName InProfileName)
+{
+	if (ProfileName != InProfileName)
+	{
+		ProfileName = InProfileName;
+		for (FAvaBroadcastOutputChannel& Channel : Channels)
+		{
+			Channel.SetProfileName(InProfileName);
+		}
+	}
+}
+
 void FAvaBroadcastProfile::CopyProfiles(const FAvaBroadcastProfile& InSourceProfile, FAvaBroadcastProfile& OutTargetProfile)
 {
 	OutTargetProfile.Channels.Empty(InSourceProfile.Channels.Num());
 	for (const FAvaBroadcastOutputChannel& SourceChannel : InSourceProfile.Channels)
 	{
-		FAvaBroadcastOutputChannel& TargetChannel = OutTargetProfile.Channels.Add_GetRef(FAvaBroadcastOutputChannel(&OutTargetProfile));
+		FAvaBroadcastOutputChannel& TargetChannel = OutTargetProfile.Channels.Add_GetRef(FAvaBroadcastOutputChannel(OutTargetProfile.GetName(), &OutTargetProfile.GetBroadcast()));
 		TargetChannel.SetChannelIndex(SourceChannel.GetChannelIndex());
 		
 		FAvaBroadcastOutputChannel::DuplicateChannel(SourceChannel, TargetChannel);
@@ -123,7 +135,7 @@ void FAvaBroadcastProfile::PostLoadProfile(bool bInIsProfileActive, UAvaBroadcas
 	for (FAvaBroadcastOutputChannel& Channel : Channels)
 	{
 		UE::AvaBroadcastProfile::Private::ConditionalSetChannelIndex(Channel, Index++);
-		Channel.PostLoadMediaOutputs(bInIsProfileActive, this);
+		Channel.PostLoadMediaOutputs(bInIsProfileActive, ProfileName, InBroadcast);
 	}
 
 	ResolveChannels();
@@ -191,7 +203,7 @@ FAvaBroadcastOutputChannel& FAvaBroadcastProfile::AddChannel(FName InChannelName
 		Broadcast.GetOrAddChannelName(ChannelNameIndex);
 	}
 
-	FAvaBroadcastOutputChannel& Channel = Channels.Add_GetRef(FAvaBroadcastOutputChannel(this));
+	FAvaBroadcastOutputChannel& Channel = Channels.Add_GetRef(FAvaBroadcastOutputChannel(ProfileName, &GetBroadcast()));
 	Channel.SetChannelIndex(ChannelNameIndex);
 	Channel.UpdateChannelResources(Broadcast.GetCurrentProfileName() == GetName());
 	
diff --git a/Engine/Plugins/Experimental/Avalanche/Source/AvalancheMedia/Private/Broadcast/Channel/AvaBroadcastOutputChannel.cpp b/Engine/Plugins/Experimental/Avalanche/Source/AvalancheMedia/Private/Broadcast/Channel/AvaBroadcastOutputChannel.cpp
--- a/Engine/Plugins/Experimental/Avalanche/Source/AvalancheMedia/Private/Broadcast/Channel/AvaBroadcastOutputChannel.cpp
+++ b/Engine/Plugins/Experimental/Avalanche/Source/AvalancheMedia/Private/Broadcast/Channel/AvaBroadcastOutputChannel.cpp
@@ -67,6 +67,21 @@ FAvaBroadcastOutputChannel::FAvaBroadcastOutputChannel(ENoInit NoInit)
 {
 }
 
+FAvaBroadcastOutputChannel::FAvaBroadcastOutputChannel(FAvaBroadcastProfile* InProfile)
+	: FAvaBroadcastOutputChannel(
+		InProfile ? InProfile->GetName() : NAME_None,
+		InProfile ? &InProfile->GetBroadcast() : &UAvaBroadcast::Get())
+{
+	
+}
+
+FAvaBroadcastOutputChannel::FAvaBroadcastOutputChannel(FName InProfileName, UAvaBroadcast* InBroadcast)
+	: ProfileName(InProfileName)
+	, ParentBroadcastWeak(InBroadcast)
+{
+	
+}
+
 namespace UE::AvaMedia::Private
 {
 	EMediaCaptureOverrunAction GetCaptureOverrunAction()
@@ -276,18 +291,18 @@ EAvaBroadcastChannelType FAvaBroadcastOutputChannel::GetChannelType() const
 
 FName FAvaBroadcastOutputChannel::GetProfileName() const
 {
-	// Note: during the migration towards a non-singleton UAvaBroadcast,
-	// a fallback path to the singleton is kept, but it has an ensure to try and catch
-	// the code paths that lead to this.
-	return ensure(Profile) ? Profile->GetName() : UAvaBroadcast::Get().GetCurrentProfile().GetName();
+	return ProfileName;
 }
 
 FAvaBroadcastProfile& FAvaBroadcastOutputChannel::GetProfile() const
 {
-	// Note: during the migration towards a non-singleton UAvaBroadcast,
+	// Note: During the migration towards a non-singleton UAvaBroadcast,
 	// a fallback path to the singleton is kept, but it has an ensure to try and catch
 	// the code paths that lead to this.
-	return ensure(Profile) ? *Profile : UAvaBroadcast::Get().GetCurrentProfile();
+	ensure(ParentBroadcastWeak.IsValid()); 
+	UAvaBroadcast& Broadcast = ParentBroadcastWeak.IsValid() ? *ParentBroadcastWeak.Get() : UAvaBroadcast::Get();
+	FAvaBroadcastProfile& Profile = Broadcast.GetProfile(ProfileName);
+	return Profile.IsValidProfile() ? Profile : Broadcast.GetCurrentProfile();
 }
 
 EAvaBroadcastChannelState FAvaBroadcastOutputChannel::RefreshState()
@@ -668,8 +683,7 @@ bool FAvaBroadcastOutputChannel::StartChannelBroadcast()
 		const TArray<UMediaOutput*> RemoteMediaOutputs = GetRemoteMediaOutputs();
 		if (!RemoteMediaOutputs.IsEmpty())
 		{
-			const FString ProfileName = GetProfileName().ToString();
-			Module.GetPlaybackClient().RequestBroadcast(ProfileName, GetChannelName(), RemoteMediaOutputs, EAvaBroadcastAction::Start);
+			Module.GetPlaybackClient().RequestBroadcast(ProfileName.ToString(), GetChannelName(), RemoteMediaOutputs, EAvaBroadcastAction::Start);
 		}
 	}
 
@@ -774,8 +788,7 @@ void FAvaBroadcastOutputChannel::StopChannelBroadcast()
 		const TArray<UMediaOutput*> RemoteMediaOutputs = GetRemoteMediaOutputs();
 		if (!RemoteMediaOutputs.IsEmpty())
 		{
-			const FString ProfileName = GetProfileName().ToString();
-			Module.GetPlaybackClient().RequestBroadcast(ProfileName, GetChannelName(), RemoteMediaOutputs, EAvaBroadcastAction::Stop);
+			Module.GetPlaybackClient().RequestBroadcast(ProfileName.ToString(), GetChannelName(), RemoteMediaOutputs, EAvaBroadcastAction::Stop);
 		}
 	}
 
@@ -903,9 +916,8 @@ void FAvaBroadcastOutputChannel::OnMediaOutputModified(UMediaOutput* InMediaOutp
 	IAvaMediaModule& AvaMediaModule = IAvaMediaModule::Get();
 	if (IsMediaOutputRemote(InMediaOutput) && AvaMediaModule.IsPlaybackClientStarted())
 	{
-		const FString ProfileName = GetProfileName().ToString();
 		IAvaPlaybackClient& PlaybackClient = AvaMediaModule.GetPlaybackClient();
-		PlaybackClient.RequestBroadcast(ProfileName, GetChannelName(), GetRemoteMediaOutputs(), EAvaBroadcastAction::UpdateConfig);
+		PlaybackClient.RequestBroadcast(ProfileName.ToString(), GetChannelName(), GetRemoteMediaOutputs(), EAvaBroadcastAction::UpdateConfig);
 	}
 	else
 	{
@@ -914,9 +926,21 @@ void FAvaBroadcastOutputChannel::OnMediaOutputModified(UMediaOutput* InMediaOutp
 	}
 }
 
+void FAvaBroadcastOutputChannel::SetProfileName(FName InProfileName)
+{
+	ProfileName = InProfileName;	
+}
+
 void FAvaBroadcastOutputChannel::PostLoadMediaOutputs(bool bInIsProfileActive, FAvaBroadcastProfile* InProfile)
 {
-	Profile = InProfile;
+	PostLoadMediaOutputs(bInIsProfileActive, InProfile ? InProfile->GetName() : NAME_None,
+		InProfile ? &InProfile->GetBroadcast() : &UAvaBroadcast::Get());
+}
+
+void FAvaBroadcastOutputChannel::PostLoadMediaOutputs(bool bInIsProfileActive, FName InProfileName, UAvaBroadcast* InBroadcast)
+{
+	ProfileName = InProfileName;
+	ParentBroadcastWeak = InBroadcast;
 	
 	// Create missing output infos (legacy support)
 	if (MediaOutputInfos.Num() != MediaOutputs.Num())
diff --git a/Engine/Plugins/Experimental/Avalanche/Source/AvalancheMedia/Public/Broadcast/AvaBroadcastProfile.h b/Engine/Plugins/Experimental/Avalanche/Source/AvalancheMedia/Public/Broadcast/AvaBroadcastProfile.h
--- a/Engine/Plugins/Experimental/Avalanche/Source/AvalancheMedia/Public/Broadcast/AvaBroadcastProfile.h
+++ b/Engine/Plugins/Experimental/Avalanche/Source/AvalancheMedia/Public/Broadcast/AvaBroadcastProfile.h
@@ -23,6 +23,8 @@ public:
 
 	FName GetName() const { return ProfileName; }
 	UAvaBroadcast& GetBroadcast() const;
+
+	void SetProfileName(FName InProfileName);
 	
 	static void CopyProfiles(const FAvaBroadcastProfile& InSourceProfile, FAvaBroadcastProfile& OutTargetProfile);
 	
diff --git a/Engine/Plugins/Experimental/Avalanche/Source/AvalancheMedia/Public/Broadcast/Channel/AvaBroadcastOutputChannel.h b/Engine/Plugins/Experimental/Avalanche/Source/AvalancheMedia/Public/Broadcast/Channel/AvaBroadcastOutputChannel.h
--- a/Engine/Plugins/Experimental/Avalanche/Source/AvalancheMedia/Public/Broadcast/Channel/AvaBroadcastOutputChannel.h
+++ b/Engine/Plugins/Experimental/Avalanche/Source/AvalancheMedia/Public/Broadcast/Channel/AvaBroadcastOutputChannel.h
@@ -18,6 +18,7 @@ class FWidgetRenderer;
 class SAvaBroadcastPlaceholderWidget;
 class SVirtualWindow;
 class SWidget;
+class UAvaBroadcast;
 class UAvaPlayableGroup;
 class UMediaCapture;
 class UMediaOutput;
@@ -70,7 +71,11 @@ struct AVALANCHEMEDIA_API FAvaBroadcastOutputChannel
 public:
 	FAvaBroadcastOutputChannel() = default;
 	FAvaBroadcastOutputChannel(ENoInit NoInit);
-	explicit FAvaBroadcastOutputChannel(FAvaBroadcastProfile* InProfile) : Profile(InProfile) {}
+
+	UE_DEPRECATED(5.6, "Use constructor with profile name and broadcast.")
+	explicit FAvaBroadcastOutputChannel(FAvaBroadcastProfile* InProfile);
+	
+	explicit FAvaBroadcastOutputChannel(FName InProfileName, UAvaBroadcast* InBroadcast);
 	~FAvaBroadcastOutputChannel();
 
 	void ReleasePlaceholderResources();
@@ -137,9 +142,14 @@ public:
 	 * Propagates the configuration change accordingly.
 	 */
 	void OnMediaOutputModified(UMediaOutput* InMediaOutput);
+
+	void SetProfileName(FName InProfileName);
 	
+	UE_DEPRECATED(5.6, "Use version with profile name and broadcast.")
 	void PostLoadMediaOutputs(bool bInIsProfileActive, FAvaBroadcastProfile* InProfile);
-
+	
+	void PostLoadMediaOutputs(bool bInIsProfileActive, FName InProfileName, UAvaBroadcast* InBroadcast);
+	
 	void UpdateRenderTarget(UAvaPlayableGroup* InPlayableGroup, UTextureRenderTarget2D* InRenderTarget);
 	void UpdateAudioDevice(const FAudioDeviceHandle& InAudioDeviceHandle);
 	UAvaPlayableGroup* GetLastActivePlayableGroup() const;
@@ -246,8 +256,12 @@ protected:
 	// Extra render target to perform alpha inversion needed to mimic scene rendering.
 	UPROPERTY(Transient)
 	TObjectPtr<UTextureRenderTarget2D> PlaceholderRenderTargetTmp;
-	
-	FAvaBroadcastProfile* Profile = nullptr;
+
+	/** Profile name owning this channel. */
+	FName ProfileName;
+
+	/** Broadcast object owning this channel. */
+	TWeakObjectPtr<UAvaBroadcast> ParentBroadcastWeak;
 	
 	/**
 	 * Last Active Playable group associated with the current render target.
