Motion Design: fix outliner crash when switching levels under particular circumstances. The previous fix didn't fully reset all the shared references from SListView/STreeView.

diff --git a/Engine/Plugins/Experimental/Avalanche/Source/AvalancheOutliner/Private/AvaOutlinerView.cpp b/Engine/Plugins/Experimental/Avalanche/Source/AvalancheOutliner/Private/AvaOutlinerView.cpp
--- a/Engine/Plugins/Experimental/Avalanche/Source/AvalancheOutliner/Private/AvaOutlinerView.cpp
+++ b/Engine/Plugins/Experimental/Avalanche/Source/AvalancheOutliner/Private/AvaOutlinerView.cpp
@@ -32,6 +32,7 @@
 #include "ToolMenu.h"
 #include "ToolMenuContext/AvaOutlinerItemsContext.h"
 #include "ToolMenus.h"
+#include "Widgets/Layout/SBorder.h"
 #include "Widgets/SNullWidget.h"
 
 #define LOCTEXT_NAMESPACE "AvaOutlinerView"
@@ -46,7 +47,16 @@ FAvaOutlinerView::FAvaOutlinerView(FPrivateToken)
 FAvaOutlinerView::~FAvaOutlinerView()
 {
 	SaveState();
-	ClearRootItemSource();
+
+	if (OutlinerWidgetContainer.IsValid())
+	{
+		OutlinerWidgetContainer->ClearContent();
+
+		const int32 OutlinerWidgetReferenceCount = OutlinerWidget.GetSharedReferenceCount();
+		ensureMsgf(OutlinerWidgetReferenceCount == 1, TEXT("Outliner Widget expected to have 1 shared reference count but had %d"), OutlinerWidgetReferenceCount);
+	}
+
+	OutlinerWidget.Reset();
 
 	if (UObjectInitialized())
 	{
@@ -82,10 +92,16 @@ void FAvaOutlinerView::Init(const TSharedRef<FAvaOutliner>& InOutliner, bool bCr
 	UpdateCustomFilters();
 
 	TextFilter->OnChanged().AddSP(this, &FAvaOutlinerView::OnFilterChanged);
-	
+
 	if (bCreateOutlinerWidget)
 	{
-		OutlinerWidget = SNew(SAvaOutliner, SharedThis(this));
+		OutlinerWidgetContainer = SNew(SBorder)
+			.BorderImage(FAppStyle::GetBrush("NoBorder"))
+			.Padding(0)
+			[
+				SAssignNew(OutlinerWidget, SAvaOutliner, SharedThis(this))
+			];
+
 		CreateToolbar();
 	}
 
@@ -355,9 +371,9 @@ bool FAvaOutlinerView::IsMostRecentOutlinerView() const
 
 TSharedRef<SWidget> FAvaOutlinerView::GetOutlinerWidget() const
 {
-	if (OutlinerWidget.IsValid())
+	if (OutlinerWidgetContainer.IsValid())
 	{
-		return OutlinerWidget.ToSharedRef();
+		return OutlinerWidgetContainer.ToSharedRef();
 	}
 	return SNullWidget::NullWidget;
 }
@@ -582,21 +598,6 @@ void FAvaOutlinerView::UpdateRootVisibleItems()
 	OutlinerStats->MarkCountTypeDirty(EAvaOutlinerStatCountType::VisibleItemCount);
 }
 
-void FAvaOutlinerView::ClearRootItemSource()
-{
-	RootVisibleItems.Reset();
-	ReadOnlyItems.Reset();
-	OutlinerStats->MarkCountTypeDirty(EAvaOutlinerStatCountType::VisibleItemCount);
-
-	if (OutlinerWidget.IsValid())
-	{
-		if (TSharedPtr<SAvaOutlinerTreeView> TreeView = OutlinerWidget->GetTreeView())
-		{
-			TreeView->ClearItems();
-		}
-	}
-}
-
 void FAvaOutlinerView::UpdateItemExpansions()
 {
 	TArray<FAvaOutlinerItemPtr> Items = RootVisibleItems;
diff --git a/Engine/Plugins/Experimental/Avalanche/Source/AvalancheOutliner/Private/AvaOutlinerView.h b/Engine/Plugins/Experimental/Avalanche/Source/AvalancheOutliner/Private/AvaOutlinerView.h
--- a/Engine/Plugins/Experimental/Avalanche/Source/AvalancheOutliner/Private/AvaOutlinerView.h
+++ b/Engine/Plugins/Experimental/Avalanche/Source/AvalancheOutliner/Private/AvaOutlinerView.h
@@ -14,6 +14,7 @@ class FUICommandList;
 class IAvaOutlinerColumn;
 class IAvaOutlinerItemFilter;
 class SAvaOutliner;
+class SBorder;
 class SHorizontalBox;
 class SWidget;
 class UToolMenu;
@@ -143,8 +144,6 @@ public:
 	/** Refreshes the Items that will be at the Top Level of the Tree (e.g. Actors with no parent item) */
 	void UpdateRootVisibleItems();
 
-	void ClearRootItemSource();
-
 	/** Updates the Item Expansions in the Tree Widget based on whether the Items have the Expanded flag or not */
 	void UpdateItemExpansions();
 
@@ -416,6 +415,9 @@ private:
 	/** Weak pointer to the outliner this is a view of */
 	TWeakPtr<FAvaOutliner> OutlinerWeak;
 
+	/** Widget wrapping the outliner widget. Used so that the outliner widget's lifetime can be controlled by this container */
+	TSharedPtr<SBorder> OutlinerWidgetContainer;
+
 	/** Widget showing the View. Can be null if instanced for testing (i.e. bCreateOutlinerWidget is false in Init) */
 	TSharedPtr<SAvaOutliner> OutlinerWidget;
 
diff --git a/Engine/Plugins/Experimental/Avalanche/Source/AvalancheOutliner/Private/Slate/SAvaOutlinerTreeView.cpp b/Engine/Plugins/Experimental/Avalanche/Source/AvalancheOutliner/Private/Slate/SAvaOutlinerTreeView.cpp
--- a/Engine/Plugins/Experimental/Avalanche/Source/AvalancheOutliner/Private/Slate/SAvaOutlinerTreeView.cpp
+++ b/Engine/Plugins/Experimental/Avalanche/Source/AvalancheOutliner/Private/Slate/SAvaOutlinerTreeView.cpp
@@ -20,13 +20,6 @@ int32 SAvaOutlinerTreeView::GetItemIndex(FAvaOutlinerItemPtr Item) const
 	return INDEX_NONE;
 }
 
-void SAvaOutlinerTreeView::ClearItems()
-{
-	ClearExpandedItems();
-	PreviousSelectedItems.Reset();
-	RebuildList();
-}
-
 void SAvaOutlinerTreeView::FocusOnItem(const FAvaOutlinerItemPtr& InItem)
 {
 	SelectorItem = InItem;
diff --git a/Engine/Plugins/Experimental/Avalanche/Source/AvalancheOutliner/Private/Slate/SAvaOutlinerTreeView.h b/Engine/Plugins/Experimental/Avalanche/Source/AvalancheOutliner/Private/Slate/SAvaOutlinerTreeView.h
--- a/Engine/Plugins/Experimental/Avalanche/Source/AvalancheOutliner/Private/Slate/SAvaOutlinerTreeView.h
+++ b/Engine/Plugins/Experimental/Avalanche/Source/AvalancheOutliner/Private/Slate/SAvaOutlinerTreeView.h
@@ -18,8 +18,6 @@ public:
 
 	int32 GetItemIndex(FAvaOutlinerItemPtr Item) const;
 
-	void ClearItems();
-
 	void FocusOnItem(const FAvaOutlinerItemPtr& InItem);
 
 	void ScrollItemIntoView(const FAvaOutlinerItemPtr& InItem);
