Motion Design: activating/deactivating motion design mode to persist across sessions

diff --git a/Engine/Plugins/Experimental/Avalanche/Source/Avalanche/Private/AvaScene.cpp b/Engine/Plugins/Experimental/Avalanche/Source/Avalanche/Private/AvaScene.cpp
--- a/Engine/Plugins/Experimental/Avalanche/Source/Avalanche/Private/AvaScene.cpp
+++ b/Engine/Plugins/Experimental/Avalanche/Source/Avalanche/Private/AvaScene.cpp
@@ -122,6 +122,11 @@ IAvaSequencePlaybackObject* AAvaScene::GetScenePlayback() const
 }
 
 #if WITH_EDITOR
+void AAvaScene::SetAutoStartMode(bool bInAutoStartMode)
+{
+	bAutoStartMode = bInAutoStartMode;
+}
+
 void AAvaScene::OnWorldRenamed(UWorld* InWorld)
 {
 	if (FUObjectThreadContext::Get().IsRoutingPostLoad || InWorld != GetWorld())
diff --git a/Engine/Plugins/Experimental/Avalanche/Source/Avalanche/Public/AvaScene.h b/Engine/Plugins/Experimental/Avalanche/Source/Avalanche/Public/AvaScene.h
--- a/Engine/Plugins/Experimental/Avalanche/Source/Avalanche/Public/AvaScene.h
+++ b/Engine/Plugins/Experimental/Avalanche/Source/Avalanche/Public/AvaScene.h
@@ -45,6 +45,13 @@ public:
 	IAvaSequencePlaybackObject* GetScenePlayback() const;
 
 #if WITH_EDITOR
+	bool ShouldAutoStartMode() const
+	{
+		return bAutoStartMode;
+	}
+
+	AVALANCHE_API void SetAutoStartMode(bool bInAutoStartMode);
+
 	TArray<uint8>& GetOutlinerData() { return OutlinerData; }
 
 	void OnWorldRenamed(UWorld* InWorld);
@@ -132,10 +139,6 @@ protected:
 	UPROPERTY()
 	TScriptInterface<IAvaSequencePlaybackObject> PlaybackObject;
 
-	/** The index to the animation to use as the default animation */
-	UPROPERTY()
-	int32 DefaultSequenceIndex = 0;
-
 	/** A list of only the Root Animations (those without Parent Animations) */
 	UPROPERTY(Transient, DuplicateTransient, TextExportTransient)
 	TArray<TWeakObjectPtr<UAvaSequence>> RootAnimations;
@@ -155,6 +158,17 @@ protected:
 #endif
 
 	FSimpleMulticastDelegate OnTreeAnimationRebuilt;
-	
-	bool bPendingAnimTreeUpdate = false; 
+
+	/** The index to the animation to use as the default animation */
+	UPROPERTY()
+	int32 DefaultSequenceIndex = 0;
+
+	bool bPendingAnimTreeUpdate = false;
+
+#if WITH_EDITORONLY_DATA
+	/** Whether the motion design mode should auto start */
+	UPROPERTY()
+	bool bAutoStartMode = true;
+#endif
+
 };
diff --git a/Engine/Plugins/Experimental/Avalanche/Source/AvalancheEditor/Private/EditorProvider/AvaEditorProvider.cpp b/Engine/Plugins/Experimental/Avalanche/Source/AvalancheEditor/Private/EditorProvider/AvaEditorProvider.cpp
--- a/Engine/Plugins/Experimental/Avalanche/Source/AvalancheEditor/Private/EditorProvider/AvaEditorProvider.cpp
+++ b/Engine/Plugins/Experimental/Avalanche/Source/AvalancheEditor/Private/EditorProvider/AvaEditorProvider.cpp
@@ -26,6 +26,24 @@ UObject* FAvaEditorProvider::GetSceneObject(UWorld* InWorld, EAvaEditorObjectQue
 	return AAvaScene::GetScene(SceneLevel, bCreateSceneIfNotFound);
 }
 
+bool FAvaEditorProvider::ShouldAutoActivateScene(UObject* InSceneObject) const
+{
+	if (AAvaScene* Scene = Cast<AAvaScene>(InSceneObject))
+	{
+		return Scene->ShouldAutoStartMode();
+	}
+	return IAvaEditorProvider::ShouldAutoActivateScene(InSceneObject);
+}
+
+void FAvaEditorProvider::SetAutoActivateScene(UObject* InSceneObject, bool bInAutoActivateScene) const
+{
+	if (AAvaScene* Scene = Cast<AAvaScene>(InSceneObject))
+	{
+		Scene->Modify();
+		Scene->SetAutoStartMode(bInAutoActivateScene);
+	}
+}
+
 void FAvaEditorProvider::GetActorsToEdit(TArray<AActor*>& InOutActorsToEdit) const
 {
 	FAvaEditorActorUtils::GetActorsToEdit(InOutActorsToEdit);
diff --git a/Engine/Plugins/Experimental/Avalanche/Source/AvalancheEditor/Private/EditorProvider/AvaEditorProvider.h b/Engine/Plugins/Experimental/Avalanche/Source/AvalancheEditor/Private/EditorProvider/AvaEditorProvider.h
--- a/Engine/Plugins/Experimental/Avalanche/Source/AvalancheEditor/Private/EditorProvider/AvaEditorProvider.h
+++ b/Engine/Plugins/Experimental/Avalanche/Source/AvalancheEditor/Private/EditorProvider/AvaEditorProvider.h
@@ -10,6 +10,8 @@ class FAvaEditorProvider : public IAvaEditorProvider
 public:
 	//~ Begin IAvaEditorProvider
 	virtual UObject* GetSceneObject(UWorld* InWorld, EAvaEditorObjectQueryType InQueryType) override;
+	virtual bool ShouldAutoActivateScene(UObject* InSceneObject) const override;
+	virtual void SetAutoActivateScene(UObject* InSceneObject, bool bInAutoActivateScene) const override;
 	virtual void GetActorsToEdit(TArray<AActor*>& InOutActorsToEdit) const override;
 	virtual void OnSceneActivated() override;
 	virtual void OnSceneDeactivated() override;
diff --git a/Engine/Plugins/Experimental/Avalanche/Source/AvalancheEditorCore/Private/AvaEditor.cpp b/Engine/Plugins/Experimental/Avalanche/Source/AvalancheEditorCore/Private/AvaEditor.cpp
--- a/Engine/Plugins/Experimental/Avalanche/Source/AvalancheEditorCore/Private/AvaEditor.cpp
+++ b/Engine/Plugins/Experimental/Avalanche/Source/AvalancheEditorCore/Private/AvaEditor.cpp
@@ -338,6 +338,16 @@ bool FAvaEditor::CanActivate() const
 	return !IsActive() && !bActiveStateChanging;
 }
 
+bool FAvaEditor::ShouldAutoActivate() const
+{
+	return Provider->ShouldAutoActivateScene(GetSceneObject(EAvaEditorObjectQueryType::SearchOnly));
+}
+
+void FAvaEditor::SetAutoActivate(bool bInAutoActivate)
+{
+	Provider->SetAutoActivateScene(GetSceneObject(EAvaEditorObjectQueryType::SkipSearch), bInAutoActivate);
+}
+
 bool FAvaEditor::CanDeactivate() const
 {
 	return IsActive() && !bActiveStateChanging;
diff --git a/Engine/Plugins/Experimental/Avalanche/Source/AvalancheEditorCore/Private/AvaEditor.h b/Engine/Plugins/Experimental/Avalanche/Source/AvalancheEditorCore/Private/AvaEditor.h
--- a/Engine/Plugins/Experimental/Avalanche/Source/AvalancheEditorCore/Private/AvaEditor.h
+++ b/Engine/Plugins/Experimental/Avalanche/Source/AvalancheEditorCore/Private/AvaEditor.h
@@ -43,6 +43,8 @@ protected:
 	virtual void Cleanup() override;
 	virtual bool IsActive() const { return bIsActive; }
 	virtual bool CanActivate() const override;
+	virtual bool ShouldAutoActivate() const override;
+	virtual void SetAutoActivate(bool bInAutoActivate) override;
 	virtual bool CanDeactivate() const override;
 	virtual void SetToolkitHost(TSharedRef<IToolkitHost> InToolkitHost) override;
 	virtual void BindCommands(const TSharedRef<FUICommandList>& InCommandList) override;
diff --git a/Engine/Plugins/Experimental/Avalanche/Source/AvalancheEditorCore/Private/LevelEditor/AvaLevelEditor.cpp b/Engine/Plugins/Experimental/Avalanche/Source/AvalancheEditorCore/Private/LevelEditor/AvaLevelEditor.cpp
--- a/Engine/Plugins/Experimental/Avalanche/Source/AvalancheEditorCore/Private/LevelEditor/AvaLevelEditor.cpp
+++ b/Engine/Plugins/Experimental/Avalanche/Source/AvalancheEditorCore/Private/LevelEditor/AvaLevelEditor.cpp
@@ -4,6 +4,7 @@
 #include "AvaLevelEditorToolbar.h"
 #include "Delegates/Delegate.h"
 #include "IAvaEditorExtension.h"
+#include "IAvaEditorProvider.h"
 #include "LevelEditor.h"
 #include "LevelEditor/AvaLevelEditorUtils.h"
 #include "Styling/SlateIconFinder.h"
@@ -126,12 +127,18 @@ void FAvaLevelEditor::OnMapChanged(UWorld* InWorld, EMapChangeType InChangeType)
 	switch (InChangeType)
 	{
 	case EMapChangeType::NewMap:
-		TryOpenScene();
+		if (ShouldAutoActivate())
+		{
+			TryOpenScene();
+		}
 		break;
 
 	case EMapChangeType::LoadMap:
-		TryOpenScene();
-		Load();
+		if (ShouldAutoActivate())
+		{
+			TryOpenScene();
+			Load();
+		}
 		break;
 
 	case EMapChangeType::SaveMap:
@@ -159,7 +166,10 @@ void FAvaLevelEditor::OnLevelEditorCreated(TSharedPtr<ILevelEditor> InLevelEdito
 
 	DeactivateExtensions();
 
-	TryOpenScene();
+	if (ShouldAutoActivate())
+	{
+		TryOpenScene();
+	}
 }
 
 #undef LOCTEXT_NAMESPACE
diff --git a/Engine/Plugins/Experimental/Avalanche/Source/AvalancheEditorCore/Private/LevelEditor/AvaLevelEditorToolbar.cpp b/Engine/Plugins/Experimental/Avalanche/Source/AvalancheEditorCore/Private/LevelEditor/AvaLevelEditorToolbar.cpp
--- a/Engine/Plugins/Experimental/Avalanche/Source/AvalancheEditorCore/Private/LevelEditor/AvaLevelEditorToolbar.cpp
+++ b/Engine/Plugins/Experimental/Avalanche/Source/AvalancheEditorCore/Private/LevelEditor/AvaLevelEditorToolbar.cpp
@@ -36,6 +36,7 @@ namespace UE::AvaEditorCore::Private
 	{
 		if (TSharedPtr<IAvaEditor> Editor = InEditorWeak.Pin())
 		{
+			Editor->SetAutoActivate(true);
 			Editor->Activate();
 		}
 	}
@@ -55,6 +56,7 @@ namespace UE::AvaEditorCore::Private
 
 		if (TSharedPtr<IAvaEditor> Editor = InEditorWeak.Pin())
 		{
+			Editor->SetAutoActivate(false);
 			Editor->Deactivate();
 		}
 	}
diff --git a/Engine/Plugins/Experimental/Avalanche/Source/AvalancheEditorCore/Public/IAvaEditor.h b/Engine/Plugins/Experimental/Avalanche/Source/AvalancheEditorCore/Public/IAvaEditor.h
--- a/Engine/Plugins/Experimental/Avalanche/Source/AvalancheEditorCore/Public/IAvaEditor.h
+++ b/Engine/Plugins/Experimental/Avalanche/Source/AvalancheEditorCore/Public/IAvaEditor.h
@@ -54,6 +54,12 @@ public:
 
 	virtual bool CanActivate() const = 0;
 
+	/** Whether the editor should auto-activate if available */
+	virtual bool ShouldAutoActivate() const = 0;
+
+	/** Sets whether the editor should auto-activate */
+	virtual void SetAutoActivate(bool bInAutoActivate) = 0;
+
 	virtual bool CanDeactivate() const = 0;
 
 	virtual void BindCommands(const TSharedRef<FUICommandList>& InCommandList) = 0;
diff --git a/Engine/Plugins/Experimental/Avalanche/Source/AvalancheEditorCore/Public/IAvaEditorProvider.h b/Engine/Plugins/Experimental/Avalanche/Source/AvalancheEditorCore/Public/IAvaEditorProvider.h
--- a/Engine/Plugins/Experimental/Avalanche/Source/AvalancheEditorCore/Public/IAvaEditorProvider.h
+++ b/Engine/Plugins/Experimental/Avalanche/Source/AvalancheEditorCore/Public/IAvaEditorProvider.h
@@ -16,7 +16,9 @@ public:
 	virtual ~IAvaEditorProvider() = default;
 
 	/** Called right after instantiating the provider */
-	virtual void Construct() {}
+	virtual void Construct()
+	{
+	}
 
 	/**
 	 * Gets the Scene Object this Ava Editor Instance is using.
@@ -25,10 +27,28 @@ public:
 	 */
 	virtual UObject* GetSceneObject(UWorld* InWorld, EAvaEditorObjectQueryType InQueryType) = 0;
 
+	/** Whether the scene object can be auto-activated if the functionality is available */
+	virtual bool ShouldAutoActivateScene(UObject* InSceneObject) const
+	{
+		return true;
+	}
+
+	/** Sets the auto activate scene option for the scene object to persist it */
+	virtual void SetAutoActivateScene(UObject* InSceneObject, bool bInAutoActivateScene) const
+	{
+		
+	}
+	
 	/** Optional: Add/Remove from the Actor array when performing operations such as Copy/Cut */
-	virtual void GetActorsToEdit(TArray<AActor*>& InOutActorsToEdit) const {}
+	virtual void GetActorsToEdit(TArray<AActor*>& InOutActorsToEdit) const
+	{
+	}
 
-	virtual void OnSceneActivated() {}
+	virtual void OnSceneActivated()
+	{
+	}
 
-	virtual void OnSceneDeactivated() {}
+	virtual void OnSceneDeactivated()
+	{
+	}
 };
