[Motion Design] Playback Server - Better discarding of transitions.

- unloaded instances may lead to discarded transitions.
- adding a command (bindable macro) to stop transitions per channel.
- also includes a small fix for status refresh "spam" caused by not caching the non-playing instance status.

diff --git a/Engine/Plugins/Experimental/Avalanche/Source/AvalancheMedia/Private/Playback/AvaPlaybackClient.h b/Engine/Plugins/Experimental/Avalanche/Source/AvalancheMedia/Private/Playback/AvaPlaybackClient.h
--- a/Engine/Plugins/Experimental/Avalanche/Source/AvalancheMedia/Private/Playback/AvaPlaybackClient.h
+++ b/Engine/Plugins/Experimental/Avalanche/Source/AvalancheMedia/Private/Playback/AvaPlaybackClient.h
@@ -217,47 +217,10 @@ private:
 	{
 		TMap<FGuid, FPlaybackInstanceInfo> InstanceByIds;
 
-		/*
-		 * Return a relevance factor for status reporting, based on what we would like to see in the page status.
-		 */
-		static int GetStatusRelevance(EAvaPlaybackStatus InPlaybackStatus)
-		{
-			switch (InPlaybackStatus)
-			{
-			case EAvaPlaybackStatus::Unknown: return 0;
-			case EAvaPlaybackStatus::Missing: return 1;
-			case EAvaPlaybackStatus::Syncing: return 2;
-			case EAvaPlaybackStatus::Available: return 3;
-			case EAvaPlaybackStatus::Loading: return 4;
-			case EAvaPlaybackStatus::Loaded: return 5;
-			case EAvaPlaybackStatus::Starting: return 6;
-			case EAvaPlaybackStatus::Started: return 7;
-			case EAvaPlaybackStatus::Unloading: return 3;
-			case EAvaPlaybackStatus::Error: return 1;
-			default: return 0;
-			}
-		}
-		
 		TOptional<EAvaPlaybackStatus> GetInstanceStatus(const FGuid& InInstanceId) const
 		{
-			// If the request is for a specific instance, check if we have it.
-			if (InInstanceId.IsValid())
-			{
-				const FPlaybackInstanceInfo* Instance = InstanceByIds.Find(InInstanceId);
-				return Instance ? TOptional(Instance->Status) : TOptional<EAvaPlaybackStatus>();
-			}
-			
-			// If the instanceId is not specified, return the most relevant status we have already received.
-			// Todo: investigate why this happens and if there is better way to handle this (like better request handling).
-			TOptional<EAvaPlaybackStatus> MostRelevantStatus;
-			for (const TPair<FGuid, FPlaybackInstanceInfo>& PlaybackInfo :  InstanceByIds)
-			{
-				if (!MostRelevantStatus.IsSet() || GetStatusRelevance(PlaybackInfo.Value.Status) > GetStatusRelevance(MostRelevantStatus.GetValue()))
-				{
-					MostRelevantStatus = PlaybackInfo.Value.Status;
-				}
-			}
-			return MostRelevantStatus;
+			const FPlaybackInstanceInfo* Instance = InstanceByIds.Find(InInstanceId);
+			return Instance ? TOptional(Instance->Status) : TOptional<EAvaPlaybackStatus>();
 		}
 
 		const FPlaybackInstanceInfo* GetInstanceInfo(const FGuid& InInstanceId) const
@@ -267,12 +230,8 @@ private:
 
 		FPlaybackInstanceInfo* GetOrAddInstanceInfo(const FGuid& InInstanceId)
 		{
-			if (InInstanceId.IsValid())
-			{
-				FPlaybackInstanceInfo* Instance = InstanceByIds.Find(InInstanceId);
-				return Instance ? Instance : &InstanceByIds.Emplace(InInstanceId);
-			}
-			return nullptr;
+			FPlaybackInstanceInfo* Instance = InstanceByIds.Find(InInstanceId);
+			return Instance ? Instance : &InstanceByIds.Emplace(InInstanceId);
 		}
 		
 		void SetInstanceStatus(const FGuid& InInstanceId, EAvaPlaybackStatus InStatus)
diff --git a/Engine/Plugins/Experimental/Avalanche/Source/AvalancheMedia/Private/Playback/AvaPlaybackServer.cpp b/Engine/Plugins/Experimental/Avalanche/Source/AvalancheMedia/Private/Playback/AvaPlaybackServer.cpp
--- a/Engine/Plugins/Experimental/Avalanche/Source/AvalancheMedia/Private/Playback/AvaPlaybackServer.cpp
+++ b/Engine/Plugins/Experimental/Avalanche/Source/AvalancheMedia/Private/Playback/AvaPlaybackServer.cpp
@@ -100,8 +100,18 @@ public:
 		
 	virtual FString GetReferencerName() const override { return TEXT("FServerPlaybackInstanceTransitions"); }
 	//~ End FGCObject
-
+	
+	/** Map Key: TransitionId. */
 	TMap<FGuid, TObjectPtr<UAvaPlaybackServerTransition>> Transitions;
+
+	/** Transitions marked for stop and discard. */
+	TSet<FGuid> MarkedForStopAndDiscard;
+
+	UAvaPlaybackServerTransition* FindTransition(const FGuid& InTransitionId) const
+	{
+		const TObjectPtr<UAvaPlaybackServerTransition>* FoundTransition = Transitions.Find(InTransitionId);
+		return FoundTransition ? (*FoundTransition).Get() : nullptr;
+	}
 };
 
 FAvaPlaybackServer::FAvaPlaybackServer()
@@ -720,21 +730,17 @@ void FAvaPlaybackServer::HandlePlayableTransitionStartRequest(const FAvaPlayback
 
 void FAvaPlaybackServer::HandlePlayableTransitionStopRequest(const FAvaPlaybackTransitionStopRequest& InMessage, const TSharedRef<IMessageContext, ESPMode::ThreadSafe>& InContext)
 {
-	check(PlaybackInstanceTransitions);
+	if (!PlaybackInstanceTransitions)
+	{
+		UE_LOG(LogAvaPlaybackServer, Error, 
+			TEXT("Stopping Transition \"%s\" failed because transition collection is not available."), *InMessage.TransitionId.ToString());
+		// todo: should reply with an error message to client. (need a request id)
+		return;
+	}
 	
-	if (const TObjectPtr<UAvaPlaybackServerTransition>* FoundTransition = PlaybackInstanceTransitions->Transitions.Find(InMessage.TransitionId))
+	if (PlaybackInstanceTransitions->Transitions.Contains(InMessage.TransitionId))
 	{
-		if (*FoundTransition)
-		{
-			(*FoundTransition)->Stop();
-			
-			// That should have removed the transition from the list.
-			if (PlaybackInstanceTransitions->Transitions.Contains(InMessage.TransitionId))
-			{
-				UE_LOG(LogAvaPlaybackServer, Error,
-					TEXT("Stopping Transition \"%s\" didn't remove it from the active list."), *InMessage.TransitionId.ToString());
-			}
-		}
+		PlaybackInstanceTransitions->MarkedForStopAndDiscard.Add(InMessage.TransitionId);
 	}
 	else
 	{
@@ -863,10 +869,29 @@ void FAvaPlaybackServer::Tick()
 	// Execute the pending commands in batch for this tick.
 	ExecutePendingPlaybackCommands(CurrentTimeUtc);
 
-	// Try to resolve the instance for loaded transitions.
-	for (const TPair<FGuid, TObjectPtr<UAvaPlaybackServerTransition>>& Transition : PlaybackInstanceTransitions->Transitions)
+	if (PlaybackInstanceTransitions)
 	{
-		Transition.Value->TryResolveInstances(*this);
+		for (const FGuid& TransitionId : PlaybackInstanceTransitions->MarkedForStopAndDiscard)
+		{
+			if (UAvaPlaybackServerTransition* Transition = PlaybackInstanceTransitions->FindTransition(TransitionId))
+			{
+				Transition->Stop();	// Should remove from active list.
+
+				if (PlaybackInstanceTransitions->Transitions.Contains(TransitionId))
+				{
+					UE_LOG(LogAvaPlaybackServer, Error, TEXT("Playback Transition %s is still in the list after being stopped."), *TransitionId.ToString());
+					PlaybackInstanceTransitions->Transitions.Remove(TransitionId);	// Remove anyway.
+				}
+			}
+		}
+		
+		PlaybackInstanceTransitions->MarkedForStopAndDiscard.Reset();
+
+		// Try to resolve the instance for loaded transitions.
+		for (const TPair<FGuid, TObjectPtr<UAvaPlaybackServerTransition>>& Transition : PlaybackInstanceTransitions->Transitions)
+		{
+			Transition.Value->TryResolveInstances(*this);
+		}
 	}
 	
 	// TODO:
@@ -1497,6 +1522,9 @@ void FAvaPlaybackServer::UnloadPlayback(const FMessageAddress& InReplyToAddress,
 				{
 					if (Transition.Value && Transition.Value->ContainsInstance(InInstanceId))
 					{
+						// Force Take out should also force stop the transition.
+						PlaybackInstanceTransitions->MarkedForStopAndDiscard.Add(Transition.Value->GetTransitionId());
+						
 						UE_LOG(LogAvaPlaybackServer, Error, TEXT("%s Unloading instance \"%s\" (id:%s) while it is part of transition %s."),
 							*UE::AvaPlayback::Utils::GetBriefFrameInfo(), *Instance->GetSourcePath().GetAssetName(), *InInstanceId.ToString(), *Transition.Value->GetTransitionId().ToString());
 					}
@@ -1685,6 +1713,12 @@ void FAvaPlaybackServer::SendAllPlaybackStatusesForChannelAndAssetPath(const FMe
 	{
 		SendPlaybackStatuses(InSendTo, InChannelName, StatusEntry.Value, StatusEntry.Key);
 	}
+
+	// In the case there is no playing instances, return the status of the unloaded asset as playback status.
+	if (InstancesPerStatus.IsEmpty())
+	{
+		SendPlaybackStatus(InSendTo, FGuid(), InChannelName, InAssetPath, GetUnloadedPlaybackStatus(InAssetPath));
+	}
 }
 
 void FAvaPlaybackServer::SendPlaybackAssetStatus(const FMessageAddress& InSendTo, const FSoftObjectPath& InAssetPath, EAvaPlaybackAssetStatus InStatus)
diff --git a/Engine/Plugins/Experimental/Avalanche/Source/AvalancheMedia/Private/Playback/Transition/AvaPlaybackServerTransition.h b/Engine/Plugins/Experimental/Avalanche/Source/AvalancheMedia/Private/Playback/Transition/AvaPlaybackServerTransition.h
--- a/Engine/Plugins/Experimental/Avalanche/Source/AvalancheMedia/Private/Playback/Transition/AvaPlaybackServerTransition.h
+++ b/Engine/Plugins/Experimental/Avalanche/Source/AvalancheMedia/Private/Playback/Transition/AvaPlaybackServerTransition.h
@@ -105,4 +105,10 @@ protected:
 
 	UPROPERTY(Transient)
 	TObjectPtr<UAvaPlayableTransition> PlayableTransition;
+
+	/** Re-entry guard for stop() function. */
+	bool bIsStopping = false;
+
+	/** Indicate if the "finish" event has been sent. There is 2 possible paths to send the finish event but we only want to send once. */
+	bool bIsFinishedSent = false;
 };
\ No newline at end of file
diff --git a/Engine/Plugins/Experimental/Avalanche/Source/AvalancheMedia/Private/Rundown/AvaRundown.cpp b/Engine/Plugins/Experimental/Avalanche/Source/AvalancheMedia/Private/Rundown/AvaRundown.cpp
--- a/Engine/Plugins/Experimental/Avalanche/Source/AvalancheMedia/Private/Rundown/AvaRundown.cpp
+++ b/Engine/Plugins/Experimental/Avalanche/Source/AvalancheMedia/Private/Rundown/AvaRundown.cpp
@@ -1390,6 +1390,7 @@ TArray<int32> UAvaRundown::StopPages(const TArray<int32>& InPageIds, EAvaRundown
 
 TArray<int32> UAvaRundown::StopPages(const TArray<int32>& InPageIds, EAvaRundownPageStopOptions InOptions, bool bInPreview, const FName& InPreviewChannelName)
 {
+	const bool bForceNoTransition = EnumHasAnyFlags(InOptions, EAvaRundownPageStopOptions::ForceNoTransition);
 	TArray<int32> StoppedPageIds;
 	StoppedPageIds.Reserve(InPageIds.Num());
 
@@ -1404,24 +1405,25 @@ TArray<int32> UAvaRundown::StopPages(const TArray<int32>& InPageIds, EAvaRundown
 			continue;
 		}
 
-		if (SelectedPage.HasTransitionLogic(this))
+		// Force stop all transitions for the selected page (if any).
+		if (bForceNoTransition)
 		{
-			if (!EnumHasAnyFlags(InOptions, EAvaRundownPageStopOptions::ForceNoTransition))
-			{
-				if (StopPageWithTransition(TransitionBuilder, SelectedPage, bInPreview, InPreviewChannelName))
-				{
-					StoppedPageIds.Add(PageId);
-				}
-				continue;
-			}
-
-			// Force stop all page transitions for the given page, then stop the page without TL.
 			StopPageTransitionsForPage(SelectedPage, bInPreview, InPreviewChannelName);
 		}
 
-		if (StopPageNoTransition(SelectedPage, bInPreview, InPreviewChannelName))
+		if (SelectedPage.HasTransitionLogic(this) && !bForceNoTransition)
 		{
-			StoppedPageIds.Add(PageId);
+			if (StopPageWithTransition(TransitionBuilder, SelectedPage, bInPreview, InPreviewChannelName))
+			{
+				StoppedPageIds.Add(PageId);
+			}
+		}
+		else
+		{
+			if (StopPageNoTransition(SelectedPage, bInPreview, InPreviewChannelName))
+			{
+				StoppedPageIds.Add(PageId);
+			}
 		}
 	}
 	return StoppedPageIds;
@@ -2570,19 +2572,35 @@ bool UAvaRundown::CanStartTransitionForPage(const FAvaRundownPage& InPage, bool
 	return true;
 }
 
-void UAvaRundown::StopPageTransitionsForPage(const FAvaRundownPage& InPage, bool bInIsPreview, const FName& InPreviewChannelName)
+int32 UAvaRundown::StopPageTransitionsForPage(const FAvaRundownPage& InPage, bool bInIsPreview, const FName& InPreviewChannelName)
+{
+	// Implementation note: there can only be one transition at the moment and all playing pages are part of it.
+	// This is why this function is equivalent to stopping the transitions in a given channel.
+	// Keeping the function around for future evolution of the system.
+	const FName ChannelName = bInIsPreview ? InPreviewChannelName : InPage.GetChannelName();
+	return StopPageTransitionsForChannel(ChannelName);
+}
+
+int32 UAvaRundown::StopPageTransitionsForChannel(const FName InChannelName)
+{
+	return StopPageTransitionsByPredicate([InChannelName](UAvaRundownPageTransition* InTransition)
+	{
+		return InTransition->GetChannelName() == InChannelName;
+	});
+}
+
+int32 UAvaRundown::StopPageTransitionsByPredicate(TFunctionRef<bool(UAvaRundownPageTransition*)> InPredicate)
 {
 	TArray<UAvaRundownPageTransition*> TransitionsToStop;
 	TransitionsToStop.Reserve(PageTransitions.Num());
 
 	// Note: we build a separate list because stopping the transitions should
 	// lead to the transitions being removed from PageTransitions (through the events).
-	const FName ChannelName = bInIsPreview ? InPreviewChannelName : InPage.GetChannelName();
 	for (UAvaRundownPageTransition* PageTransition : PageTransitions)
 	{
-		if (PageTransition && PageTransition->GetChannelName() == ChannelName)
+		if (PageTransition && InPredicate(PageTransition))
 		{
-			TransitionsToStop.Add(PageTransition);	
+			TransitionsToStop.Add(PageTransition);
 		}
 	}
 
@@ -2598,6 +2616,8 @@ void UAvaRundown::StopPageTransitionsForPage(const FAvaRundownPage& InPage, bool
 			PageTransitions.Remove(Transition);
 		}
 	}
+
+	return TransitionsToStop.Num();
 }
 
 UAvaRundownPagePlayer* UAvaRundown::FindPlayerForProgramPage(int32 InPageId) const
diff --git a/Engine/Plugins/Experimental/Avalanche/Source/AvalancheMedia/Public/Rundown/AvaRundown.h b/Engine/Plugins/Experimental/Avalanche/Source/AvalancheMedia/Public/Rundown/AvaRundown.h
--- a/Engine/Plugins/Experimental/Avalanche/Source/AvalancheMedia/Public/Rundown/AvaRundown.h
+++ b/Engine/Plugins/Experimental/Avalanche/Source/AvalancheMedia/Public/Rundown/AvaRundown.h
@@ -776,9 +776,25 @@ public:
 	UAvaRundownPageTransition* GetPageTransition(const FGuid& InTransitionId) const;
 
 	bool CanStartTransitionForPage(const FAvaRundownPage& InPage, bool bInIsPreview, const FName& InPreviewChannelName, FString* OutFailureReason) const;
+	
+	/**
+	 * Stop transitions involving the given page playing in the specified channel.
+	 * @return number of transitions stopped.
+	 */
+	int32 StopPageTransitionsForPage(const FAvaRundownPage& InPage, bool bInIsPreview, const FName& InPreviewChannelName);
 
-	void StopPageTransitionsForPage(const FAvaRundownPage& InPage, bool bInIsPreview, const FName& InPreviewChannelName);
+	/**
+	 * Stop all current transitions in the specified channel.
+	 * @return number of transitions stopped.
+	 */
+	int32 StopPageTransitionsForChannel(const FName InChannelName);
 
+	/**
+	 * Stop transitions that return true for the given predicate.
+	 * @return number of transitions stopped.
+	 */
+	int32 StopPageTransitionsByPredicate(TFunctionRef<bool(UAvaRundownPageTransition*)> InPredicate);
+	
 	const TArray<TObjectPtr<UAvaRundownPagePlayer>>& GetPagePlayers() const { return PagePlayers; }
 	
 	UAvaRundownPagePlayer* FindPlayerForProgramPage(int32 InPageId) const;
diff --git a/Engine/Plugins/Experimental/Avalanche/Source/AvalancheMediaEditor/Private/Rundown/AvaRundownEditor.cpp b/Engine/Plugins/Experimental/Avalanche/Source/AvalancheMediaEditor/Private/Rundown/AvaRundownEditor.cpp
--- a/Engine/Plugins/Experimental/Avalanche/Source/AvalancheMediaEditor/Private/Rundown/AvaRundownEditor.cpp
+++ b/Engine/Plugins/Experimental/Avalanche/Source/AvalancheMediaEditor/Private/Rundown/AvaRundownEditor.cpp
@@ -1466,6 +1466,40 @@ void FAvaRundownEditor::TakeToProgramCommand(const TArray<FString>& InArgs)
 	}
 }
 
+void FAvaRundownEditor::StopPageTransitionCommand(const TArray<FString>& InArgs)
+{
+	UAvaRundown* Rundown = AvaRundown.Get();
+	if (!Rundown)
+	{
+		UE_LOG(LogAvaRundown, Error, TEXT("No valid rundown available."));
+		return;
+	}
+
+	const UAvaBroadcast& Broadcast = UAvaBroadcast::Get();
+	TArray<FString> ChannelNames;
+	ChannelNames.Reserve(Broadcast.GetChannelNameCount());
+	for (int32 ChannelIndex = 0; ChannelIndex < Broadcast.GetChannelNameCount(); ++ChannelIndex)
+	{
+		ChannelNames.Add(Broadcast.GetChannelName(ChannelIndex).ToString());
+	}
+	
+	// Arguments have to be channel names, but we support partial name match.
+	for (const FString& Argument : InArgs)
+	{
+		for (const FString& ChannelName : ChannelNames)
+		{
+			if (ChannelName == Argument || ChannelName.Contains(Argument))
+			{
+				const int32 NumTransitionsStopped = Rundown->StopPageTransitionsForChannel(FName(*ChannelName)); 
+				if (NumTransitionsStopped > 0)
+				{
+					UE_LOG(LogAvaRundown, Log, TEXT("Stopped %d transition(s) on channel %s."), NumTransitionsStopped, *ChannelName);
+				}
+			}
+		}
+	}
+}
+
 void FAvaRundownEditor::StartChannelCommand(const TArray<FString>& InArgs)
 {
 	UAvaBroadcast& Broadcast = UAvaBroadcast::Get();
@@ -1559,6 +1593,10 @@ TArray<int32> FAvaRundownEditor::ArgumentsToPageIds(const TArray<FString>& InArg
 				PageIds.Append(Rundown->GetPlayingPageIds());
 			}
 		}
+		else if (Arg == TEXT("selected")) // selected pages for the current list (in focus)
+		{
+			PageIds.Append(GetSelectedPagesOnFocusedWidget());
+		}
 		else if (UAvaBroadcast::Get().GetChannelIndex(FName(Arg)) != INDEX_NONE)
 		{
 			// Any playing pages on the specified channel.
@@ -1597,6 +1635,7 @@ const FAvaRundownEditor::FBindableMacroCommands& FAvaRundownEditor::GetBindableM
 		BindableMacroCommands.Add(GetName(EAvaRundownEditorMacroCommand::PreviewNext), [this](const TArray<FString>& InArgs){PlayNextPageCommand(InArgs, /*bInPreview*/ true);});
 		BindableMacroCommands.Add(GetName(EAvaRundownEditorMacroCommand::ContinuePreview), [this](const TArray<FString>& InArgs){ContinuePageCommand(InArgs, /*bInPreview*/ true);});
 		BindableMacroCommands.Add(GetName(EAvaRundownEditorMacroCommand::TakeToProgram), [this](const TArray<FString>& InArgs){TakeToProgramCommand(InArgs);});
+		BindableMacroCommands.Add(GetName(EAvaRundownEditorMacroCommand::StopPageTransition), [this](const TArray<FString>& InArgs){StopPageTransitionCommand(InArgs);});
 		BindableMacroCommands.Add(GetName(EAvaRundownEditorMacroCommand::StartChannel), [this](const TArray<FString>& InArgs){StartChannelCommand(InArgs);});
 		BindableMacroCommands.Add(GetName(EAvaRundownEditorMacroCommand::StopChannel), [this](const TArray<FString>& InArgs){StopChannelCommand(InArgs);});
 	}
diff --git a/Engine/Plugins/Experimental/Avalanche/Source/AvalancheMediaEditor/Private/Rundown/AvaRundownEditorMacroCommands.h b/Engine/Plugins/Experimental/Avalanche/Source/AvalancheMediaEditor/Private/Rundown/AvaRundownEditorMacroCommands.h
--- a/Engine/Plugins/Experimental/Avalanche/Source/AvalancheMediaEditor/Private/Rundown/AvaRundownEditorMacroCommands.h
+++ b/Engine/Plugins/Experimental/Avalanche/Source/AvalancheMediaEditor/Private/Rundown/AvaRundownEditorMacroCommands.h
@@ -23,6 +23,7 @@ enum class EAvaRundownEditorMacroCommand : uint8
 	PreviewNext,
 	ContinuePreview,
 	TakeToProgram,
+	StopPageTransition,
 	StartChannel,
 	StopChannel
 };
diff --git a/Engine/Plugins/Experimental/Avalanche/Source/AvalancheMediaEditor/Public/Rundown/AvaRundownEditor.h b/Engine/Plugins/Experimental/Avalanche/Source/AvalancheMediaEditor/Public/Rundown/AvaRundownEditor.h
--- a/Engine/Plugins/Experimental/Avalanche/Source/AvalancheMediaEditor/Public/Rundown/AvaRundownEditor.h
+++ b/Engine/Plugins/Experimental/Avalanche/Source/AvalancheMediaEditor/Public/Rundown/AvaRundownEditor.h
@@ -135,6 +135,7 @@ protected:
 	void StopPageCommand(const TArray<FString>& InArgs, bool bInPreview, bool bInForce);
 	void StopLayersCommand(const TArray<FString>& InArgs, bool bInPreview, bool bInForce);
 	void TakeToProgramCommand(const TArray<FString>& InArgs);
+	void StopPageTransitionCommand(const TArray<FString>& InArgs);
 	void StartChannelCommand(const TArray<FString>& InArgs);
 	void StopChannelCommand(const TArray<FString>& InArgs);
 	
