[Motion Design] Playback Server - Better discarding of transitions (Part 2).

- guarding against re-entrency when stopping a playback server transition.

(should have been part of CL 41241150)

diff --git a/Engine/Plugins/Experimental/Avalanche/Source/AvalancheMedia/Private/Playback/Transition/AvaPlaybackServerTransition.cpp b/Engine/Plugins/Experimental/Avalanche/Source/AvalancheMedia/Private/Playback/Transition/AvaPlaybackServerTransition.cpp
--- a/Engine/Plugins/Experimental/Avalanche/Source/AvalancheMedia/Private/Playback/Transition/AvaPlaybackServerTransition.cpp
+++ b/Engine/Plugins/Experimental/Avalanche/Source/AvalancheMedia/Private/Playback/Transition/AvaPlaybackServerTransition.cpp
@@ -353,6 +353,15 @@ void UAvaPlaybackServerTransition::Start()
 
 void UAvaPlaybackServerTransition::Stop()
 {
+	// Guard against reentrancy from different event propagation paths,
+	// i.e. direct stop call from server command or from playable transition event.
+	if (bIsStopping)
+	{
+		return;
+	}
+
+	TGuardValue ReentrancyGuard(bIsStopping, true);
+
 	if (PlayableTransition)
 	{
 		PlayableTransition->Stop();
@@ -381,6 +390,15 @@ void UAvaPlaybackServerTransition::Stop()
 				TEXT("%s Failed to remove Playback Transition {%s}. Reason: not found in server's active transitions."),
 				*UE::AvaPlayback::Utils::GetBriefFrameInfo(), *GetPrettyTransitionInfo());
 		}
+
+		// Make sure we send the finish transition event to the client, in case it hasn't been sent through the
+		// playable transition event path already.
+		if (!bIsFinishedSent)
+		{
+			constexpr EAvaPlayableTransitionEventFlags FinishedFlag = EAvaPlayableTransitionEventFlags::Finished;
+			PlaybackServer->SendPlayableTransitionEvent(TransitionId, FGuid(), FinishedFlag, ChannelName, ClientName);
+			bIsFinishedSent = true;
+		}
 	}
 }
 
@@ -508,8 +526,9 @@ void UAvaPlaybackServerTransition::OnTransitionEvent(UAvaPlayable* InPlayable, U
 		if (PlaybackServer)
 		{
 			PlaybackServer->SendPlayableTransitionEvent(TransitionId, FGuid(), InTransitionFlags, ChannelName, ClientName);
+			bIsFinishedSent = true;
 		}
-		
+
 		Stop();
 	}
 }
