Use new EditorViewportClient override to control how the final view matrix is calculated

diff --git a/Engine/Plugins/Experimental/Avalanche/Source/AvalancheLevelViewport/Private/ViewportClient/AvaLevelViewportClient.cpp b/Engine/Plugins/Experimental/Avalanche/Source/AvalancheLevelViewport/Private/ViewportClient/AvaLevelViewportClient.cpp
--- a/Engine/Plugins/Experimental/Avalanche/Source/AvalancheLevelViewport/Private/ViewportClient/AvaLevelViewportClient.cpp
+++ b/Engine/Plugins/Experimental/Avalanche/Source/AvalancheLevelViewport/Private/ViewportClient/AvaLevelViewportClient.cpp
@@ -402,6 +402,11 @@ void FAvaLevelViewportClient::DrawCanvas(FViewport& InViewport, FSceneView& View
 	Super::DrawCanvas(InViewport, View, Canvas);
 }
 
+FMatrix FAvaLevelViewportClient::CalcViewRotationMatrixForControllingActorView(const FRotator& InViewRotation) const
+{
+	return CalcViewRotationMatrix(InViewRotation);
+}
+
 FSceneView* FAvaLevelViewportClient::CalcSceneView(FSceneViewFamily* ViewFamily, const int32 StereoViewIndex)
 {
 	UCameraComponent* CachedCameraComponent = nullptr;
diff --git a/Engine/Plugins/Experimental/Avalanche/Source/AvalancheLevelViewport/Private/ViewportClient/AvaLevelViewportClient.h b/Engine/Plugins/Experimental/Avalanche/Source/AvalancheLevelViewport/Private/ViewportClient/AvaLevelViewportClient.h
--- a/Engine/Plugins/Experimental/Avalanche/Source/AvalancheLevelViewport/Private/ViewportClient/AvaLevelViewportClient.h
+++ b/Engine/Plugins/Experimental/Avalanche/Source/AvalancheLevelViewport/Private/ViewportClient/AvaLevelViewportClient.h
@@ -72,6 +72,8 @@ public:
 	virtual bool InputKey(const FInputKeyEventArgs& EventArgs) override;
 	virtual void UpdateMouseDelta() override;
 	virtual EMouseCursor::Type GetCursor(FViewport* InViewport, int32 InX, int32 InY) override;
+	virtual FMatrix CalcViewRotationMatrixForControllingActorView(const FRotator& InViewRotation) const override;
+
 	//~ End FEditorViewportClient
 
 	//~ Begin IAvaViewportWorldCoordinateConverter
