MotionDesign : Text3D
- Splitting reusable logic in extensions (material, geometry, layout, rendering, transform) to support multiple renderers
- Migrating properties and code from AvaText3D to Text3D
- Adapting mask modifier to support these changes
- Adapting motion design vizualizers to support these changes
- Adapting property animators to support these changes
- Reducing amount of base materials needed from 18 to 4 (translucent+opaque, lit+unlit)
- Adapting material code that uses text3d bounds to work without absolute position

#virtualized

diff --git a/Engine/Plugins/Experimental/Avalanche/Source/AvalancheMask/Private/Handling/AvaMaskText3DActorMaterialCollectionHandle.cpp b/Engine/Plugins/Experimental/Avalanche/Source/AvalancheMask/Private/Handling/AvaMaskText3DActorMaterialCollectionHandle.cpp
--- a/Engine/Plugins/Experimental/Avalanche/Source/AvalancheMask/Private/Handling/AvaMaskText3DActorMaterialCollectionHandle.cpp
+++ b/Engine/Plugins/Experimental/Avalanche/Source/AvalancheMask/Private/Handling/AvaMaskText3DActorMaterialCollectionHandle.cpp
@@ -8,28 +8,28 @@
 #include "AvaObjectHandleSubsystem.h"
 #include "AvaText3DComponent.h"
 #include "Components/PrimitiveComponent.h"
+#include "Extensions/Text3DDefaultMaterialExtension.h"
 #include "GameFramework/Actor.h"
 #include "Handling/AvaHandleUtilities.h"
 #include "IAvaMaskMaterialCollectionHandle.h"
 #include "IAvaMaskMaterialHandle.h"
-#include "MaterialHub/AvaTextMaterialHub.h"
 #include "Materials/MaterialInterface.h"
 #include "Text3DComponent.h"
 
 namespace UE::AvaMask::Private
 {
-	EAvaTextTranslucency GetTargetTranslucencyType(
-		const EAvaTextTranslucency InFromMaterial
+	EText3DMaterialBlendMode GetTargetTranslucencyType(
+		const EText3DMaterialBlendMode InFromMaterial
 		, const EBlendMode InRequired)
 	{
 		if (InRequired == EBlendMode::BLEND_Opaque)
 		{
 			return InFromMaterial;
 		}
-		
-		if (InFromMaterial == EAvaTextTranslucency::None)
+
+		if (InFromMaterial == EText3DMaterialBlendMode::Opaque)
 		{
-			return EAvaTextTranslucency::Translucent;
+			return EText3DMaterialBlendMode::Translucent;
 		}
 
 		return InFromMaterial;
@@ -40,6 +40,18 @@ FAvaMaskText3DActorMaterialCollectionHandle::FAvaMaskText3DActorMaterialCollecti
 	: WeakActor(InActor)
 	, WeakComponent(InActor->GetComponentByClass<UText3DComponent>())
 {
+	if (UText3DComponent* Text3DComponent = WeakComponent.Get())
+	{
+		Text3DComponent->OnTextPostUpdate().AddRaw(this, &FAvaMaskText3DActorMaterialCollectionHandle::OnTextPostUpdate);
+	}
+}
+
+FAvaMaskText3DActorMaterialCollectionHandle::~FAvaMaskText3DActorMaterialCollectionHandle()
+{
+	if (UText3DComponent* Text3DComponent = WeakComponent.Get())
+	{
+		Text3DComponent->OnTextPostUpdate().RemoveAll(this);
+	}
 }
 
 TArray<TObjectPtr<UMaterialInterface>> FAvaMaskText3DActorMaterialCollectionHandle::GetMaterials()
@@ -62,19 +74,26 @@ TArray<TSharedPtr<IAvaMaskMaterialHandle>> FAvaMaskText3DActorMaterialCollection
 {
 	if (MaterialHandles.IsEmpty() || !UE::Ava::Internal::IsHandleValid(MaterialHandles.Last()))
 	{
-		if (const UText3DComponent* TextComponent = GetComponent())
+		if (UText3DComponent* TextComponent = GetComponent())
 		{
 			UAvaObjectHandleSubsystem* HandleSubsystem = UE::Ava::Internal::GetObjectHandleSubsystem();
 			if (!HandleSubsystem)
 			{
 				return MaterialHandles;
 			}
-			
+
+			const UText3DMaterialExtensionBase* MaterialExtension = TextComponent->GetMaterialExtension();
+
+			UMaterialInterface* FrontMaterial = MaterialExtension->GetMaterial(EText3DGroupType::Front);
+			UMaterialInterface* BackMaterial = MaterialExtension->GetMaterial(EText3DGroupType::Back);
+			UMaterialInterface* ExtrudeMaterial = MaterialExtension->GetMaterial(EText3DGroupType::Extrude);
+			UMaterialInterface* BevelMaterial = MaterialExtension->GetMaterial(EText3DGroupType::Bevel);
+
 			MaterialHandles.Init(nullptr, static_cast<int32>(EText3DGroupType::TypeCount));
-			MaterialHandles[static_cast<int32>(EText3DGroupType::Front)] = HandleSubsystem->MakeHandle<IAvaMaskMaterialHandle>(TextComponent->GetFrontMaterial(), UE::AvaMask::Internal::HandleTag);
-			MaterialHandles[static_cast<int32>(EText3DGroupType::Bevel)] = HandleSubsystem->MakeHandle<IAvaMaskMaterialHandle>(TextComponent->GetBevelMaterial(), UE::AvaMask::Internal::HandleTag);
-			MaterialHandles[static_cast<int32>(EText3DGroupType::Extrude)] = HandleSubsystem->MakeHandle<IAvaMaskMaterialHandle>(TextComponent->GetExtrudeMaterial(), UE::AvaMask::Internal::HandleTag);
-			MaterialHandles[static_cast<int32>(EText3DGroupType::Back)] = HandleSubsystem->MakeHandle<IAvaMaskMaterialHandle>(TextComponent->GetBackMaterial(), UE::AvaMask::Internal::HandleTag);
+			MaterialHandles[static_cast<int32>(EText3DGroupType::Front)] = HandleSubsystem->MakeHandle<IAvaMaskMaterialHandle>(FrontMaterial, UE::AvaMask::Internal::HandleTag);
+			MaterialHandles[static_cast<int32>(EText3DGroupType::Bevel)] = HandleSubsystem->MakeHandle<IAvaMaskMaterialHandle>(BevelMaterial, UE::AvaMask::Internal::HandleTag);
+			MaterialHandles[static_cast<int32>(EText3DGroupType::Extrude)] = HandleSubsystem->MakeHandle<IAvaMaskMaterialHandle>(ExtrudeMaterial, UE::AvaMask::Internal::HandleTag);
+			MaterialHandles[static_cast<int32>(EText3DGroupType::Back)] = HandleSubsystem->MakeHandle<IAvaMaskMaterialHandle>(BackMaterial, UE::AvaMask::Internal::HandleTag);
 		}
 	}
 
@@ -185,7 +204,7 @@ void FAvaMaskText3DActorMaterialCollectionHandle::ForEachMaterial(
 		{
 			return;
 		}
-		
+
 		InFunction(ComponentReference, static_cast<int32>(EText3DGroupType::Back), TextComponent->GetBackMaterial());
 	}
 }
@@ -305,365 +324,7 @@ bool FAvaMaskText3DActorMaterialCollectionHandle::IsValid() const
 	return WeakActor.IsValid() && WeakComponent.IsValid();
 }
 
-bool FAvaMaskText3DActorMaterialCollectionHandle::IsSupported(const FAvaMaskMaterialReference& InInstance, FName InTag)
-{
-	if (InTag == UE::AvaMask::Internal::HandleTag)
-	{
-		if (const AActor* Actor = InInstance.GetTypedObject<AActor>())
-		{
-			return Actor->GetComponentByClass<UText3DComponent>() != nullptr;
-		}
-	}
-
-	return false;
-}
-
-FStructView FAvaMaskText3DActorMaterialCollectionHandle::GetMaterialHandleData(
-	FHandleData* InParentHandleData
-	, const FSoftComponentReference& InComponent
-	, const int32 InSlotIdx)
-{
-	if (InParentHandleData)
-	{
-		if (ensure(InParentHandleData->GroupMaterialData.Contains(InSlotIdx)))
-		{
-			return InParentHandleData->GroupMaterialData[InSlotIdx];
-		}
-	}
-
-	return nullptr;
-}
-
-FStructView FAvaMaskText3DActorMaterialCollectionHandle::GetOrAddMaterialHandleData(
-	FHandleData* InParentHandleData
-	, const TSharedPtr<IAvaMaskMaterialHandle>& InMaterialHandle
-	, const FSoftComponentReference& InComponent
-	, const int32 InSlotIdx)
-{
-	if (InParentHandleData)
-	{
-		return InParentHandleData->GroupMaterialData.FindOrAdd(InSlotIdx, InMaterialHandle->MakeDataStruct());
-	}
-
-	return nullptr;
-}
-
-UText3DComponent* FAvaMaskText3DActorMaterialCollectionHandle::GetComponent()
-{
-	if (UText3DComponent* Component = WeakComponent.Get())
-	{
-		return Component;
-	}
-
-	if (const AActor* Actor = WeakActor.Get())
-	{
-		WeakComponent = Actor->GetComponentByClass<UText3DComponent>();
-		if (UText3DComponent* Component = WeakComponent.Get())
-		{
-			return Component;
-		}
-	}
-
-	return nullptr;
-}
-
-FAvaMaskAvaTextActorMaterialCollectionHandle::FAvaMaskAvaTextActorMaterialCollectionHandle(AActor* InActor)
-	: WeakActor(InActor)
-	, WeakComponent(InActor->GetComponentByClass<UAvaText3DComponent>())
-{
-	if (UAvaText3DComponent* TextComponent = GetComponent())
-	{
-		TextComponent->GetMaterialProviderDelegate().BindRaw(this, &FAvaMaskAvaTextActorMaterialCollectionHandle::ProvideMaterialWithSettings);
-	}
-}
-
-FAvaMaskAvaTextActorMaterialCollectionHandle::~FAvaMaskAvaTextActorMaterialCollectionHandle()
-{
-	if (UAvaText3DComponent* TextComponent = GetComponent())
-	{
-		TextComponent->GetMaterialProviderDelegate().Unbind();
-	}
-}
-
-TArray<TObjectPtr<UMaterialInterface>> FAvaMaskAvaTextActorMaterialCollectionHandle::GetMaterials()
-{
-	TArray<TObjectPtr<UMaterialInterface>> Materials;
-	Materials.Reserve(GetNumMaterials());
-	
-	if (const UAvaText3DComponent* TextComponent = GetComponent())
-	{
-		Materials.Emplace(TextComponent->GetFrontMaterial());
-		Materials.Emplace(TextComponent->GetBevelMaterial());
-		Materials.Emplace(TextComponent->GetExtrudeMaterial());
-		Materials.Emplace(TextComponent->GetBackMaterial());
-	}
-
-	return Materials;
-}
-
-TArray<TSharedPtr<IAvaMaskMaterialHandle>> FAvaMaskAvaTextActorMaterialCollectionHandle::GetMaterialHandles()
-{
-	if (MaterialHandles.IsEmpty() || !UE::Ava::Internal::IsHandleValid(MaterialHandles.Last()))
-	{
-		if (const UAvaText3DComponent* TextComponent = GetComponent())
-		{
-			MaterialHandles.Init(nullptr, static_cast<int32>(EText3DGroupType::TypeCount));
-
-			UAvaObjectHandleSubsystem* HandleSubsystem = UE::Ava::Internal::GetObjectHandleSubsystem();
-			if (!HandleSubsystem)
-			{
-				return MaterialHandles;
-			}
-
-			UMaterialInterface* FrontMaterial = TextComponent->GetFrontMaterial();
-			MaterialHandles[static_cast<int32>(EText3DGroupType::Front)] = HandleSubsystem->MakeHandle<IAvaMaskMaterialHandle>(FrontMaterial, UE::AvaMask::Internal::HandleTag);
-
-			UMaterialInterface* BevelMaterial = TextComponent->GetBevelMaterial();
-			MaterialHandles[static_cast<int32>(EText3DGroupType::Bevel)] = HandleSubsystem->MakeHandle<IAvaMaskMaterialHandle>(BevelMaterial, UE::AvaMask::Internal::HandleTag);
-
-			UMaterialInterface* ExtrudeMaterial = TextComponent->GetBevelMaterial();
-			MaterialHandles[static_cast<int32>(EText3DGroupType::Extrude)] = HandleSubsystem->MakeHandle<IAvaMaskMaterialHandle>(ExtrudeMaterial, UE::AvaMask::Internal::HandleTag);
-
-			UMaterialInterface* BackMaterial = TextComponent->GetBevelMaterial();
-			MaterialHandles[static_cast<int32>(EText3DGroupType::Back)] = HandleSubsystem->MakeHandle<IAvaMaskMaterialHandle>(BackMaterial, UE::AvaMask::Internal::HandleTag);
-		}
-	}
-
-	return MaterialHandles;
-}
-
-int32 FAvaMaskAvaTextActorMaterialCollectionHandle::GetNumMaterials() const
-{
-	constexpr int32 TextMaterialCount = static_cast<int32>(EText3DGroupType::TypeCount);
-	static_assert(TextMaterialCount == 4);
-	
-	return TextMaterialCount;
-}
-
-void FAvaMaskAvaTextActorMaterialCollectionHandle::SetMaterial(
-	const FSoftComponentReference& InComponent
-	, const int32 InSlotIdx
-	, UMaterialInterface* InMaterial)
-{
-	if (!ensure(InSlotIdx < 4))
-	{
-		UE_LOG(LogAvaMask, Error, TEXT("SlotIdx out of range, should be from 0-3, was %u"), InSlotIdx);
-		return;
-	}
-
-	if (UAvaText3DComponent* TextComponent = GetComponent())
-    {
-		switch (InSlotIdx)
-		{
-		case 0:
-			TextComponent->SetFrontMaterial(InMaterial);
-
-		case 1:
-			TextComponent->SetBevelMaterial(InMaterial);
-
-		case 2:
-			TextComponent->SetExtrudeMaterial(InMaterial);
-
-		case 3:
-			TextComponent->SetBackMaterial(InMaterial);
-
-		default:
-			return;
-		}
-    }
-}
-
-void FAvaMaskAvaTextActorMaterialCollectionHandle::SetMaterials(
-	const TArray<TObjectPtr<UMaterialInterface>>& InMaterials
-	, const TBitArray<>& InSetToggle)
-{
-	if (!ensure(InMaterials.Num() == GetNumMaterials()))
-	{
-		UE_LOG(LogAvaMask, Warning, TEXT("Expected 4 materials, got %u"), InMaterials.Num());
-		return;
-	}
-
-	if (!ensure(InSetToggle.Num() == GetNumMaterials()))
-	{
-		UE_LOG(LogAvaMask, Warning, TEXT("Expected 4 materials, got %u"), InSetToggle.Num());
-		return;
-	}
-	
-	if (UAvaText3DComponent* TextComponent = GetComponent())
-	{
-		if (InSetToggle[0])
-		{
-			TextComponent->SetFrontMaterial(InMaterials[0]);	
-		}
-
-		if (InSetToggle[1])
-		{
-			TextComponent->SetBevelMaterial(InMaterials[1]);	
-		}
-
-		if (InSetToggle[2])
-		{
-			TextComponent->SetExtrudeMaterial(InMaterials[2]);	
-		}
-
-		if (InSetToggle[3])
-		{
-			TextComponent->SetBackMaterial(InMaterials[3]);	
-		}
-	}
-}
-
-void FAvaMaskAvaTextActorMaterialCollectionHandle::ForEachMaterial(
-	TFunctionRef<bool(
-		const FSoftComponentReference& InComponent
-		, const int32 InIdx
-		, UMaterialInterface* InMaterial)> InFunction)
-{
-	if (const UAvaText3DComponent* TextComponent = GetComponent())
-	{
-		const FSoftComponentReference ComponentReference = UE::AvaMask::Internal::MakeComponentReference(TextComponent->GetOwner(), TextComponent);
-		if (!InFunction(ComponentReference, static_cast<int32>(EText3DGroupType::Front), TextComponent->GetFrontMaterial()))
-		{
-			return;
-		}
-		
-		if (!InFunction(ComponentReference, static_cast<int32>(EText3DGroupType::Bevel), TextComponent->GetBevelMaterial()))
-		{
-			return;
-		}
-		
-		if (!InFunction(ComponentReference, static_cast<int32>(EText3DGroupType::Extrude), TextComponent->GetExtrudeMaterial()))
-		{
-			return;
-		}
-		
-		InFunction(ComponentReference, static_cast<int32>(EText3DGroupType::Back), TextComponent->GetBackMaterial());
-	}
-}
-
-void FAvaMaskAvaTextActorMaterialCollectionHandle::ForEachMaterialHandle(TFunctionRef<bool(const FSoftComponentReference& InComponent, const int32 InIdx, const bool bInIsSlotOccupied, const TSharedPtr<IAvaMaskMaterialHandle>& InMaterialHandle)> InFunction)
-{
-	if (const UAvaText3DComponent* TextComponent = GetComponent())
-	{
-		// Effectively refreshes material handles if necessary
-		GetMaterialHandles();
-		
-		const FSoftComponentReference ComponentReference = UE::AvaMask::Internal::MakeComponentReference(TextComponent->GetOwner(), TextComponent);
-		TSharedPtr<IAvaMaskMaterialHandle>& MaterialHandle = MaterialHandles[static_cast<int32>(EText3DGroupType::Front)];
-		if (ensureAlways(UE::Ava::Internal::IsHandleValid(MaterialHandle)))
-		{
-			if (!InFunction(ComponentReference, static_cast<int32>(EText3DGroupType::Front), true, MaterialHandle))
-			{
-				return;
-			}
-		}
-
-		MaterialHandle = MaterialHandles[static_cast<int32>(EText3DGroupType::Bevel)];
-		if (ensureAlways(UE::Ava::Internal::IsHandleValid(MaterialHandle)))
-		{
-			if (!InFunction(ComponentReference, static_cast<int32>(EText3DGroupType::Bevel), true, MaterialHandle))
-			{
-				return;
-			}
-		}
-
-		MaterialHandle = MaterialHandles[static_cast<int32>(EText3DGroupType::Extrude)];
-		if (ensureAlways(UE::Ava::Internal::IsHandleValid(MaterialHandle)))
-		{
-			if (!InFunction(ComponentReference, static_cast<int32>(EText3DGroupType::Extrude), true, MaterialHandle))
-			{
-				return;
-			}
-		}
-
-		MaterialHandle = MaterialHandles[static_cast<int32>(EText3DGroupType::Back)];
-		if (ensureAlways(UE::Ava::Internal::IsHandleValid(MaterialHandle)))
-		{
-			InFunction(ComponentReference, static_cast<int32>(EText3DGroupType::Back), true, MaterialHandle);
-		}
-	}
-}
-
-void FAvaMaskAvaTextActorMaterialCollectionHandle::MapEachMaterial(
-	TFunctionRef<UMaterialInterface*(
-		const FSoftComponentReference& InComponent
-		, const int32 InIdx
-		, UMaterialInterface* InMaterial)> InFunction)
-{
-	if (const UAvaText3DComponent* TextComponent = GetComponent())
-	{
-		TArray<TObjectPtr<UMaterialInterface>> MappedMaterials;
-		MappedMaterials.Init(nullptr, static_cast<int32>(EText3DGroupType::TypeCount));
-
-		TBitArray<> SetFlags;
-		SetFlags.Init(false, MappedMaterials.Num());
-
-		const FSoftComponentReference ComponentReference = UE::AvaMask::Internal::MakeComponentReference(TextComponent->GetOwner(), TextComponent);
-		
-		MappedMaterials[0] = InFunction(ComponentReference, static_cast<int32>(EText3DGroupType::Front), TextComponent->GetFrontMaterial());
-		SetFlags[0] = MappedMaterials[0] != nullptr;
-		
-		MappedMaterials[1] = InFunction(ComponentReference, static_cast<int32>(EText3DGroupType::Bevel), TextComponent->GetBevelMaterial());
-		SetFlags[1] = MappedMaterials[1] != nullptr;
-		
-		MappedMaterials[2] = InFunction(ComponentReference, static_cast<int32>(EText3DGroupType::Extrude), TextComponent->GetExtrudeMaterial());
-		SetFlags[2] = MappedMaterials[2] != nullptr;
-		
-		MappedMaterials[3] = InFunction(ComponentReference, static_cast<int32>(EText3DGroupType::Back), TextComponent->GetBackMaterial());
-		SetFlags[3] = MappedMaterials[3] != nullptr;
-		
-		SetMaterials(MappedMaterials, SetFlags);
-	}
-}
-
-void FAvaMaskAvaTextActorMaterialCollectionHandle::MapEachMaterialHandle(TFunctionRef<TSharedPtr<IAvaMaskMaterialHandle>(const FSoftComponentReference& InComponent, const int32 InIdx, const bool bInIsSlotOccupied, const TSharedPtr<IAvaMaskMaterialHandle>& InMaterialHandle)> InFunction)
-{
-	if (const UAvaText3DComponent* TextComponent = GetComponent())
-	{		
-		// Effectively refreshes material handles if necessary
-		GetMaterialHandles();
-
-		TArray<TObjectPtr<UMaterialInterface>> MappedMaterials;
-		MappedMaterials.Reserve(GetNumMaterials());
-
-		TBitArray<> SetFlags;
-		SetFlags.Init(false, GetNumMaterials());
-
-		const FSoftComponentReference ComponentReference = UE::AvaMask::Internal::MakeComponentReference(TextComponent->GetOwner(), TextComponent);
-
-		auto TryMap = [this, InFunction, &ComponentReference, &MappedMaterials, &SetFlags](const EText3DGroupType InMaterialGroup)
-		{
-			int32 InMaterialGroupIdx = static_cast<int32>(InMaterialGroup);			
-			
-			TSharedPtr<IAvaMaskMaterialHandle>& MaterialHandle = MaterialHandles[InMaterialGroupIdx];
-			if (ensureAlways(UE::Ava::Internal::IsHandleValid(MaterialHandle)))
-			{
-				const TSharedPtr<IAvaMaskMaterialHandle> MappedMaterialHandle = InFunction(ComponentReference, InMaterialGroupIdx, true, MaterialHandle);
-				UMaterialInterface* RemappedMaterial = nullptr;
-				if (MappedMaterialHandle.IsValid())
-				{
-					RemappedMaterial = MappedMaterialHandle->GetMaterial();	
-				}
-				else
-				{
-					RemappedMaterial = MaterialHandle->GetMaterial();
-				}
-
-				MappedMaterials[InMaterialGroupIdx] = RemappedMaterial;
-				SetFlags[InMaterialGroupIdx] = RemappedMaterial != nullptr;
-			}	
-		};
-
-		TryMap(EText3DGroupType::Front);
-		TryMap(EText3DGroupType::Bevel);
-		TryMap(EText3DGroupType::Extrude);
-		TryMap(EText3DGroupType::Back);
-
-		SetMaterials(MappedMaterials, SetFlags);
-	}
-}
-
-bool FAvaMaskAvaTextActorMaterialCollectionHandle::SaveOriginalState(const FStructView& InHandleData)
+bool FAvaMaskText3DActorMaterialCollectionHandle::SaveOriginalState(const FStructView& InHandleData)
 {
 	if (FHandleData* HandleData = InHandleData.GetPtr<FHandleData>())
 	{
@@ -673,13 +334,22 @@ bool FAvaMaskAvaTextActorMaterialCollectionHandle::SaveOriginalState(const FStru
 			, const bool bInIsSlotOccupied
 			, const TSharedPtr<IAvaMaskMaterialHandle>& InMaterialHandle)
 		{
-			FInstancedStruct& MaterialHandleData = HandleData->GroupMaterialData.FindOrAdd(InSlotIdx, InMaterialHandle->MakeDataStruct());
-			return InMaterialHandle->SaveOriginalState(MaterialHandleData);
+			if (InMaterialHandle.IsValid())
+			{
+				FInstancedStruct& MaterialHandleData = HandleData->GroupMaterialData.FindOrAdd(InSlotIdx, InMaterialHandle->MakeDataStruct());
+				return InMaterialHandle->SaveOriginalState(MaterialHandleData);
+			}
+
+			return false;
 		});
 
-		if (const UAvaText3DComponent* TextComponent = GetComponent())
+		if (const UText3DComponent* TextComponent = GetComponent())
 		{
-			HandleData->OriginalTranslucencyStyle = TextComponent->GetTranslucencyStyle();
+			if (const UText3DDefaultMaterialExtension* MaterialExtension = TextComponent->GetCastedMaterialExtension<UText3DDefaultMaterialExtension>())
+			{
+				HandleData->BlendMode = MaterialExtension->GetBlendMode();
+			}
+
 			return true;
 		}
 		
@@ -689,17 +359,16 @@ bool FAvaMaskAvaTextActorMaterialCollectionHandle::SaveOriginalState(const FStru
 	return false;
 }
 
-bool FAvaMaskAvaTextActorMaterialCollectionHandle::ApplyOriginalState(const FStructView& InHandleData)
+bool FAvaMaskText3DActorMaterialCollectionHandle::ApplyOriginalState(const FStructView& InHandleData)
 {
 	if (const FHandleData* HandleData = InHandleData.GetPtr<FHandleData>())
 	{
-		if (UAvaText3DComponent* TextComponent = GetComponent())
+		if (UText3DComponent* TextComponent = GetComponent())
 		{
 			// Only restore if the original was NOT translucent AND the current alpha is 1.0
-			if (HandleData->OriginalTranslucencyStyle == EAvaTextTranslucency::None
-				&& FMath::IsNearlyEqual(TextComponent->GetOpacity(), 1.0f))
+			if (UText3DDefaultMaterialExtension* MaterialExtension = TextComponent->GetCastedMaterialExtension<UText3DDefaultMaterialExtension>())
 			{
-				TextComponent->SetTranslucencyStyle(HandleData->OriginalTranslucencyStyle);	
+				MaterialExtension->SetBlendMode(HandleData->BlendMode);
 			}
 			
 			Super::ApplyOriginalState(InHandleData);
@@ -713,23 +382,27 @@ bool FAvaMaskAvaTextActorMaterialCollectionHandle::ApplyOriginalState(const FStr
 	return false;
 }
 
-bool FAvaMaskAvaTextActorMaterialCollectionHandle::ApplyModifiedState(
-	const FAvaMask2DSubjectParameters& InModifiedParameters
-	, const FStructView& InHandleData)
+bool FAvaMaskText3DActorMaterialCollectionHandle::ApplyModifiedState(const FAvaMask2DSubjectParameters& InModifiedParameters,
+	const FStructView& InHandleData)
 {
 	if (InHandleData.GetPtr<FHandleData>())
 	{
-		if (!Super::ApplyModifiedState(InModifiedParameters, InHandleData))
+		if (UText3DComponent* TextComponent = GetComponent())
 		{
-			return false;
+			if (UText3DDefaultMaterialExtension* MaterialExtension = TextComponent->GetCastedMaterialExtension<UText3DDefaultMaterialExtension>())
+			{
+				const EText3DMaterialBlendMode TargetTranslucencyStyle = UE::AvaMask::Private::GetTargetTranslucencyType(MaterialExtension->GetBlendMode(), InModifiedParameters.MaterialParameters.BlendMode);
+				MaterialExtension->SetBlendMode(TargetTranslucencyStyle);
+
+				// Refresh materials if needed immediately
+				constexpr bool bImmediateUpdate = true;
+				TextComponent->RequestUpdate(EText3DRendererFlags::None, bImmediateUpdate);
+			}
 		}
 
-		if (UAvaText3DComponent* TextComponent = GetComponent())
+		if (!Super::ApplyModifiedState(InModifiedParameters, InHandleData))
 		{
-			const EAvaTextTranslucency TargetTranslucencyStyle = UE::AvaMask::Private::GetTargetTranslucencyType(TextComponent->GetTranslucencyStyle(), InModifiedParameters.MaterialParameters.BlendMode);
-			TextComponent->SetTranslucencyStyle(TargetTranslucencyStyle);
-			TextComponent->RefreshMaterialInstances();
-			return true;
+			return false;
 		}
 
 		return true;
@@ -738,25 +411,20 @@ bool FAvaMaskAvaTextActorMaterialCollectionHandle::ApplyModifiedState(
 	return false;
 }
 
-bool FAvaMaskAvaTextActorMaterialCollectionHandle::IsValid() const
-{
-	return WeakActor.IsValid() && WeakComponent.IsValid();
-}
-
-bool FAvaMaskAvaTextActorMaterialCollectionHandle::IsSupported(const FAvaMaskMaterialReference& InInstance, FName InTag)
+bool FAvaMaskText3DActorMaterialCollectionHandle::IsSupported(const FAvaMaskMaterialReference& InInstance, FName InTag)
 {
 	if (InTag == UE::AvaMask::Internal::HandleTag)
 	{
 		if (const AActor* Actor = InInstance.GetTypedObject<AActor>())
 		{
-			return Actor->GetComponentByClass<UAvaText3DComponent>() != nullptr;
+			return Actor->GetComponentByClass<UText3DComponent>() != nullptr;
 		}
 	}
 
 	return false;
 }
 
-FStructView FAvaMaskAvaTextActorMaterialCollectionHandle::GetMaterialHandleData(
+FStructView FAvaMaskText3DActorMaterialCollectionHandle::GetMaterialHandleData(
 	FHandleData* InParentHandleData
 	, const FSoftComponentReference& InComponent
 	, const int32 InSlotIdx)
@@ -769,10 +437,10 @@ FStructView FAvaMaskAvaTextActorMaterialCollectionHandle::GetMaterialHandleData(
 		}
 	}
 
-	return nullptr;	
+	return nullptr;
 }
 
-FStructView FAvaMaskAvaTextActorMaterialCollectionHandle::GetOrAddMaterialHandleData(
+FStructView FAvaMaskText3DActorMaterialCollectionHandle::GetOrAddMaterialHandleData(
 	FHandleData* InParentHandleData
 	, const TSharedPtr<IAvaMaskMaterialHandle>& InMaterialHandle
 	, const FSoftComponentReference& InComponent
@@ -786,17 +454,17 @@ FStructView FAvaMaskAvaTextActorMaterialCollectionHandle::GetOrAddMaterialHandle
 	return nullptr;
 }
 
-UAvaText3DComponent* FAvaMaskAvaTextActorMaterialCollectionHandle::GetComponent()
+UText3DComponent* FAvaMaskText3DActorMaterialCollectionHandle::GetComponent()
 {
-	if (UAvaText3DComponent* Component = WeakComponent.Get())
+	if (UText3DComponent* Component = WeakComponent.Get())
 	{
 		return Component;
 	}
 
 	if (const AActor* Actor = WeakActor.Get())
 	{
-		WeakComponent = Actor->GetComponentByClass<UAvaText3DComponent>();
-		if (UAvaText3DComponent* Component = WeakComponent.Get())
+		WeakComponent = Actor->GetComponentByClass<UText3DComponent>();
+		if (UText3DComponent* Component = WeakComponent.Get())
 		{
 			return Component;
 		}
@@ -805,20 +473,34 @@ UAvaText3DComponent* FAvaMaskAvaTextActorMaterialCollectionHandle::GetComponent(
 	return nullptr;
 }
 
-UMaterialInterface* FAvaMaskAvaTextActorMaterialCollectionHandle::ProvideMaterialWithSettings(const UMaterialInterface* InPreviousMaterial, const FAvaTextMaterialSettings& InMaterialSettings)
+void FAvaMaskText3DActorMaterialCollectionHandle::OnTextPostUpdate(UText3DComponent* InComponent, EText3DRendererFlags InFlags)
 {
-	UMaterialInterface* AttemptedAssignedMaterial = UAvaTextMaterialHub::GetMaterial(InMaterialSettings);
-
-	if (OnSourceMaterialPreAssignment().IsBound())
+	if (EnumHasAnyFlags(InFlags, EText3DRendererFlags::Material))
 	{
-		UMaterialInterface* UserAssignedMaterial = OnSourceMaterialPreAssignment().Execute(InPreviousMaterial, AttemptedAssignedMaterial);
-
-		// Flag handles for refresh
-		MaterialHandles.Reset();
-		GetMaterialHandles();
+		constexpr int32 FrontIndex = static_cast<int32>(EText3DGroupType::Front);
+		constexpr int32 BackIndex = static_cast<int32>(EText3DGroupType::Back);
+		constexpr int32 ExtrudeIndex = static_cast<int32>(EText3DGroupType::Extrude);
+		constexpr int32 BevelIndex = static_cast<int32>(EText3DGroupType::Bevel);
+		
+		bool bChanged = false;
+		if (MaterialHandles.IsEmpty() || !UE::Ava::Internal::IsHandleValid(MaterialHandles.Last()))
+		{
+			GetMaterialHandles();
+			bChanged = true;
+		}
+		else if (MaterialHandles[FrontIndex]->GetMaterial() != InComponent->GetFrontMaterial()
+			|| MaterialHandles[BevelIndex]->GetMaterial() != InComponent->GetBevelMaterial()
+			|| MaterialHandles[ExtrudeIndex]->GetMaterial() != InComponent->GetExtrudeMaterial()
+			|| MaterialHandles[BackIndex]->GetMaterial() != InComponent->GetBackMaterial())
+		{
+			MaterialHandles.Empty();
+			GetMaterialHandles();
+			bChanged = true;
+		}
 
-		return UserAssignedMaterial;
+		if (bChanged)
+		{
+			OnSourceMaterialsChanged().ExecuteIfBound(InComponent, MaterialHandles);
+		}
 	}
-
-	return nullptr;
 }
diff --git a/Engine/Plugins/Experimental/Avalanche/Source/AvalancheMask/Private/Handling/AvaMaskText3DActorMaterialCollectionHandle.h b/Engine/Plugins/Experimental/Avalanche/Source/AvalancheMask/Private/Handling/AvaMaskText3DActorMaterialCollectionHandle.h
--- a/Engine/Plugins/Experimental/Avalanche/Source/AvalancheMask/Private/Handling/AvaMaskText3DActorMaterialCollectionHandle.h
+++ b/Engine/Plugins/Experimental/Avalanche/Source/AvalancheMask/Private/Handling/AvaMaskText3DActorMaterialCollectionHandle.h
@@ -20,6 +20,7 @@ class UText3DComponent;
 class AText3DActor;
 class AActor;
 class UMaterialInterface;
+enum class EText3DRendererFlags : uint8;
 struct FAvaMaskMaterialReference;
 
 USTRUCT()
@@ -28,7 +29,7 @@ struct FAvaMaskText3DActorMaterialCollectionHandleData
 	GENERATED_BODY()
 
 	UPROPERTY()
-	EAvaTextTranslucency OriginalTranslucencyStyle = EAvaTextTranslucency::None;
+	EText3DMaterialBlendMode BlendMode = EText3DMaterialBlendMode::Opaque;
 
 	UPROPERTY()
 	TMap<int32, FInstancedStruct> GroupMaterialData;
@@ -40,9 +41,9 @@ class FAvaMaskText3DActorMaterialCollectionHandle
 {
 public:
 	UE_AVA_INHERITS_WITH_SUPER(FAvaMaskText3DActorMaterialCollectionHandle, TAvaMaskMaterialCollectionHandle<FAvaMaskText3DActorMaterialCollectionHandleData>);
-	
+
 	explicit FAvaMaskText3DActorMaterialCollectionHandle(AActor* InActor);
-	virtual ~FAvaMaskText3DActorMaterialCollectionHandle() override = default;
+	virtual ~FAvaMaskText3DActorMaterialCollectionHandle() override;
 
 	// ~Begin IAvaMaskMaterialCollectionProvider
 	virtual TArray<TObjectPtr<UMaterialInterface>> GetMaterials() override;
@@ -81,82 +82,10 @@ public:
 			, const TSharedPtr<IAvaMaskMaterialHandle>& InMaterialHandle)> InFunction) override;
 
 	virtual bool IsValid() const override;
-	// ~End IAvaMaskMaterialCollectionProvider
-
-	static bool IsSupported(const FAvaMaskMaterialReference& InInstance, FName InTag = NAME_None);
-
-protected:
-	// ~Begin TAvaMaskMaterialCollectionHandle
-	virtual FStructView GetMaterialHandleData(
-		FHandleData* InParentHandleData
-		, const FSoftComponentReference& InComponent
-		, const int32 InSlotIdx) override;
-	
-	virtual ::FStructView GetOrAddMaterialHandleData(
-		FHandleData* InParentHandleData
-		, const TSharedPtr<IAvaMaskMaterialHandle>& InMaterialHandle
-		, const FSoftComponentReference& InComponent
-		, const int32 InSlotIdx) override;
-	// ~End TAvaMaskMaterialCollectionHandle
-	
-	UText3DComponent* GetComponent();
-	
-protected:
-	TWeakObjectPtr<AActor> WeakActor;
-	TWeakObjectPtr<UText3DComponent> WeakComponent;
-	TArray<TSharedPtr<IAvaMaskMaterialHandle>> MaterialHandles;
-};
-
-/** Note that this doesn't explicitly require an Ava Text Actor, it just checks that whatever Actor is given has a UText3DComponent (or derived). */
-class FAvaMaskAvaTextActorMaterialCollectionHandle
-	: public TAvaMaskMaterialCollectionHandle<FAvaMaskText3DActorMaterialCollectionHandleData>
-{
-public:
-	UE_AVA_INHERITS_WITH_SUPER(FAvaMaskAvaTextActorMaterialCollectionHandle, TAvaMaskMaterialCollectionHandle<FAvaMaskText3DActorMaterialCollectionHandleData>);
-	
-	explicit FAvaMaskAvaTextActorMaterialCollectionHandle(AActor* InActor);
-	virtual ~FAvaMaskAvaTextActorMaterialCollectionHandle() override;
-	
-	// ~Begin IAvaMaskMaterialCollectionHandle
-	virtual TArray<TObjectPtr<UMaterialInterface>> GetMaterials() override;
-	virtual TArray<TSharedPtr<IAvaMaskMaterialHandle>> GetMaterialHandles() override;
-	virtual int32 GetNumMaterials() const override;
-
-	virtual void SetMaterial(const FSoftComponentReference& InComponent, const int32 InSlotIdx, UMaterialInterface* InMaterial) override;
-	virtual void SetMaterials(const TArray<TObjectPtr<UMaterialInterface>>& InMaterials, const TBitArray<>& InSetToggle) override;
-
-	virtual void ForEachMaterial(
-		TFunctionRef<bool(
-			const FSoftComponentReference& InComponent
-			, const int32 InIdx
-			, UMaterialInterface* InMaterial)> InFunction) override;
-
-	virtual void ForEachMaterialHandle(
-		TFunctionRef<bool(
-			const FSoftComponentReference& InComponent
-			, const int32 InIdx
-			, const bool bInIsSlotOccupied
-			, const TSharedPtr<IAvaMaskMaterialHandle>& InMaterialHandle)> InFunction) override;
-
-	virtual void MapEachMaterial(
-		TFunctionRef<UMaterialInterface*(
-			const FSoftComponentReference& InComponent
-			, const int32 InIdx
-			, UMaterialInterface* InMaterial)> InFunction) override;
-
-	virtual void MapEachMaterialHandle(
-		TFunctionRef<TSharedPtr<IAvaMaskMaterialHandle>(
-			const FSoftComponentReference& InComponent
-			, const int32 InIdx
-			, const bool bInIsSlotOccupied
-			, const TSharedPtr<IAvaMaskMaterialHandle>& InMaterialHandle)> InFunction) override;
-	
 	virtual bool SaveOriginalState(const FStructView& InHandleData) override;
 	virtual bool ApplyOriginalState(const FStructView& InHandleData) override;
 	virtual bool ApplyModifiedState(const FAvaMask2DSubjectParameters& InModifiedParameters, const FStructView& InHandleData) override;
-
-	virtual bool IsValid() const override;
-	// ~End IAvaMaskMaterialCollectionHandle
+	// ~End IAvaMaskMaterialCollectionProvider
 
 	static bool IsSupported(const FAvaMaskMaterialReference& InInstance, FName InTag = NAME_None);
 
@@ -174,12 +103,11 @@ protected:
 		, const int32 InSlotIdx) override;
 	// ~End TAvaMaskMaterialCollectionHandle
 	
-	UAvaText3DComponent* GetComponent();
+	UText3DComponent* GetComponent();
 
-	UMaterialInterface* ProvideMaterialWithSettings(const UMaterialInterface* InPreviousMaterial, const FAvaTextMaterialSettings& InMaterialSettings);
+	void OnTextPostUpdate(UText3DComponent* InComponent, EText3DRendererFlags InFlags);
 
-protected:
 	TWeakObjectPtr<AActor> WeakActor;
-	TWeakObjectPtr<UAvaText3DComponent> WeakComponent;
+	TWeakObjectPtr<UText3DComponent> WeakComponent;
 	TArray<TSharedPtr<IAvaMaskMaterialHandle>> MaterialHandles;
 };
diff --git a/Engine/Plugins/Experimental/Avalanche/Source/AvalancheMask/Private/Handling/AvaObjectHandleSubsystem.cpp b/Engine/Plugins/Experimental/Avalanche/Source/AvalancheMask/Private/Handling/AvaObjectHandleSubsystem.cpp
--- a/Engine/Plugins/Experimental/Avalanche/Source/AvalancheMask/Private/Handling/AvaObjectHandleSubsystem.cpp
+++ b/Engine/Plugins/Experimental/Avalanche/Source/AvalancheMask/Private/Handling/AvaObjectHandleSubsystem.cpp
@@ -57,13 +57,6 @@ void UAvaObjectHandleSubsystem::FindObjectHandleFactories()
 {
 	// Material Collection Handles
 	{
-		ObjectHandleFactories.Add({
-			FAvaMaskAvaTextActorMaterialCollectionHandle::IsSupported
-			, [](const FAvaMaskMaterialReference& InObject) -> TSharedPtr<IAvaMaskMaterialCollectionHandle>
-			{
-				return MakeShared<FAvaMaskAvaTextActorMaterialCollectionHandle>(InObject.GetTypedObject<AActor>());
-			}});
-
 		ObjectHandleFactories.Add({
 			FAvaMaskAvaShapeMaterialCollectionHandle::IsSupported
 			, [](const FAvaMaskMaterialReference& InObject) -> TSharedPtr<IAvaMaskMaterialCollectionHandle>
diff --git a/Engine/Plugins/Experimental/Avalanche/Source/AvalancheMask/Private/Handling/IAvaMaskMaterialCollectionHandle.h b/Engine/Plugins/Experimental/Avalanche/Source/AvalancheMask/Private/Handling/IAvaMaskMaterialCollectionHandle.h
--- a/Engine/Plugins/Experimental/Avalanche/Source/AvalancheMask/Private/Handling/IAvaMaskMaterialCollectionHandle.h
+++ b/Engine/Plugins/Experimental/Avalanche/Source/AvalancheMask/Private/Handling/IAvaMaskMaterialCollectionHandle.h
@@ -25,7 +25,7 @@ public:
 	UE_AVA_INHERITS(IAvaMaskMaterialCollectionHandle, IAvaObjectHandle);
 	
 	using FOnSourceMaterialPreAssignment = TDelegate<UMaterialInterface*(const UMaterialInterface* InPreviousMaterial, UMaterialInterface* InNewMaterial)>;
-	using FOnSourceMaterialsChanged = TDelegate<void(UPrimitiveComponent* InComponent, const TArray<TSharedPtr<IAvaMaskMaterialHandle>>& InMaterialHandles)>;
+	using FOnSourceMaterialsChanged = TDelegate<void(UObject* InMaterialOwner, const TArray<TSharedPtr<IAvaMaskMaterialHandle>>& InMaterialHandles)>;
 	
 public:
 	virtual ~IAvaMaskMaterialCollectionHandle() override = default;
diff --git a/Engine/Plugins/Experimental/Avalanche/Source/AvalancheMask/Private/Mask2D/AvaMask2DBaseModifier.cpp b/Engine/Plugins/Experimental/Avalanche/Source/AvalancheMask/Private/Mask2D/AvaMask2DBaseModifier.cpp
--- a/Engine/Plugins/Experimental/Avalanche/Source/AvalancheMask/Private/Mask2D/AvaMask2DBaseModifier.cpp
+++ b/Engine/Plugins/Experimental/Avalanche/Source/AvalancheMask/Private/Mask2D/AvaMask2DBaseModifier.cpp
@@ -812,7 +812,7 @@ UGeometryMaskCanvas* UAvaMask2DBaseModifier::GetCurrentCanvas()
 	return nullptr;
 }
 
-void UAvaMask2DBaseModifier::OnMaterialsChanged(UPrimitiveComponent* InPrimitiveComponent, const TArray<TSharedPtr<IAvaMaskMaterialHandle>>& InMaterialHandles)
+void UAvaMask2DBaseModifier::OnMaterialsChanged(UObject* InMaterialOwner, const TArray<TSharedPtr<IAvaMaskMaterialHandle>>& InMaterialHandles)
 {
 	MarkModifierDirty();
 }
diff --git a/Engine/Plugins/Experimental/Avalanche/Source/AvalancheMask/Public/Mask2D/AvaMask2DBaseModifier.h b/Engine/Plugins/Experimental/Avalanche/Source/AvalancheMask/Public/Mask2D/AvaMask2DBaseModifier.h
--- a/Engine/Plugins/Experimental/Avalanche/Source/AvalancheMask/Public/Mask2D/AvaMask2DBaseModifier.h
+++ b/Engine/Plugins/Experimental/Avalanche/Source/AvalancheMask/Public/Mask2D/AvaMask2DBaseModifier.h
@@ -168,7 +168,7 @@ protected:
 	/** Returns or resolves the currently referenced canvas. */
 	UGeometryMaskCanvas* GetCurrentCanvas();
 	
-	virtual void OnMaterialsChanged(UPrimitiveComponent* InPrimitiveComponent, const TArray<TSharedPtr<IAvaMaskMaterialHandle>>& InMaterialHandles);
+	virtual void OnMaterialsChanged(UObject* InMaterialOwner, const TArray<TSharedPtr<IAvaMaskMaterialHandle>>& InMaterialHandles);
 
 #if WITH_EDITOR
 	UFUNCTION(CallInEditor, Category = "Mask2D")
@@ -231,9 +231,6 @@ protected:
 	UPROPERTY(EditAnywhere, Transient, DuplicateTransient, Category = "Mask2D|Shared", NoClear,	meta = (DisplayName = "Canvas", AllowPrivateAccess = "true", ShowInnerProperties, NoResetToDefault))
 	TWeakObjectPtr<UGeometryMaskCanvas> CanvasWeak;
 
-	/** Reference to the underlying canvas texture */
-	TWeakObjectPtr<UTexture> CanvasTextureWeak;
-
 	UPROPERTY(DuplicateTransient)
 	TMap<TWeakObjectPtr<UObject>, FInstancedStruct> MaterialCollectionHandleData;
 
diff --git a/Engine/Plugins/Experimental/Avalanche/Source/AvalancheText/Private/AvaText3DComponent.cpp b/Engine/Plugins/Experimental/Avalanche/Source/AvalancheText/Private/AvaText3DComponent.cpp
--- a/Engine/Plugins/Experimental/Avalanche/Source/AvalancheText/Private/AvaText3DComponent.cpp
+++ b/Engine/Plugins/Experimental/Avalanche/Source/AvalancheText/Private/AvaText3DComponent.cpp
@@ -3,1557 +3,167 @@
 #include "AvaText3DComponent.h"
 #include "AvaActorUtils.h"
 #include "AvaLog.h"
+#include "Characters/Text3DCharacterTransform.h"
 #include "Engine/Texture2D.h"
+#include "Extensions/Text3DDefaultMaterialExtension.h"
+#include "Extensions/Text3DDefaultTransformExtension.h"
+#include "Extensions/Text3DMaterialExtensionBase.h"
 #include "GameFramework/Actor.h"
-#include "MaterialHub/AvaTextMaterialHub.h"
 #include "Materials/Material.h"
 #include "Materials/MaterialInstanceDynamic.h"
+#include "Text3DComponentVersion.h"
 
 #if WITH_EDITOR
 #include "Editor.h"
 #endif
 
-namespace UE::AvaText::Private
-{
-	void PatchMaterial(
-		UText3DComponent* InComponent
-		, TMemFunPtrType<true, UText3DComponent, UMaterialInterface*()>::Type&& InGetter
-		, TMemFunPtrType<false, UText3DComponent, void(UMaterialInterface*)>::Type&& InSetter
-		, UMaterialInstanceDynamic* InNewMaterial)
-	{
-		if (!InNewMaterial)
-		{
-			return;
-		}
-
-		if (UMaterialInstance* CurrentMaterialInstance = Cast<UMaterialInstance>(Invoke(InGetter, InComponent)))
-		{
-			if (CurrentMaterialInstance == InNewMaterial)
-			{
-				return;
-			}
-
-			InNewMaterial->CopyInterpParameters(CurrentMaterialInstance);
-		}
-
-		Invoke(InSetter, InComponent, InNewMaterial);
-	}
-}
-
 // Sets default values for this component's properties
 UAvaText3DComponent::UAvaText3DComponent()
 {
-	PrimaryComponentTick.bCanEverTick = true;
-
-	Extrude = 0.0f;
-	Bevel = 0.0f;
-	BevelType = EText3DBevelType::Convex;
-	BevelSegments = 8;
-
-	HorizontalAlignment = EText3DHorizontalTextAlignment::Left;
-	VerticalAlignment = EText3DVerticalTextAlignment::FirstLine;
+	SetExtrude(0.0f);
+	SetBevel(0.0f);
+	SetBevelType(EText3DBevelType::Convex);
+	SetBevelSegments(8);
 
-	bHasMaxWidth = false;
-	bHasMaxHeight = false;
-	bScaleProportionally = false;
+	SetHorizontalAlignment(EText3DHorizontalTextAlignment::Left);
+	SetVerticalAlignment(EText3DVerticalTextAlignment::FirstLine);
 
-	Kerning = 0.0f;
-	WordSpacing = 0.0f;
-	LineSpacing = 0.0f;
+	SetHasMaxWidth(false);
+	SetHasMaxHeight(false);
+	SetScaleProportionally(false);
 
-	MaxWidth = 100.0f;
-	MaxHeight = 100.0f;
+	SetKerning(0.f);
+	SetWordSpacing(0.f);
+	SetLineSpacing(0.f);
 
-	MotionDesignFont = {};
-
-	Tiling = FVector2D(1.0, 1.0);
-
-	Color = FLinearColor::White;
-	ExtrudeColor = FLinearColor::Gray;
-	BevelColor = FLinearColor::Gray;
-
-	GradientSettings.Direction = EAvaGradientDirection::Vertical;
-	GradientSettings.Smoothness = 0.1;
-
-	bIsUnlit = true;
-	TranslucencyStyle = EAvaTextTranslucency::None;
+	SetMaxWidth(100.f);
+	SetMaxHeight(100.f);
 
 	bEnforceUpperCase = false;
-
-	MarkForFullRefresh();
-
-#if WITH_EDITOR
-	RegisterOnPropertyChangeFunctions();
-#endif
-}
-
-void UAvaText3DComponent::SetEnforceUpperCase(bool bInEnforceUpperCase)
-{
-	if (bEnforceUpperCase == bInEnforceUpperCase)
-	{
-		return;
-	}
-
-	bEnforceUpperCase = bInEnforceUpperCase;
-	RefreshFormatting();
-}
-
-void UAvaText3DComponent::SetMotionDesignFont(const FAvaFont& InFont)
-{
-	if (MotionDesignFont == InFont)
-	{
-		return;
-	}
-	MotionDesignFont = InFont;
-	RefreshGeometry();
-}
-
-void UAvaText3DComponent::SetAlignment(FAvaTextAlignment InAlignment)
-{
-	if (Alignment == InAlignment)
-	{
-		return;
-	}
-
-	Alignment = InAlignment;
-	RefreshAlignment();
-}
-
-void UAvaText3DComponent::SetColoringStyle(const EAvaTextColoringStyle& InColoringStyle)
-{
-	if (ColoringStyle != InColoringStyle)
-	{
-		ColoringStyle = InColoringStyle;
-		RefreshMaterialInstances();
-	}
-}
-
-void UAvaText3DComponent::SetColor(const FLinearColor& InColor)
-{
-	if (Color != InColor)
-	{
-		Color = InColor;
-		RefreshColor();
-	}
-}
-
-void UAvaText3DComponent::SetExtrudeColor(const FLinearColor& InColor)
-{
-	if (ExtrudeColor != InColor)
-	{
-		ExtrudeColor = InColor;
-		RefreshExtrudeColor();
-	}
-}
-
-void UAvaText3DComponent::SetBevelColor(const FLinearColor& InColor)
-{
-	if (BevelColor != InColor)
-	{
-		BevelColor = InColor;
-		RefreshBevelColor();
-	}
-}
-
-void UAvaText3DComponent::SetGradientSettings(const FAvaLinearGradientSettings& InGradientSettings)
-{
-	if (GradientSettings == InGradientSettings)
-	{
-		return;
-	}
-
-	GradientSettings = InGradientSettings;
-	RefreshGradient();
-}
-
-void UAvaText3DComponent::SetGradientColors(const FLinearColor& InColorA, const FLinearColor& InColorB)
-{
-	if (GradientSettings.ColorA == InColorA && GradientSettings.ColorB == InColorB)
-	{
-		return;
-	}
-
-	GradientSettings.ColorA = InColorA;
-	GradientSettings.ColorB = InColorB;
-	RefreshGradient();
-}
-
-void UAvaText3DComponent::SetMainTexture(UTexture2D* InMainTexture)
-{
-	if (MainTexture != InMainTexture)
-	{
-		MainTexture = InMainTexture;
-		RefreshMaterialInstances();
-	}
-}
-
-void UAvaText3DComponent::SetTiling(const FVector2D InTiling)
-{
-	if (Tiling != InTiling)
-	{
-		Tiling = InTiling;
-		RefreshMaterialInstances();
-	}
-}
-
-void UAvaText3DComponent::SetTranslucencyStyle(EAvaTextTranslucency InTranslucencyStyle)
-{
-	if (TranslucencyStyle != InTranslucencyStyle)
-	{
-		TranslucencyStyle = InTranslucencyStyle;
-		RefreshMaterialInstances();
-	}
-}
-
-void UAvaText3DComponent::SetOpacity(float InOpacity)
-{
-	if (Opacity != InOpacity)
-	{
-		Opacity = InOpacity;
-		RefreshOpacity();
-	}
-}
-
-void UAvaText3DComponent::SetMaskOrientation(EAvaMaterialMaskOrientation InMaskOrientation)
-{
-	if (MaskOrientation != InMaskOrientation)
-	{
-		MaskOrientation = InMaskOrientation;
-		RefreshMask();
-	}
-}
-
-void UAvaText3DComponent::SetMaskSmoothness(float InMaskSmoothness)
-{
-	if (MaskSmoothness != InMaskSmoothness)
-	{
-		MaskSmoothness = InMaskSmoothness;
-		RefreshMask();
-	}
-}
-
-void UAvaText3DComponent::SetMaskOffset(float InMaskOffset)
-{
-	if (MaskOffset != InMaskOffset)
-	{
-		MaskOffset = InMaskOffset;
-		RefreshMask();
-	}
-}
-
-void UAvaText3DComponent::SetMaskRotation(float InMaskRotation)
-{
-	if (MaskRotation != InMaskRotation)
-	{
-		MaskRotation = InMaskRotation;
-		RefreshMask();
-	}
-}
-
-void UAvaText3DComponent::SetIsUnlit(bool bInIsUnlit)
-{
-	if (bIsUnlit != bInIsUnlit)
-	{
-		bIsUnlit = bInIsUnlit;
-		RefreshUnlit();
-	}
-}
-
-UAvaText3DComponent::FOnGetMaterialWithSettings& UAvaText3DComponent::GetMaterialProviderDelegate()
-{
-	return OnGetMaterialWithSettingsDelegate;
-}
-
-void UAvaText3DComponent::ForEachMID(TUniqueFunction<void(UMaterialInstanceDynamic* InMID)>&& InFunc) const
-{
-	auto ExecuteIfMID = [&InFunc](UMaterialInterface* InMaterial)
-	{
-		if (UMaterialInstanceDynamic* MID = Cast<UMaterialInstanceDynamic>(InMaterial))
-		{
-			InFunc(MID);
-		}
-	};
-
-	ExecuteIfMID(GetFrontMaterial());
-	ExecuteIfMID(GetBevelMaterial());
-	ExecuteIfMID(GetExtrudeMaterial());
-	ExecuteIfMID(GetBackMaterial());
 }
 
-FVector UAvaText3DComponent::GetGradientDirection() const
+void UAvaText3DComponent::Serialize(FArchive& InArchive)
 {
-	// Gradient Rot is a [0;1] ratio, representing to [0;360] degrees
-	float GradientRot;
+	Super::Serialize(InArchive);
 
-	if (GradientSettings.Direction == EAvaGradientDirection::Vertical)
-	{
-		GradientRot = UAvaTextMaterialHub::FMaterialStatics::MaterialMaskTopBottomRotationValue;
-	}
-	else if (GradientSettings.Direction == EAvaGradientDirection::Horizontal)
-	{
-		GradientRot = UAvaTextMaterialHub::FMaterialStatics::MaterialMaskLeftToRightRotationValue;
-	}
-	else
-	{
-		GradientRot = GradientSettings.Rotation;
-	}
-
-	FVector GradientDir = GetUpVector().RotateAngleAxis(-GradientRot * 360, GetForwardVector());
-
-	/**
-	 * In order to properly map gradient offset along the text surface, text bounds are not normalized (anymore) in the material function creating the gradient.
-	 * Therefore, we need to remap the gradient direction, taking into account the current text actor bounds.
-	 *
-	 * (material function: '/Avalanche/Text3DResources/MaterialFunctions/MF_LocalGradientMask.MF_LocalGradientMask')
-	 */
-	{
-		const FVector GradientDirFixer = FVector(1.0f, LocalBoundsExtent.Z, LocalBoundsExtent.Y);
-		GradientDir*=GradientDirFixer;
-	}
-
-	GradientDir.Normalize();
-
-	return GradientDir;
-}
-
-void UAvaText3DComponent::RefreshMaterialInstances()
-{
-	MarkForMaterialsFullRefresh();
-	RefreshText3D();
-}
-
-void UAvaText3DComponent::RefreshText3D()
-{
-	SetupTextGeometry();
-	SetupTextLayout();
-	SetupTextMaterials(ColoringStyle);
-
-	RefreshGradientValues();
-	RefreshMaskValues();
-	RefreshMaterialGeometryParameterValues();
-
-	ClearRefreshReasons();
-	MarkRenderDirty();
-}
-
-void UAvaText3DComponent::MarkRenderDirty()
-{
-	MarkRenderStateDirty();
-	MarkRenderInstancesDirty();
-}
-
-void UAvaText3DComponent::MarkForGeometryRefresh()
-{
-	EnumAddFlags(TextRefreshReason, EAvaTextRefreshReason::GeometryChange);
-}
-
-bool UAvaText3DComponent::IsMarkedForGeometryRefresh() const
-{
-	return EnumHasAnyFlags(TextRefreshReason, EAvaTextRefreshReason::GeometryChange);
-}
-
-void UAvaText3DComponent::MarkForLayoutRefresh()
-{
-	EnumAddFlags(TextRefreshReason, EAvaTextRefreshReason::LayoutChange);
-}
-
-bool UAvaText3DComponent::IsMarkedForLayoutRefresh() const
-{
-	return EnumHasAnyFlags(TextRefreshReason, EAvaTextRefreshReason::LayoutChange);
-}
-
-void UAvaText3DComponent::MarkForMaterialsFullRefresh()
-{
-	EnumAddFlags(TextRefreshReason, EAvaTextRefreshReason::MaterialInstancesChange);
-	EnumAddFlags(TextRefreshReason, EAvaTextRefreshReason::MaterialParametersChange);
-}
-
-bool UAvaText3DComponent::IsMarkedForMaterialRefresh() const
-{
-	return EnumHasAnyFlags(TextRefreshReason, EAvaTextRefreshReason::MaterialInstancesChange);
-}
-
-void UAvaText3DComponent::MarkForMaterialParametersRefresh()
-{
-	EnumAddFlags(TextRefreshReason, EAvaTextRefreshReason::MaterialParametersChange);
-}
-
-bool UAvaText3DComponent::IsMarkedForMaterialParametersRefresh() const
-{
-	return EnumHasAnyFlags(TextRefreshReason, EAvaTextRefreshReason::MaterialParametersChange);
-}
-
-void UAvaText3DComponent::RemoveRefreshReason(EAvaTextRefreshReason InRefreshReasonToRemove)
-{
-	EnumRemoveFlags(TextRefreshReason, InRefreshReasonToRemove);
-}
-
-void UAvaText3DComponent::MarkForFullRefresh()
-{
-	EnumAddFlags(TextRefreshReason, EAvaTextRefreshReason::FullRefresh);
-}
-
-void UAvaText3DComponent::ClearRefreshReasons()
-{
-	EnumAddFlags(TextRefreshReason, EAvaTextRefreshReason::NoRefresh);
-}
-
-void UAvaText3DComponent::RegisterCallbacks()
-{
-	if (!OnTextGeometryGeneratedDelegate.IsValid())
-	{
-		OnTextGeometryGeneratedDelegate = OnTextGenerated().AddUObject(this, &UAvaText3DComponent::OnTextGeometryGenerated);
-	}
-
-	if (ParentRootComponent.IsValid() && !OnTransformUpdatedDelegate.IsValid())
-	{
-		OnTransformUpdatedDelegate = ParentRootComponent->TransformUpdated.AddUObject(this, &UAvaText3DComponent::OnRootComponentTransformUpdated);
-	}
-}
-
-void UAvaText3DComponent::RefreshParentRootComponent()
-{
-	if (GetOwner())
-	{
-		ParentRootComponent = GetOwner()->GetRootComponent();
-	}
-}
-
-void UAvaText3DComponent::SetupParentRootAndCallbacks()
-{
-	RefreshParentRootComponent();
-	RegisterCallbacks();
-}
-
-void UAvaText3DComponent::SetupTextMaterials(const EAvaTextColoringStyle& InColoringStyle)
-{
-	if (IsMarkedForMaterialRefresh())
-	{
-		switch (InColoringStyle)
-		{
-		case EAvaTextColoringStyle::Invalid:
-			break;
-
-		case EAvaTextColoringStyle::Solid:
-			SetupSolidMaterial();
-			break;
-
-		case EAvaTextColoringStyle::Gradient:
-			SetupGradientMaterial();
-			break;
-
-		case EAvaTextColoringStyle::FromTexture:
-			SetupTexturedMaterial();
-			break;
+	const int32 Version = InArchive.CustomVer(FText3DComponentVersion::GUID);
 
-		case EAvaTextColoringStyle::CustomMaterial:
-			SetupCustomMaterial();
-			break;
-
-		default:
-			break;
-		}
-	}
-}
-
-void UAvaText3DComponent::SetupTextGeometry()
-{
-	if (IsMarkedForGeometryRefresh())
+	if (Version < FText3DComponentVersion::LatestVersion)
 	{
-		// we are about to setup multiple properties together, let's freeze rebuild
-		SetFreeze(true);
+		UE_LOG(LogAva, Log, TEXT("AvaText3D : Migrating from %i to %i version"), Version, FText3DComponentVersion::LatestVersion)
 
-		SetText(Text);
-
-		if (UFont* UnderlyingFont = MotionDesignFont.GetFont())
+		if (Version < FText3DComponentVersion::Extensions)
 		{
-			SetFont(UnderlyingFont);
-		}
+#if WITH_EDITORONLY_DATA
+			PRAGMA_DISABLE_DEPRECATION_WARNINGS
 
-		SetExtrude(Extrude);
+#if WITH_EDITOR
+			SetFontByName(MotionDesignFont_DEPRECATED.GetFontNameAsString());
+#else
+			SetFont(MotionDesignFont.GetFont());
+#endif
 
-		SetBevel(Bevel);
-		SetBevelType(BevelType);
-		SetBevelSegments(BevelSegments);
+			if (UText3DDefaultMaterialExtension* DefaultMaterialExtension = GetCastedMaterialExtension<UText3DDefaultMaterialExtension>())
+			{
+				DefaultMaterialExtension->SetStyle(static_cast<EText3DMaterialStyle>(ColoringStyle_DEPRECATED));
+				DefaultMaterialExtension->SetFrontColor(Color_DEPRECATED);
+				DefaultMaterialExtension->SetBackColor(Color_DEPRECATED);
+				DefaultMaterialExtension->SetExtrudeColor(ExtrudeColor_DEPRECATED);
+				DefaultMaterialExtension->SetBevelColor(BevelColor_DEPRECATED);
+				DefaultMaterialExtension->SetGradientOffset(GradientSettings_DEPRECATED.Offset);
+				DefaultMaterialExtension->SetGradientSmoothness(GradientSettings_DEPRECATED.Smoothness);
+				DefaultMaterialExtension->SetGradientColorA(GradientSettings_DEPRECATED.ColorA);
+				DefaultMaterialExtension->SetGradientColorB(GradientSettings_DEPRECATED.ColorB);
 
-		//... and finally unfreeze, to trigger pending rebuild
-		SetFreeze(false);
+				if (GradientSettings_DEPRECATED.Direction == EAvaGradientDirection::Vertical)
+				{
+					DefaultMaterialExtension->SetGradientRotation(0.f);
+				}
+				else if (GradientSettings_DEPRECATED.Direction == EAvaGradientDirection::Horizontal)
+				{
+					DefaultMaterialExtension->SetGradientRotation(0.75f);
+				}
+				else
+				{
+					DefaultMaterialExtension->SetGradientRotation(GradientSettings_DEPRECATED.Rotation);
+				}
+				
+				DefaultMaterialExtension->SetTextureAsset(MainTexture_DEPRECATED);
+				DefaultMaterialExtension->SetTextureTiling(Tiling_DEPRECATED);
 
-		RemoveRefreshReason(EAvaTextRefreshReason::GeometryChange);
-	}
-}
+				if (TranslucencyStyle_DEPRECATED == EAvaTextTranslucency::None)
+				{
+					DefaultMaterialExtension->SetBlendMode(EText3DMaterialBlendMode::Opaque);
+				}
+				else
+				{
+					DefaultMaterialExtension->SetBlendMode(EText3DMaterialBlendMode::Translucent);
+				}
 
-void UAvaText3DComponent::SetupTextLayout()
-{
-	if (IsMarkedForLayoutRefresh())
-	{
-		TriggerInternalRebuild(EText3DModifyFlags::Layout);
-		RemoveRefreshReason(EAvaTextRefreshReason::LayoutChange);
-	}
-}
+				DefaultMaterialExtension->SetOpacity(Opacity_DEPRECATED);
 
-void UAvaText3DComponent::SetupSolidMaterial()
-{
-	TexturedTextMaterialInstance = nullptr;
-	GradientTextMaterialInstance = nullptr;
+				if (MaskOrientation_DEPRECATED == EAvaMaterialMaskOrientation::LeftRight)
+				{
+					DefaultMaterialExtension->SetMaskRotation(0.75f);
+				}
+				else if (MaskOrientation_DEPRECATED == EAvaMaterialMaskOrientation::RightLeft)
+				{
+					DefaultMaterialExtension->SetMaskRotation(0.25f);
+				}
+				else
+				{
+					DefaultMaterialExtension->SetMaskRotation(MaskRotation_DEPRECATED);
+				}
 
-	if (!IsValid(TextMaterialInstance) || TextMaterialInstance->GetOuter() != this || IsMarkedForMaterialRefresh())
-	{
-		const EAvaTextMaterialFeatures MaterialFeatures = GetMaterialFeaturesFromProperties();
-		const FAvaTextMaterialSettings MaterialSettings(EAvaTextColoringStyle::Solid, MaterialFeatures);
+				DefaultMaterialExtension->SetIsUnlit(bIsUnlit_DEPRECATED);
+				DefaultMaterialExtension->SetMaskSmoothness(MaskSmoothness_DEPRECATED);
 
-		// Front & Back
-		UMaterialInstanceDynamic* PreviousMaterial = nullptr;
+				if (TranslucencyStyle_DEPRECATED == EAvaTextTranslucency::GradientMask)
+				{
+					DefaultMaterialExtension->SetMaskOffset(MaskOffset_DEPRECATED);
+				}
+				else
+				{
+					DefaultMaterialExtension->SetMaskOffset(1.f);
+				}
 
-		if (UMaterialInstanceDynamic* const SolidMaterial = GetMIDWithSettings(TextMaterialInstance, MaterialSettings))
-		{
-			PreviousMaterial = TextMaterialInstance;
-			TextMaterialInstance = SolidMaterial;
-			if (PreviousMaterial)
-			{
-				TextMaterialInstance->CopyInterpParameters(PreviousMaterial);
+				// Patch materials for other systems like mask
+				DefaultMaterialExtension->PreCacheMaterials();
 			}
-		}
 
-		if (UMaterialInstanceDynamic* const SolidMaterial = GetMIDWithSettings(ExtrudeTextMaterialInstance, MaterialSettings))
-		{
-			PreviousMaterial = ExtrudeTextMaterialInstance;
-			ExtrudeTextMaterialInstance = SolidMaterial;
-			if (PreviousMaterial)
+			if (const AActor* Actor = GetOwner())
 			{
-				ExtrudeTextMaterialInstance->CopyInterpParameters(PreviousMaterial);
+				if (const UText3DCharacterTransform* CharacterTransform = Actor->FindComponentByClass<UText3DCharacterTransform>())
+				{
+					UText3DDefaultTransformExtension* TransformExtension = NewObject<UText3DDefaultTransformExtension>(this, TEXT("AvaTransformExtension"));
+
+					TransformExtension->SetLocationEnabled(CharacterTransform->GetLocationEnabled());
+					TransformExtension->SetLocationProgress(CharacterTransform->GetLocationProgress());
+					TransformExtension->SetLocationOrder(CharacterTransform->GetLocationOrder());
+					TransformExtension->SetLocationRange(CharacterTransform->GetLocationRange());
+					TransformExtension->SetLocationDistance(CharacterTransform->GetLocationDistance());
+					TransformExtension->SetLocationDistance(CharacterTransform->GetLocationDistance());
+
+					TransformExtension->SetRotationEnabled(CharacterTransform->GetRotationEnabled());
+					TransformExtension->SetRotationProgress(CharacterTransform->GetRotationProgress());
+					TransformExtension->SetRotationOrder(CharacterTransform->GetRotationOrder());
+					TransformExtension->SetRotationRange(CharacterTransform->GetRotationRange());
+					TransformExtension->SetRotationBegin(CharacterTransform->GetRotationBegin());
+					TransformExtension->SetRotationEnd(CharacterTransform->GetRotationEnd());
+
+					TransformExtension->SetScaleEnabled(CharacterTransform->GetScaleEnabled());
+					TransformExtension->SetScaleProgress(CharacterTransform->GetScaleProgress());
+					TransformExtension->SetScaleOrder(CharacterTransform->GetScaleOrder());
+					TransformExtension->SetScaleRange(CharacterTransform->GetScaleRange());
+					TransformExtension->SetScaleBegin(CharacterTransform->GetScaleBegin());
+					TransformExtension->SetScaleEnd(CharacterTransform->GetScaleEnd());
+
+					TransformExtensions.Emplace(TransformExtension);
+				}
 			}
-		}
 
-		if (UMaterialInstanceDynamic* const SolidMaterial = GetMIDWithSettings(BevelTextMaterialInstance, MaterialSettings))
-		{
-			PreviousMaterial = BevelTextMaterialInstance;
-			BevelTextMaterialInstance = SolidMaterial;
-			if (PreviousMaterial)
-			{
-				BevelTextMaterialInstance->CopyInterpParameters(PreviousMaterial);
-			}
-		}
-
-		RemoveRefreshReason(EAvaTextRefreshReason::MaterialInstancesChange);
-	}
-	{
-		UE::AvaText::Private::PatchMaterial(
-			this
-			, &UText3DComponent::GetFrontMaterial
-			, &UText3DComponent::SetFrontMaterial
-			, TextMaterialInstance);
-
-		UE::AvaText::Private::PatchMaterial(
-			this
-			, &UText3DComponent::GetExtrudeMaterial
-			, &UText3DComponent::SetExtrudeMaterial
-			, ExtrudeTextMaterialInstance);
-
-		UE::AvaText::Private::PatchMaterial(
-			this
-			, &UText3DComponent::GetBevelMaterial
-			, &UText3DComponent::SetBevelMaterial
-			, BevelTextMaterialInstance);
-
-		UE::AvaText::Private::PatchMaterial(
-			this
-			, &UText3DComponent::GetBackMaterial
-			, &UText3DComponent::SetBackMaterial
-			, TextMaterialInstance);
-	}
-
-	if (IsMarkedForMaterialParametersRefresh())
-	{
-		RefreshColor();
-		RefreshExtrudeColor();
-		RefreshBevelColor();
-
-		RefreshTranslucencyValues();
-	}
-}
-
-void UAvaText3DComponent::SetupGradientMaterial()
-{
-	TexturedTextMaterialInstance = nullptr;
-
-	if (!IsValid(GradientTextMaterialInstance) || GradientTextMaterialInstance->GetOuter() != this || IsMarkedForMaterialRefresh())
-	{
-		const EAvaTextMaterialFeatures MaterialFeatures = GetMaterialFeaturesFromProperties();
-		const FAvaTextMaterialSettings MaterialSettings(EAvaTextColoringStyle::Gradient, MaterialFeatures);
-
-		if (UMaterialInstanceDynamic* const GradientMaterial = GetMIDWithSettings(GradientTextMaterialInstance, MaterialSettings))
-		{
-			UMaterialInstanceDynamic* PreviousMaterial = GradientTextMaterialInstance;
-			GradientTextMaterialInstance = GradientMaterial;
-			if (PreviousMaterial)
-			{
-				GradientTextMaterialInstance->CopyInterpParameters(PreviousMaterial);
-			}
-
-			RemoveRefreshReason(EAvaTextRefreshReason::MaterialInstancesChange);
-		}
-	}
-
-	if (!GradientTextMaterialInstance)
-	{
-		return;
-	}
-
-	UE::AvaText::Private::PatchMaterial(
-		this
-		, &UText3DComponent::GetFrontMaterial
-		, &UText3DComponent::SetFrontMaterial
-		, GradientTextMaterialInstance);
-
-	UE::AvaText::Private::PatchMaterial(
-		this
-		, &UText3DComponent::GetExtrudeMaterial
-		, &UText3DComponent::SetExtrudeMaterial
-		, GradientTextMaterialInstance);
-
-	UE::AvaText::Private::PatchMaterial(
-		this
-		, &UText3DComponent::GetBevelMaterial
-		, &UText3DComponent::SetBevelMaterial
-		, GradientTextMaterialInstance);
-
-	UE::AvaText::Private::PatchMaterial(
-		this
-		, &UText3DComponent::GetBackMaterial
-		, &UText3DComponent::SetBackMaterial
-		, GradientTextMaterialInstance);
-
-	if (IsMarkedForMaterialParametersRefresh())
-	{
-		RefreshGradientValues();
-		RefreshTranslucencyValues();
-	}
-}
-
-void UAvaText3DComponent::SetupCustomMaterial()
-{
-	TexturedTextMaterialInstance = nullptr;
-	GradientTextMaterialInstance = nullptr;
-
-	if (IsMarkedForMaterialRefresh())
-	{
-		RemoveRefreshReason(EAvaTextRefreshReason::MaterialInstancesChange);
-
-		SetFrontMaterial(FrontMaterial);
-		SetExtrudeMaterial(ExtrudeMaterial);
-		SetBevelMaterial(BevelMaterial);
-		SetBackMaterial(BackMaterial);
-	}
-}
-
-void UAvaText3DComponent::RefreshTranslucencyValues()
-{
-	switch (TranslucencyStyle)
-	{
-		case EAvaTextTranslucency::Translucent:
-			RefreshOpacity();
-			break;
-
-		case EAvaTextTranslucency::GradientMask:
-			RefreshMaskValues();
-			break;
-
-		default: ;
-	}
-}
-
-void UAvaText3DComponent::SetupTexturedMaterial()
-{
-	GradientTextMaterialInstance = nullptr;
-
-	if (!IsValid(TexturedTextMaterialInstance) || TexturedTextMaterialInstance->GetOuter() != this || IsMarkedForMaterialRefresh())
-	{
-		const EAvaTextMaterialFeatures MaterialFeatures = GetMaterialFeaturesFromProperties();
-		const FAvaTextMaterialSettings MaterialSettings(EAvaTextColoringStyle::FromTexture, MaterialFeatures);
-
-		if (UMaterialInstanceDynamic* const TexturedMaterial = GetMIDWithSettings(TexturedTextMaterialInstance, MaterialSettings))
-		{
-			UMaterialInstanceDynamic* PreviousMaterial = TexturedTextMaterialInstance;
-			TexturedTextMaterialInstance = TexturedMaterial;
-			if (PreviousMaterial)
-			{
-				TexturedTextMaterialInstance->CopyInterpParameters(PreviousMaterial);
-			}
-		}
-
-		RemoveRefreshReason(EAvaTextRefreshReason::MaterialInstancesChange);
-	}
-
-	if (!TexturedTextMaterialInstance)
-	{
-		return;
-	}
-
-	{
-		UE::AvaText::Private::PatchMaterial(
-			this
-			, &UText3DComponent::GetFrontMaterial
-			, &UText3DComponent::SetFrontMaterial
-			, TexturedTextMaterialInstance);
-
-		UE::AvaText::Private::PatchMaterial(
-			this
-			, &UText3DComponent::GetExtrudeMaterial
-			, &UText3DComponent::SetExtrudeMaterial
-			, TexturedTextMaterialInstance);
-
-		UE::AvaText::Private::PatchMaterial(
-			this
-			, &UText3DComponent::GetBevelMaterial
-			, &UText3DComponent::SetBevelMaterial
-			, TexturedTextMaterialInstance);
-
-		UE::AvaText::Private::PatchMaterial(
-			this
-			, &UText3DComponent::GetBackMaterial
-			, &UText3DComponent::SetBackMaterial
-			, TexturedTextMaterialInstance);
-	}
-
-	UTexture* CurrentTexture;
-	const FMemoryImageMaterialParameterInfo MaterialInfo(UAvaTextMaterialHub::FMaterialStatics::MainTexture_MatParam);
-	TexturedTextMaterialInstance->GetTextureParameterValue(MaterialInfo, CurrentTexture);
-
-	if (MainTexture != CurrentTexture)
-	{
-		TexturedTextMaterialInstance->SetTextureParameterValue(UAvaTextMaterialHub::FMaterialStatics::MainTexture_MatParam, MainTexture);
-	}
-
-	TexturedTextMaterialInstance->SetScalarParameterValue(UAvaTextMaterialHub::FMaterialStatics::TexturedUTiling_MatParam, Tiling.X);
-	TexturedTextMaterialInstance->SetScalarParameterValue(UAvaTextMaterialHub::FMaterialStatics::TexturedVTiling_MatParam, Tiling.Y);
-
-	RefreshTranslucencyValues();
-}
-
-void UAvaText3DComponent::RefreshStoredBoundsValues()
-{
-	if (const AActor* ParentActor = GetOwner())
-	{
-		// local bounds (used for gradient or masked materials)
-		const FBox LocalBounds = FAvaActorUtils::GetActorLocalBoundingBox(ParentActor);
-		LocalBoundsLowerLeftCorner = LocalBounds.Min;
-		LocalBoundsExtent = LocalBounds.GetExtent();
-	}
-}
-
-void UAvaText3DComponent::PostLoad()
-{
-	Super::PostLoad();
-
-	SetupParentRootAndCallbacks();
-	RetrieveDataFromBaseComponent();
-}
-
-void UAvaText3DComponent::PostDuplicate(EDuplicateMode::Type DuplicateMode)
-{
-	Super::PostDuplicate(DuplicateMode);
-
-	PostDuplicateActions();
-}
-
-void UAvaText3DComponent::PostInitProperties()
-{
-	Super::PostInitProperties();
-
-	SetupParentRootAndCallbacks();
-	MarkForFullRefresh();
-}
-
-// overriding PostEditImport to handle Owner duplication
-void UAvaText3DComponent::PostEditImport()
-{
-	Super::PostEditImport();
-
-	PostDuplicateActions();
-}
-
-void UAvaText3DComponent::PostDuplicateActions()
-{
-	SetupParentRootAndCallbacks();
-
-	RetrieveCurrentTextAndFont();
-
-	MarkForFullRefresh();
-	RefreshText3D();
-}
-
-void UAvaText3DComponent::ScheduleTextRefreshOnTicker()
-{
-	if (!TickerHandle.IsValid())
-	{
-		TWeakObjectPtr<UAvaText3DComponent> WeakThis(this);
-		FTSTicker& Ticker = FTSTicker::GetCoreTicker();
-
-		TickerHandle = Ticker.AddTicker(FTickerDelegate::CreateLambda(
-			[WeakThis](float InDeltaTime)->bool
-			{
-				if (UAvaText3DComponent* This = WeakThis.Get())
-				{
-					if (!UE::IsSavingPackage(This))
-					{
-						This->RefreshStoredBoundsValues();
-						This->MarkForMaterialParametersRefresh();
-						This->RefreshText3D();
-					}
-
-					This->TickerHandle.Reset();
-				}
-
-				// Executes only once
-				return false;
-			}));
-	}
-}
-
-void UAvaText3DComponent::OnTextGeneratedActions()
-{
-	RefreshStoredBoundsValues();
-
-	// bounding box has been updated, so we need to also update gradient and mask values
-	RefreshGradient();
-	RefreshMask();
-}
-
-void UAvaText3DComponent::OnTextGeometryGenerated()
-{
-	OnTextGeneratedActions();
-
-#if WITH_EDITOR
-	if (GEditor)
-	{
-		GEditor->BroadcastOnComponentTransformChanged(this, ETeleportType::None);
-	}
+			PRAGMA_ENABLE_DEPRECATION_WARNINGS
 #endif
-
-	TransformUpdated.Broadcast(this, EUpdateTransformFlags::None, ETeleportType::None);
-}
-
-void UAvaText3DComponent::OnRootComponentTransformUpdated(USceneComponent* InSceneComponent, EUpdateTransformFlags InUpdateTransformFlags, ETeleportType InTeleportType)
-{
-	ScheduleTextRefreshOnTicker();
-}
-
-void UAvaText3DComponent::RetrieveDataFromBaseComponent()
-{
-	RetrieveCurrentColoringStyle();
-	RetrieveCurrentMaterials();
-	RetrieveCurrentTextAndFont();
-	RetrieveCurrentAlignment();
-}
-
-void UAvaText3DComponent::RetrieveCurrentTextAndFont()
-{
-	// if default had to be loaded (e.g. font asset missing), then we regenerate geometry
-	if (MotionDesignFont.IsDefaultFont())
-	{
-		if (IsValid(Font) && Font != MotionDesignFont.GetFont())
-		{
-			MotionDesignFont.InitFromFont(Font);
-		}
-
-		MarkForGeometryRefresh();
-	}
-
-	if (!IsValid(Font))
-	{
-		MarkForGeometryRefresh();
-	}
-}
-
-void UAvaText3DComponent::RetrieveCurrentAlignment()
-{
-	Alignment.HorizontalAlignment = HorizontalAlignment;
-	Alignment.VerticalAlignment = VerticalAlignment;
-}
-
-void UAvaText3DComponent::RetrieveCurrentMaterials()
-{
-	switch (ColoringStyle)
-	{
-		case EAvaTextColoringStyle::Invalid:
-			break;
-
-		case EAvaTextColoringStyle::Solid:
-			RetrieveSolidMaterialInstances();
-			break;
-
-		case EAvaTextColoringStyle::Gradient:
-			RetrieveGradientMaterialInstance();
-			break;
-
-		case EAvaTextColoringStyle::FromTexture: 
-			RetrieveTexturedMaterialInstance();
-			break;
-
-		case EAvaTextColoringStyle::CustomMaterial: 
-			break;
-	}
-}
-
-void UAvaText3DComponent::RetrieveCurrentColoringStyle()
-{
-	if (!FrontMaterial)
-	{
-		return;
-	}
-
-	if (UMaterialInterface* CurrentFrontMaterial = FrontMaterial->GetMaterial())
-	{
-		if (const FAvaTextMaterialSettings* MaterialSettings = UAvaTextMaterialHub::GetSettingsFromMaterial(CurrentFrontMaterial))
-		{
-			ColoringStyle = MaterialSettings->ColoringStyle;
-
-			if (EnumHasAnyFlags(MaterialSettings->MaterialFeatures, EAvaTextMaterialFeatures::GradientMask))
-			{
-				TranslucencyStyle = EAvaTextTranslucency::GradientMask;
-			}
-
-			if (EnumHasAnyFlags(MaterialSettings->MaterialFeatures, EAvaTextMaterialFeatures::Translucent))
-			{
-				TranslucencyStyle = EAvaTextTranslucency::Translucent;
-			}
-
-			bIsUnlit = EnumHasAnyFlags(MaterialSettings->MaterialFeatures, EAvaTextMaterialFeatures::Unlit);
-		}
-		else // if there is no match, best guess is custom material
-		{
-			ColoringStyle = EAvaTextColoringStyle::CustomMaterial;
-			TranslucencyStyle = EAvaTextTranslucency::None;
-		}
-	}
-}
-
-void UAvaText3DComponent::RetrieveSolidMaterialInstances()
-{
-	if (FrontMaterial)
-	{
-		TextMaterialInstance = Cast<UMaterialInstanceDynamic>(FrontMaterial);
-		if (TextMaterialInstance)
-		{
-			TextMaterialInstance->GetVectorParameterValue(UAvaTextMaterialHub::FMaterialStatics::SolidColor_MatParam, Color);
-
-			RetrieveTranslucentMaterialParameters();
-		}
-	}
-
-	if (ExtrudeMaterial)
-	{
-		ExtrudeTextMaterialInstance = Cast<UMaterialInstanceDynamic>(ExtrudeMaterial);
-		if (ExtrudeTextMaterialInstance)
-		{
-			ExtrudeTextMaterialInstance->GetVectorParameterValue(UAvaTextMaterialHub::FMaterialStatics::SolidColor_MatParam, ExtrudeColor);
 		}
 	}
-
-	if (BevelMaterial)
-	{
-		BevelTextMaterialInstance = Cast<UMaterialInstanceDynamic>(BevelMaterial);
-		if (BevelTextMaterialInstance)
-		{
-			BevelTextMaterialInstance->GetVectorParameterValue(UAvaTextMaterialHub::FMaterialStatics::SolidColor_MatParam, BevelColor);
-		}
-	}
-}
-
-void UAvaText3DComponent::RetrieveGradientMaterialInstance()
-{
-	if (!FrontMaterial)
-	{
-		return;
-	}
-
-	GradientTextMaterialInstance = Cast<UMaterialInstanceDynamic>(FrontMaterial);
-
-	if (GradientTextMaterialInstance)
-	{
-		GradientTextMaterialInstance->GetVectorParameterValue(UAvaTextMaterialHub::FMaterialStatics::GradientColorA_MatParam, GradientSettings.ColorA);
-		GradientTextMaterialInstance->GetVectorParameterValue(UAvaTextMaterialHub::FMaterialStatics::GradientColorB_MatParam, GradientSettings.ColorB);
-
-		GradientTextMaterialInstance->GetScalarParameterValue(UAvaTextMaterialHub::FMaterialStatics::GradientOffset_MatParam, GradientSettings.Offset);
-		GradientTextMaterialInstance->GetScalarParameterValue(UAvaTextMaterialHub::FMaterialStatics::GradientSmoothness_MatParam, GradientSettings.Smoothness);
-
-		const bool bRotationRetrieved = GradientTextMaterialInstance->GetScalarParameterValue(UAvaTextMaterialHub::FMaterialStatics::GradientRotation_MatParam, GradientSettings.Rotation);
-
-		if (bRotationRetrieved)
-		{
-			if (GradientSettings.Rotation == UAvaTextMaterialHub::FMaterialStatics::MaterialMaskTopBottomRotationValue)
-			{
-				GradientSettings.Direction = EAvaGradientDirection::Vertical;
-			}
-			else if (GradientSettings.Rotation == UAvaTextMaterialHub::FMaterialStatics::MaterialMaskLeftToRightRotationValue)
-			{
-				GradientSettings.Direction = EAvaGradientDirection::Horizontal;
-			}
-			else
-			{
-				GradientSettings.Direction = EAvaGradientDirection::Custom;
-			}
-		}
-
-		RetrieveTranslucentMaterialParameters();
-	}
-}
-
-void UAvaText3DComponent::RetrieveTexturedMaterialInstance()
-{
-	if (!FrontMaterial)
-	{
-		return;
-	}
-
-	TexturedTextMaterialInstance = Cast<UMaterialInstanceDynamic>(FrontMaterial);
-
-	if (TexturedTextMaterialInstance)
-	{
-		float Tiling_X;
-		float Tiling_Y;
-		TexturedTextMaterialInstance->GetScalarParameterValue(UAvaTextMaterialHub::FMaterialStatics::TexturedUTiling_MatParam, Tiling_X);
-		TexturedTextMaterialInstance->GetScalarParameterValue(UAvaTextMaterialHub::FMaterialStatics::TexturedVTiling_MatParam, Tiling_Y);
-
-		Tiling = FVector2D(Tiling_X, Tiling_Y);
-
-		UTexture* CurrentTexture;
-		TexturedTextMaterialInstance->GetTextureParameterValue(UAvaTextMaterialHub::FMaterialStatics::MainTexture_MatParam, CurrentTexture);
-
-		if (CurrentTexture)
-		{
-			MainTexture = Cast<UTexture2D>(CurrentTexture);
-		}
-
-		RetrieveTranslucentMaterialParameters();
-	}
-}
-
-void UAvaText3DComponent::RetrieveTranslucentMaterialParameters()
-{
-	switch (TranslucencyStyle)
-	{
-	case EAvaTextTranslucency::Translucent:
-		RetrieveSimpleTranslucentMaterialParameters(TextMaterialInstance);
-		break;
-
-	case EAvaTextTranslucency::GradientMask:
-		RetrieveMaskedMaterialParameters(TextMaterialInstance);
-		break;
-
-	default: ;
-	}
-}
-
-void UAvaText3DComponent::RetrieveMaskedMaterialParameters(const UMaterialInstanceDynamic* InMaterialInstance)
-{
-	if (!InMaterialInstance)
-	{
-		return;
-	}
-
-	InMaterialInstance->GetScalarParameterValue(UAvaTextMaterialHub::FMaterialStatics::MaskSmoothness_MatParam, MaskSmoothness);
-	InMaterialInstance->GetScalarParameterValue(UAvaTextMaterialHub::FMaterialStatics::MaskOffset_MatParam, MaskOffset);
-	InMaterialInstance->GetScalarParameterValue(UAvaTextMaterialHub::FMaterialStatics::MaskRotation_MatParam, MaskRotation);
-
-	// in case the material was saved with the deprecated value, let's fix the rotation value
-	if (MaskRotation == UAvaTextMaterialHub::FMaterialStatics::MaterialMaskRightToLeftRotationValue_DEPRECATED)
-	{
-		MaskRotation = UAvaTextMaterialHub::FMaterialStatics::MaterialMaskRightToLeftRotationValue;
-	}
-
-	if (MaskRotation == UAvaTextMaterialHub::FMaterialStatics::MaterialMaskLeftToRightRotationValue)
-	{
-		MaskOrientation = EAvaMaterialMaskOrientation::LeftRight;
-	}
-	else if (MaskRotation == UAvaTextMaterialHub::FMaterialStatics::MaterialMaskRightToLeftRotationValue)
-	{
-		MaskOrientation = EAvaMaterialMaskOrientation::RightLeft;
-	}
-	else
-	{
-		MaskOrientation = EAvaMaterialMaskOrientation::Custom;
-	}
-}
-
-void UAvaText3DComponent::RetrieveSimpleTranslucentMaterialParameters(const UMaterialInstanceDynamic* InMaterialInstance)
-{
-	if (!InMaterialInstance)
-	{
-		return;
-	}
-
-	InMaterialInstance->GetScalarParameterValue(UAvaTextMaterialHub::FMaterialStatics::Opacity_MatParam, Opacity);
-}
-
-UMaterialInstanceDynamic* UAvaText3DComponent::GetMIDWithSettings(
-	const UMaterialInterface* InPreviousMaterial
-	, const FAvaTextMaterialSettings& InSettings)
-{
-	UMaterialInterface* Result = nullptr;
-	if (OnGetMaterialWithSettingsDelegate.IsBound())
-	{
-		Result = OnGetMaterialWithSettingsDelegate.Execute(InPreviousMaterial, InSettings);
-	}
-
-	if (!Result)
-	{
-		Result = UAvaTextMaterialHub::GetMaterial(InSettings);
-	}
-
-	if (!::IsValid(Result))
-	{
-		UE_LOG(LogAva, Warning, TEXT("Material could not be found or was invalid for the provided settings."));
-		return nullptr;
-	}
-
-	if (UMaterialInstanceDynamic* ResultAsMID = Cast<UMaterialInstanceDynamic>(Result))
-	{
-		return ResultAsMID;
-	}
-
-	return UMaterialInstanceDynamic::Create(Result, this);
-}
-
-EAvaTextMaterialFeatures UAvaText3DComponent::GetMaterialFeaturesFromProperties() const
-{
-	EAvaTextMaterialFeatures MaterialFeatures = EAvaTextMaterialFeatures::None;
-
-	if (TranslucencyStyle == EAvaTextTranslucency::GradientMask)
-	{
-		MaterialFeatures |= EAvaTextMaterialFeatures::GradientMask;
-	}
-	else if (TranslucencyStyle == EAvaTextTranslucency::Translucent)
-	{
-		MaterialFeatures |= EAvaTextMaterialFeatures::Translucent;
-	}
-
-	MaterialFeatures |= bIsUnlit ? EAvaTextMaterialFeatures::Unlit : EAvaTextMaterialFeatures::None;
-
-	return MaterialFeatures;
-}
-
-#if WITH_EDITOR
-void UAvaText3DComponent::RegisterOnPropertyChangeFunctions()
-{
-	if (PropertyChangeFunctions.IsEmpty())
-	{
-		PropertyChangeFunctions.Add(GET_MEMBER_NAME_CHECKED(UAvaText3DComponent, Color), &UAvaText3DComponent::RefreshColor);
-		PropertyChangeFunctions.Add(GET_MEMBER_NAME_CHECKED(UAvaText3DComponent, ExtrudeColor), &UAvaText3DComponent::RefreshExtrudeColor);
-		PropertyChangeFunctions.Add(GET_MEMBER_NAME_CHECKED(UAvaText3DComponent, BevelColor), &UAvaText3DComponent::RefreshBevelColor);
-		PropertyChangeFunctions.Add(GET_MEMBER_NAME_CHECKED(UAvaText3DComponent, bIsUnlit), &UAvaText3DComponent::RefreshUnlit);
-		PropertyChangeFunctions.Add(GET_MEMBER_NAME_CHECKED(UAvaText3DComponent, Opacity), &UAvaText3DComponent::RefreshOpacity);
-		PropertyChangeFunctions.Add(GET_MEMBER_NAME_CHECKED(UAvaText3DComponent, GradientSettings), &UAvaText3DComponent::RefreshGradient);
-
-		PropertyChangeFunctions.Add(GET_MEMBER_NAME_CHECKED(UAvaText3DComponent, Extrude), &UAvaText3DComponent::RefreshGeometry);
-		PropertyChangeFunctions.Add(GET_MEMBER_NAME_CHECKED(UAvaText3DComponent, Bevel), &UAvaText3DComponent::RefreshGeometry);
-		PropertyChangeFunctions.Add(GET_MEMBER_NAME_CHECKED(UAvaText3DComponent, MotionDesignFont), &UAvaText3DComponent::RefreshGeometry);
-		PropertyChangeFunctions.Add(GET_MEMBER_NAME_CHECKED(UAvaText3DComponent, Font), &UAvaText3DComponent::RefreshMotionDesignFont);
-		PropertyChangeFunctions.Add(GET_MEMBER_NAME_CHECKED(UAvaText3DComponent, Text), &UAvaText3DComponent::RefreshGeometry);
-		PropertyChangeFunctions.Add(GET_MEMBER_NAME_CHECKED(UAvaText3DComponent, bEnforceUpperCase), &UAvaText3DComponent::OnTextChanged);
-
-		PropertyChangeFunctions.Add(GET_MEMBER_NAME_CHECKED(UAvaText3DComponent, ColoringStyle), &UAvaText3DComponent::RefreshMaterialInstances);
-		PropertyChangeFunctions.Add(GET_MEMBER_NAME_CHECKED(UAvaText3DComponent, MainTexture), &UAvaText3DComponent::RefreshMaterialInstances);
-		PropertyChangeFunctions.Add(GET_MEMBER_NAME_CHECKED(UAvaText3DComponent, Tiling), &UAvaText3DComponent::RefreshMaterialInstances);
-		PropertyChangeFunctions.Add(GET_MEMBER_NAME_CHECKED(UAvaText3DComponent, TranslucencyStyle), &UAvaText3DComponent::RefreshMaterialInstances);
-
-		PropertyChangeFunctions.Add(GET_MEMBER_NAME_CHECKED(UAvaText3DComponent, MaskOrientation), &UAvaText3DComponent::RefreshMask);
-		PropertyChangeFunctions.Add(GET_MEMBER_NAME_CHECKED(UAvaText3DComponent, MaskSmoothness), &UAvaText3DComponent::RefreshMask);
-		PropertyChangeFunctions.Add(GET_MEMBER_NAME_CHECKED(UAvaText3DComponent, MaskOffset), &UAvaText3DComponent::RefreshMask);
-		PropertyChangeFunctions.Add(GET_MEMBER_NAME_CHECKED(UAvaText3DComponent, MaskRotation), &UAvaText3DComponent::RefreshMask);
-
-		PropertyChangeFunctions.Add(GET_MEMBER_NAME_CHECKED(UAvaText3DComponent, Alignment), &UAvaText3DComponent::RefreshAlignment);
-		PropertyChangeFunctions.Add(GET_MEMBER_NAME_CHECKED(UAvaText3DComponent, HorizontalAlignment), &UAvaText3DComponent::RefreshHorizontalAlignment);
-		PropertyChangeFunctions.Add(GET_MEMBER_NAME_CHECKED(UAvaText3DComponent, VerticalAlignment), &UAvaText3DComponent::RefreshVerticalAlignment);
-		PropertyChangeFunctions.Add(GET_MEMBER_NAME_CHECKED(UAvaText3DComponent, bHasMaxWidth), &UAvaText3DComponent::RefreshLayout);
-		PropertyChangeFunctions.Add(GET_MEMBER_NAME_CHECKED(UAvaText3DComponent, MaxWidth), &UAvaText3DComponent::RefreshLayout);
-		PropertyChangeFunctions.Add(GET_MEMBER_NAME_CHECKED(UAvaText3DComponent, bHasMaxHeight), &UAvaText3DComponent::RefreshLayout);
-		PropertyChangeFunctions.Add(GET_MEMBER_NAME_CHECKED(UAvaText3DComponent, MaxHeight), &UAvaText3DComponent::RefreshLayout);
-		PropertyChangeFunctions.Add(GET_MEMBER_NAME_CHECKED(UAvaText3DComponent, bScaleProportionally), &UAvaText3DComponent::RefreshLayout);
-		PropertyChangeFunctions.Add(GET_MEMBER_NAME_CHECKED(UAvaText3DComponent, Kerning), &UAvaText3DComponent::RefreshLayout);
-		PropertyChangeFunctions.Add(GET_MEMBER_NAME_CHECKED(UAvaText3DComponent, WordSpacing), &UAvaText3DComponent::RefreshLayout);
-		PropertyChangeFunctions.Add(GET_MEMBER_NAME_CHECKED(UAvaText3DComponent, LineSpacing), &UAvaText3DComponent::RefreshLayout);
-	}
-}
-#endif
-
-void UAvaText3DComponent::RefreshOpacity() const
-{
-	if (TranslucencyStyle == EAvaTextTranslucency::Translucent)
-	{
-		// Set for any/all materials assigned
-		ForEachMID([this](UMaterialInstanceDynamic* InMID)
-		{
-			InMID->SetScalarParameterValue(UAvaTextMaterialHub::FMaterialStatics::Opacity_MatParam, Opacity);
-		});
-
-		// ... and set for the default materials
-		switch (ColoringStyle)
-		{
-		case EAvaTextColoringStyle::Solid:
-			if (TextMaterialInstance)
-			{
-				TextMaterialInstance->SetScalarParameterValue(UAvaTextMaterialHub::FMaterialStatics::Opacity_MatParam, Opacity);
-			}
-
-			if (ExtrudeTextMaterialInstance)
-			{
-				ExtrudeTextMaterialInstance->SetScalarParameterValue(UAvaTextMaterialHub::FMaterialStatics::Opacity_MatParam, Opacity);
-			}
-
-			if (BevelTextMaterialInstance)
-			{
-				BevelTextMaterialInstance->SetScalarParameterValue(UAvaTextMaterialHub::FMaterialStatics::Opacity_MatParam, Opacity);
-			}
-			break;
-
-		case EAvaTextColoringStyle::Gradient:
-			if (GradientTextMaterialInstance)
-			{
-				GradientTextMaterialInstance->SetScalarParameterValue(UAvaTextMaterialHub::FMaterialStatics::Opacity_MatParam, Opacity);
-			}
-			break;
-
-		case EAvaTextColoringStyle::FromTexture:
-			if (TexturedTextMaterialInstance)
-			{
-				TexturedTextMaterialInstance->SetScalarParameterValue(UAvaTextMaterialHub::FMaterialStatics::Opacity_MatParam, Opacity);
-			}
-			break;
-
-		case EAvaTextColoringStyle::CustomMaterial:
-			break;
-
-		default: ;
-		}
-	}
-}
-
-void UAvaText3DComponent::RefreshColor() const
-{
-	if (UMaterialInstanceDynamic* MID = Cast<UMaterialInstanceDynamic>(GetFrontMaterial()))
-	{
-		MID->SetVectorParameterValue(UAvaTextMaterialHub::FMaterialStatics::SolidColor_MatParam, Color);
-	}
-
-	if (TextMaterialInstance)
-	{
-		TextMaterialInstance->SetVectorParameterValue(UAvaTextMaterialHub::FMaterialStatics::SolidColor_MatParam, Color);
-	}
-}
-
-void UAvaText3DComponent::RefreshExtrudeColor() const
-{
-	if (UMaterialInstanceDynamic* MID = Cast<UMaterialInstanceDynamic>(GetExtrudeMaterial()))
-	{
-		MID->SetVectorParameterValue(UAvaTextMaterialHub::FMaterialStatics::SolidColor_MatParam, ExtrudeColor);
-	}
-
-	if (ExtrudeTextMaterialInstance)
-	{
-		ExtrudeTextMaterialInstance->SetVectorParameterValue(UAvaTextMaterialHub::FMaterialStatics::SolidColor_MatParam, ExtrudeColor);
-	}
-}
-
-void UAvaText3DComponent::RefreshBevelColor() const
-{
-	if (UMaterialInstanceDynamic* MID = Cast<UMaterialInstanceDynamic>(GetBevelMaterial()))
-	{
-		MID->SetVectorParameterValue(UAvaTextMaterialHub::FMaterialStatics::SolidColor_MatParam, BevelColor);
-	}
-
-	if (BevelTextMaterialInstance)
-	{
-		BevelTextMaterialInstance->SetVectorParameterValue(UAvaTextMaterialHub::FMaterialStatics::SolidColor_MatParam, BevelColor);
-	}
-}
-
-void UAvaText3DComponent::RefreshAlignment()
-{
-	HorizontalAlignment = Alignment.HorizontalAlignment;
-	VerticalAlignment = Alignment.VerticalAlignment;
-
-	ScheduleTextRefreshOnTicker();
-	RefreshLayout();
-}
-
-void UAvaText3DComponent::RefreshMotionDesignFont()
-{
-	if (IsValid(Font) && Font != MotionDesignFont.GetFont())
-	{
-		// todo: this leads to proper update of the MotionDesignFont, and of the text component geometry
-		// nonetheless, custom FAvaFont widget is not updating properly yet
-		MotionDesignFont.InitFromFont(Font);
-	}
-}
-
-void UAvaText3DComponent::RefreshHorizontalAlignment()
-{
-	Alignment.HorizontalAlignment = HorizontalAlignment;
-	RefreshLayout();
-}
-
-void UAvaText3DComponent::RefreshVerticalAlignment()
-{
-	Alignment.VerticalAlignment = VerticalAlignment;
-	RefreshLayout();
-}
-
-void UAvaText3DComponent::RefreshUnlit()
-{
-	MarkForMaterialsFullRefresh();
-	RefreshText3D();
-}
-
-void UAvaText3DComponent::RefreshMasked()
-{
-	MarkForMaterialsFullRefresh();
-	RefreshText3D();
-}
-
-void UAvaText3DComponent::RefreshGeometry()
-{
-	MarkForGeometryRefresh();
-	RefreshText3D();
-}
-
-void UAvaText3DComponent::RefreshLayout()
-{
-	MarkForLayoutRefresh();
-	RefreshText3D();
-}
-
-void UAvaText3DComponent::RefreshGradientValues()
-{
-	if (IsMarkedForMaterialRefresh() || IsMarkedForMaterialParametersRefresh())
-	{
-		if (IsValid(GradientTextMaterialInstance))
-		{
-			GradientTextMaterialInstance->SetScalarParameterValue(UAvaTextMaterialHub::FMaterialStatics::GradientOffset_MatParam, GradientSettings.Offset);
-			GradientTextMaterialInstance->SetScalarParameterValue(UAvaTextMaterialHub::FMaterialStatics::GradientSmoothness_MatParam, GradientSettings.Smoothness);
-			GradientTextMaterialInstance->SetVectorParameterValue(UAvaTextMaterialHub::FMaterialStatics::GradientColorA_MatParam, GradientSettings.ColorA);
-			GradientTextMaterialInstance->SetVectorParameterValue(UAvaTextMaterialHub::FMaterialStatics::GradientColorB_MatParam, GradientSettings.ColorB);
-
-			float GradientRotation = UAvaTextMaterialHub::FMaterialStatics::MaterialMaskTopBottomRotationValue;
-
-			switch (GradientSettings.Direction)
-			{
-			case EAvaGradientDirection::None:
-				break;
-
-			case EAvaGradientDirection::Vertical:
-					GradientRotation = UAvaTextMaterialHub::FMaterialStatics::MaterialMaskTopBottomRotationValue;
-					break;
-
-			case EAvaGradientDirection::Horizontal:
-					GradientRotation = UAvaTextMaterialHub::FMaterialStatics::MaterialMaskLeftToRightRotationValue;
-					break;
-
-			case EAvaGradientDirection::Custom:
-					GradientRotation = GradientSettings.Rotation;
-					break;
-			}
-
-			GradientTextMaterialInstance->SetScalarParameterValue(UAvaTextMaterialHub::FMaterialStatics::GradientRotation_MatParam, GradientRotation);
-		}
-
-		// scaling factor required to "normalize" origin and local bounds to properly update the material
-		const FVector TextScaleFactor = GetTextScale() * GetComponentScale();
-		SetMaterialGeometryParameterValues(GradientTextMaterialInstance, TextScaleFactor);
-
-		RemoveRefreshReason(EAvaTextRefreshReason::MaterialParametersChange);
-	}
-}
-
-void UAvaText3DComponent::RefreshGradient()
-{
-	MarkForMaterialParametersRefresh();
-	RefreshGradientValues();
-}
-
-void UAvaText3DComponent::RefreshMaskValues()
-{
-	if (IsMarkedForMaterialRefresh() || IsMarkedForMaterialParametersRefresh())
-	{
-		// Set for any/all materials assigned
-		ForEachMID([this](UMaterialInstanceDynamic* InMID)
-		{
-			SetMaskedMaterialValues(InMID);
-		});
-
-		// ... and set for the default materials
-		switch (ColoringStyle)
-		{
-		case EAvaTextColoringStyle::Invalid:
-			break;
-
-		case EAvaTextColoringStyle::Solid:
-			SetMaskedMaterialValues(TextMaterialInstance);
-			SetMaskedMaterialValues(ExtrudeTextMaterialInstance);
-			SetMaskedMaterialValues(BevelTextMaterialInstance);
-			break;
-
-		case EAvaTextColoringStyle::Gradient:
-			SetMaskedMaterialValues(GradientTextMaterialInstance);
-			break;
-
-		case EAvaTextColoringStyle::FromTexture:
-			SetMaskedMaterialValues(TexturedTextMaterialInstance);
-			break;
-
-		case EAvaTextColoringStyle::CustomMaterial:
-			break;
-		}
-
-		RemoveRefreshReason(EAvaTextRefreshReason::MaterialParametersChange);
-	}
-}
-
-void UAvaText3DComponent::RefreshMask()
-{
-	MarkForMaterialParametersRefresh();
-	RefreshMaskValues();
-}
-
-void UAvaText3DComponent::RefreshMaterialGeometryParameterValues()
-{
-	// scaling factor required to "normalize" origin and local bounds to properly update the material
-	const FVector TextScaleFactor = GetTextScale() * GetComponentScale();
-
-	switch (ColoringStyle)
-	{
-	case EAvaTextColoringStyle::Invalid:
-		break;
-
-	case EAvaTextColoringStyle::Solid:
-		SetMaterialGeometryParameterValues(TextMaterialInstance, TextScaleFactor);
-		SetMaterialGeometryParameterValues(ExtrudeTextMaterialInstance, TextScaleFactor);
-		SetMaterialGeometryParameterValues(BevelTextMaterialInstance, TextScaleFactor);
-		break;
-
-	case EAvaTextColoringStyle::Gradient:
-		SetMaterialGeometryParameterValues(GradientTextMaterialInstance, TextScaleFactor);
-		break;
-
-	case EAvaTextColoringStyle::FromTexture:
-		SetMaterialGeometryParameterValues(TexturedTextMaterialInstance, TextScaleFactor);
-		break;
-
-	case EAvaTextColoringStyle::CustomMaterial:
-		break;
-	}
-}
-
-void UAvaText3DComponent::RefreshFormatting()
-{
-	TriggerInternalRebuild(EText3DModifyFlags::Geometry);
-}
-
-void UAvaText3DComponent::SetMaterialGeometryParameterValues(UMaterialInstanceDynamic* InMaterial, const FVector& TextScaleFactor) const
-{
-	if (!InMaterial)
-	{
-		return;
-	}
-
-	// we use the bounds minimum point, because the material needs the bbox origin in the the lower left corner
-	const FVector BoundsLowerLeftCornerLocalScaled = LocalBoundsLowerLeftCorner/TextScaleFactor;
-	const FVector BoundsSizeScaled = LocalBoundsExtent/TextScaleFactor;
-
-	InMaterial->SetVectorParameterValue(UAvaTextMaterialHub::FMaterialStatics::BoundsOrigin_MatParam, BoundsLowerLeftCornerLocalScaled);
-	InMaterial->SetVectorParameterValue(UAvaTextMaterialHub::FMaterialStatics::BoundsSize_MatParam, BoundsSizeScaled);
-
-	FVector TextPositionParameter = GetComponentLocation();
-
-	// in case this component is not the root, let's take into account also relative location
-	if (ParentRootComponent != this)
-	{
-		TextPositionParameter -= GetRelativeLocation();
-	}
-
-	InMaterial->SetVectorParameterValue(UAvaTextMaterialHub::FMaterialStatics::TextPosition_MatParam, TextPositionParameter);
-}
-
-void UAvaText3DComponent::SetMaskedMaterialValues(UMaterialInstanceDynamic* InMaterial)
-{
-	if (InMaterial)
-	{
-		InMaterial->SetScalarParameterValue(UAvaTextMaterialHub::FMaterialStatics::MaskOffset_MatParam, MaskOffset);
-
-		float GradientRotation = UAvaTextMaterialHub::FMaterialStatics::MaterialMaskLeftToRightRotationValue;
-
-		switch (MaskOrientation)
-		{
-		case EAvaMaterialMaskOrientation::LeftRight:
-			GradientRotation = UAvaTextMaterialHub::FMaterialStatics::MaterialMaskLeftToRightRotationValue;
-			break;
-
-		case EAvaMaterialMaskOrientation::RightLeft:
-			GradientRotation = UAvaTextMaterialHub::FMaterialStatics::MaterialMaskRightToLeftRotationValue;
-			break;
-
-		case EAvaMaterialMaskOrientation::Custom:
-			GradientRotation = MaskRotation;
-			break;
-		}
-
-		InMaterial->SetScalarParameterValue(UAvaTextMaterialHub::FMaterialStatics::MaskRotation_MatParam, GradientRotation);
-		InMaterial->SetScalarParameterValue(UAvaTextMaterialHub::FMaterialStatics::MaskSmoothness_MatParam, MaskSmoothness);
-
-		RefreshMaterialGeometryParameterValues();
-	}
-}
-
-#if WITH_EDITOR
-void UAvaText3DComponent::PostEditChangeProperty(FPropertyChangedEvent& PropertyChangedEvent)
-{
-	const FName ChangedPropertyName = PropertyChangedEvent.GetMemberPropertyName();
-
-	if (PropertyChangeFunctions.Contains(ChangedPropertyName))
-	{
-		PropertyChangeFunctions[ChangedPropertyName](this);
-	}
-
-	// Call the Super::PostEditChangeProperty after the PropertyChangeFunctions execute the right Setter/Refresh
-	Super::PostEditChangeProperty(PropertyChangedEvent);
-}
-
-void UAvaText3DComponent::PostEditComponentMove(bool bFinished)
-{
-	Super::PostEditComponentMove(bFinished);
-	ScheduleTextRefreshOnTicker();
-}
-#endif
-
-void UAvaText3DComponent::OnUpdateTransform(EUpdateTransformFlags UpdateTransformFlags, ETeleportType Teleport)
-{
-	Super::OnUpdateTransform(UpdateTransformFlags, Teleport);
-
-	ScheduleTextRefreshOnTicker();
-}
-
-void UAvaText3DComponent::OnAttachmentChanged()
-{
-	Super::OnAttachmentChanged();
-
-	SetupParentRootAndCallbacks();
-}
-
-void UAvaText3DComponent::OnRegister()
-{
-	Super::OnRegister();
-
-	SetComponentTickEnabled(false);
-
-	SetupParentRootAndCallbacks();
-	RefreshText3D();
-}
-
-void UAvaText3DComponent::FormatText(FText& InOutText) const
-{
-	if (bEnforceUpperCase)
-	{
-		InOutText = InOutText.ToUpper();
-	}
 }
diff --git a/Engine/Plugins/Experimental/Avalanche/Source/AvalancheText/Private/AvaTextActor.cpp b/Engine/Plugins/Experimental/Avalanche/Source/AvalancheText/Private/AvaTextActor.cpp
--- a/Engine/Plugins/Experimental/Avalanche/Source/AvalancheText/Private/AvaTextActor.cpp
+++ b/Engine/Plugins/Experimental/Avalanche/Source/AvalancheText/Private/AvaTextActor.cpp
@@ -2,15 +2,13 @@
 
 #include "AvaTextActor.h"
 #include "AvaDefs.h"
-#include "AvaTextCharacterTransform.h"
+#include "Characters/Text3DCharacterTransform.h"
+#include "Extensions/Text3DDefaultMaterialExtension.h"
 
 AAvaTextActor::AAvaTextActor()
 {
 	Text3DComponent = CreateDefaultSubobject<UAvaText3DComponent>(TEXT("Text3DComponent"));
 	RootComponent = Text3DComponent;
-	
-	Text3DCharacterTransform = CreateDefaultSubobject<UAvaTextCharacterTransform>(TEXT("Text3DCharacterTransform"));
-	Text3DCharacterTransform->SetupAttachment(RootComponent);
 }
 
 FAvaColorChangeData AAvaTextActor::GetColorData() const
@@ -18,27 +16,30 @@ FAvaColorChangeData AAvaTextActor::GetColorData() const
 	FAvaColorChangeData ColorData;
 	ColorData.ColorStyle = EAvaColorStyle::None;
 
-	if (const UAvaText3DComponent* const AvaTextComponent = Cast<UAvaText3DComponent>(Text3DComponent))
+	if (Text3DComponent)
 	{
-		switch (AvaTextComponent->GetColoringStyle())
+		if (const UText3DDefaultMaterialExtension* MaterialExtension = Text3DComponent->GetCastedMaterialExtension<UText3DDefaultMaterialExtension>())
 		{
-		case EAvaTextColoringStyle::Solid:
-			ColorData.ColorStyle = EAvaColorStyle::Solid;
-			ColorData.PrimaryColor = AvaTextComponent->GetColor();
-			ColorData.SecondaryColor = AvaTextComponent->GetColor();
-			break;
+			switch (MaterialExtension->GetStyle())
+			{
+			case EText3DMaterialStyle::Solid:
+				ColorData.ColorStyle = EAvaColorStyle::Solid;
+				ColorData.PrimaryColor = MaterialExtension->GetFrontColor();
+				ColorData.SecondaryColor = MaterialExtension->GetFrontColor();
+				break;
 
-		case EAvaTextColoringStyle::Gradient:
-			ColorData.ColorStyle = EAvaColorStyle::LinearGradient;
-			ColorData.PrimaryColor = AvaTextComponent->GetGradientSettings().ColorA;
-			ColorData.SecondaryColor = AvaTextComponent->GetGradientSettings().ColorB;
-			break;
+			case EText3DMaterialStyle::Gradient:
+				ColorData.ColorStyle = EAvaColorStyle::LinearGradient;
+				ColorData.PrimaryColor = MaterialExtension->GetGradientColorA();
+				ColorData.SecondaryColor = MaterialExtension->GetGradientColorB();
+				break;
 
-		default:
-			break;
-		}
+			default:
+				break;
+			}
 		
-		ColorData.bIsUnlit = AvaTextComponent->GetIsUnlit();
+			ColorData.bIsUnlit = MaterialExtension->GetIsUnlit();
+		}
 	}
 
 	return ColorData;
@@ -46,24 +47,31 @@ FAvaColorChangeData AAvaTextActor::GetColorData() const
 
 void AAvaTextActor::SetColorData(const FAvaColorChangeData& NewColorData)
 {
-	if (UAvaText3DComponent* const AvaTextComponent = Cast<UAvaText3DComponent>(Text3DComponent))
+	if (Text3DComponent)
 	{
-		switch (NewColorData.ColorStyle)
+		if (UText3DDefaultMaterialExtension* MaterialExtension = Text3DComponent->GetCastedMaterialExtension<UText3DDefaultMaterialExtension>())
 		{
+			switch (NewColorData.ColorStyle)
+			{
 			case EAvaColorStyle::Solid:
-				AvaTextComponent->SetColoringStyle(EAvaTextColoringStyle::Solid);
-				AvaTextComponent->SetColor(NewColorData.PrimaryColor);
+				MaterialExtension->SetStyle(EText3DMaterialStyle::Solid);
+				MaterialExtension->SetFrontColor(NewColorData.PrimaryColor);
+				MaterialExtension->SetBackColor(NewColorData.PrimaryColor);
+				MaterialExtension->SetBevelColor(NewColorData.PrimaryColor);
+				MaterialExtension->SetExtrudeColor(NewColorData.PrimaryColor);
 				break;
 
 			case EAvaColorStyle::LinearGradient:
-				AvaTextComponent->SetColoringStyle(EAvaTextColoringStyle::Gradient);
-				AvaTextComponent->SetGradientColors(NewColorData.PrimaryColor, NewColorData.SecondaryColor);
+				MaterialExtension->SetStyle(EText3DMaterialStyle::Gradient);
+				MaterialExtension->SetGradientColorA(NewColorData.PrimaryColor);
+				MaterialExtension->SetGradientColorB(NewColorData.SecondaryColor);
 				break;
 
 			default:
 				break;
-		}
+			}
 
-		AvaTextComponent->SetIsUnlit(NewColorData.bIsUnlit);
+			MaterialExtension->SetIsUnlit(NewColorData.bIsUnlit);
+		}
 	}
 }
