MotionDesign : Text3D

diff --git a/Engine/Plugins/Experimental/Avalanche/Source/AvalancheText/Public/AvaText3DComponent.h b/Engine/Plugins/Experimental/Avalanche/Source/AvalancheText/Public/AvaText3DComponent.h
--- a/Engine/Plugins/Experimental/Avalanche/Source/AvalancheText/Public/AvaText3DComponent.h
+++ b/Engine/Plugins/Experimental/Avalanche/Source/AvalancheText/Public/AvaText3DComponent.h
@@ -3,25 +3,10 @@
 #pragma once
 
 #include "AvaTextDefs.h"
-#include "Containers/Ticker.h"
 #include "Font/AvaFont.h"
 #include "Text3DComponent.h"
 #include "AvaText3DComponent.generated.h"
 
-class UMaterial;
-class UMaterialInstanceDynamic;
-
-enum class EAvaTextRefreshReason : uint8
-{
-	NoRefresh                 = 0,
-	GeometryChange            = 1 << 0,
-	LayoutChange              = 1 << 1,
-	MaterialInstancesChange   = 1 << 2,
-	MaterialParametersChange  = 1 << 3,
-	FullRefresh               = 0xFF
-};
-ENUM_CLASS_FLAGS(EAvaTextRefreshReason)
-
 UCLASS(MinimalAPI, ClassGroup="Text3D", PrioritizeCategories=("Text", "Layout", "Geometry", "Style", "Materials"), meta=(BlueprintSpawnableComponent))
 class UAvaText3DComponent : public UText3DComponent
 {
@@ -34,314 +19,70 @@ class UAvaText3DComponent : public UText3DComponent
 public:
 	UAvaText3DComponent();
 
-	UFUNCTION(BlueprintSetter)
-	AVALANCHETEXT_API void SetEnforceUpperCase(bool bInEnforceUpperCase);
-	bool IsEnforcingUpperCase() const { return bEnforceUpperCase; }
-
-	UFUNCTION(BlueprintSetter)
-	AVALANCHETEXT_API void SetMotionDesignFont(const FAvaFont& InFont);
-	const FAvaFont& GetMotionDesignFont() const { return MotionDesignFont; }
-
-	UFUNCTION(BlueprintSetter)
-	AVALANCHETEXT_API void SetAlignment(FAvaTextAlignment InAlignment);
-	const FAvaTextAlignment& GetAlignment() const { return Alignment; }
-
-	UFUNCTION(BlueprintSetter)
-	AVALANCHETEXT_API void SetColoringStyle(const EAvaTextColoringStyle& InColoringStyle);
-	EAvaTextColoringStyle GetColoringStyle() const { return ColoringStyle; }
-
-	UFUNCTION(BlueprintSetter)
-	AVALANCHETEXT_API void SetColor(const FLinearColor& InColor);
-	const FLinearColor& GetColor() const { return Color; }
-
-	UFUNCTION(BlueprintSetter)
-	AVALANCHETEXT_API void SetExtrudeColor(const FLinearColor& InColor);
-	const FLinearColor& GetExtrudeColor() const { return ExtrudeColor; }
-
-	UFUNCTION(BlueprintSetter)
-	AVALANCHETEXT_API void SetBevelColor(const FLinearColor& InColor);
-	const FLinearColor& GetBevelColor() const { return BevelColor; }
-
-	UFUNCTION(BlueprintSetter)
-	AVALANCHETEXT_API void SetGradientSettings(const FAvaLinearGradientSettings& InGradientSettings);
-	const FAvaLinearGradientSettings& GetGradientSettings() const { return GradientSettings; }
-
-	UFUNCTION(BlueprintSetter)
-	AVALANCHETEXT_API void SetGradientColors(const FLinearColor& InColorA, const FLinearColor& InColorB);
-
-	UFUNCTION(BlueprintSetter)
-	AVALANCHETEXT_API void SetMainTexture(UTexture2D* InMainTexture);
-	UTexture2D* GetMainTexture() const { return MainTexture; }
-
-	UFUNCTION(BlueprintSetter)
-	AVALANCHETEXT_API void SetTiling(const FVector2D InTiling);
-	const FVector2D& GetTiling() const { return Tiling; }
-
-	UFUNCTION(BlueprintSetter)
-	AVALANCHETEXT_API void SetTranslucencyStyle(EAvaTextTranslucency InTranslucencyStyle);
-	EAvaTextTranslucency GetTranslucencyStyle() const { return TranslucencyStyle; }
-
-	UFUNCTION(BlueprintSetter)
-	AVALANCHETEXT_API void SetOpacity(float InOpacity);
-	float GetOpacity() const { return Opacity; }
-
-	UFUNCTION(BlueprintSetter)
-	AVALANCHETEXT_API void SetMaskOrientation(EAvaMaterialMaskOrientation InMaskOrientation);
-	EAvaMaterialMaskOrientation GetMaskOrientation() const { return MaskOrientation; }
-
-	UFUNCTION(BlueprintSetter)
-	AVALANCHETEXT_API void SetMaskSmoothness(float InMaskSmoothness);
-	float GetMaskSmoothness() const { return MaskSmoothness; }
-
-	UFUNCTION(BlueprintSetter)
-	AVALANCHETEXT_API void SetMaskOffset(float InMaskOffset);
-	float GetMaskOffset() const { return MaskOffset; }
-
-	UFUNCTION(BlueprintSetter)
-	AVALANCHETEXT_API void SetMaskRotation(float InMaskRotation);
-	float GetMaskRotation() const { return MaskRotation; }
-
-	UFUNCTION(BlueprintSetter)
-	AVALANCHETEXT_API void SetIsUnlit(bool bInIsUnlit);
-	bool GetIsUnlit() const { return bIsUnlit; }
-
-	using FOnGetMaterialWithSettings = TDelegate<UMaterialInterface*(const UMaterialInterface* InPreviousMaterial, const FAvaTextMaterialSettings& InSettings)>;
-
-	AVALANCHETEXT_API FOnGetMaterialWithSettings& GetMaterialProviderDelegate();
-
-	AVALANCHETEXT_API void ForEachMID(TUniqueFunction<void(UMaterialInstanceDynamic* InMID)>&& InFunc) const;
-
-	AVALANCHETEXT_API FVector GetGradientDirection() const;
-
-	AVALANCHETEXT_API void RefreshMaterialInstances();
-
-protected:
-	void RefreshText3D();
-	void MarkRenderDirty();
-
-	void MarkForGeometryRefresh();
-	bool IsMarkedForGeometryRefresh() const;
-
-	void MarkForLayoutRefresh();
-	bool IsMarkedForLayoutRefresh() const;
-
-	/** Mark this Text3D component so that at the next RefreshText3D call it will refresh both materials and their parameters*/
-	void MarkForMaterialsFullRefresh();
-	bool IsMarkedForMaterialRefresh() const;
-
-	void MarkForMaterialParametersRefresh();
-	bool IsMarkedForMaterialParametersRefresh() const;
-
-	void RemoveRefreshReason(EAvaTextRefreshReason InRefreshReasonToRemove);
-
-	void MarkForFullRefresh();
-	void ClearRefreshReasons();
-
-	void RegisterCallbacks();
-	void RefreshParentRootComponent();
-	void SetupParentRootAndCallbacks();
-
-	void SetupTextMaterials(const EAvaTextColoringStyle& InColoringStyle);
-	void SetupTextGeometry();
-	void SetupTextLayout();
-
-	void SetupSolidMaterial();
-	void SetupGradientMaterial();
-	void SetupCustomMaterial();
-	void RefreshTranslucencyValues();
-	void SetupTexturedMaterial();
-
-	void RefreshStoredBoundsValues();
-
-	//~ Begin UObject
-	virtual void PostLoad() override;
-	virtual void PostDuplicate(EDuplicateMode::Type DuplicateMode) override;
-	virtual void PostInitProperties() override;
-	virtual void PostEditImport() override;
-	//~ Begin UObject
-
-	void PostDuplicateActions();
-
-	/** Will schedule a Text Refresh to be executed by Core Ticker */
-	void ScheduleTextRefreshOnTicker();
-
-	void OnTextGeneratedActions();
-	void OnTextGeometryGenerated();
-
-	void OnRootComponentTransformUpdated(USceneComponent* InSceneComponent, EUpdateTransformFlags InUpdateTransformFlags, ETeleportType InTeleportType);
-
-	/**
-	 * This function can be used to obtain implicit data from the base UText3DComponent.
-	 * 
-	 * Previously, the logic from UAvaText3DComponent was handled by AAvaTextActor, which is now mostly a wrapper for UAvaText3DComponent.
-	 * Therefore, previously saved AAvaTextActor assets will have UText3DComponent materials compatible with UAvaText3DComponent.
-	 * This function extracts implicit information from those materials, and updates UAvaText3DComponent properties accordingly.
-	 * Furthermore, materials themselves will be retrieved, so that this component won't have to store them as well.
-	 * 
-	 * Additionally, also text, font and layout information will be retrieved, to avoid potential data loss from older assets.
-	 */
-	void RetrieveDataFromBaseComponent();
-
-	/** Retrieves font and text from the values stored in the component */
-	void RetrieveCurrentTextAndFont();
-
-	/** Retrieves alignment information from the values stored in the component */
-	void RetrieveCurrentAlignment();
-
-	/** Retrieves materials from the component, according to the current settings*/
-	void RetrieveCurrentMaterials();
-
-	void RetrieveCurrentColoringStyle();
-
-	void RetrieveSolidMaterialInstances();
-	void RetrieveGradientMaterialInstance();
-	void RetrieveTexturedMaterialInstance();
-
-	void RetrieveTranslucentMaterialParameters();
-	void RetrieveMaskedMaterialParameters(const UMaterialInstanceDynamic* InMaterialInstance);
-	void RetrieveSimpleTranslucentMaterialParameters(const UMaterialInstanceDynamic* InMaterialInstance);
-
-	UMaterialInstanceDynamic* GetMIDWithSettings(const UMaterialInterface* InPreviousMaterial, const FAvaTextMaterialSettings& InSettings);
-
-	EAvaTextMaterialFeatures GetMaterialFeaturesFromProperties() const;
-
-	/**
-	 * UMaterialInstanceDynamic references are transient.
-	 * We need to retrieve them from material when loading assets using this component, in order to avoid having to create the material again.
-	 * Refer to UAvaText3DComponent::RetrieveCurrentMaterials().
-	 *
-	 * Additionally, this helps when loading old versions of the AAvaTextActor, which was hosting all the custom properties now handled by this component (UAvaText3DComponent)
-	 */
-
-	UPROPERTY(Transient)
-	TObjectPtr<UMaterialInstanceDynamic> TextMaterialInstance;
-
-	UPROPERTY(Transient)
-	TObjectPtr<UMaterialInstanceDynamic> ExtrudeTextMaterialInstance;
-
-	UPROPERTY(Transient)
-	TObjectPtr<UMaterialInstanceDynamic> BevelTextMaterialInstance;
-
-	UPROPERTY(Transient)
-	TObjectPtr<UMaterialInstanceDynamic> TexturedTextMaterialInstance;
-
-	UPROPERTY(Transient)
-	TObjectPtr<UMaterialInstanceDynamic> GradientTextMaterialInstance;
-
-	TWeakObjectPtr<USceneComponent> ParentRootComponent;
-
-	FOnGetMaterialWithSettings OnGetMaterialWithSettingsDelegate;
-
-	/** Used to filter which refreshes need be performed when calling RefreshText3D */
-	EAvaTextRefreshReason TextRefreshReason;
-
-	FVector LocalBoundsLowerLeftCorner;
-	FVector LocalBoundsExtent;
-
-	FDelegateHandle OnTextGeometryGeneratedDelegate;
-	FDelegateHandle OnTransformUpdatedDelegate;
-	FTSTicker::FDelegateHandle TickerHandle;
-
-private:
-#if WITH_EDITOR
-	static void RegisterOnPropertyChangeFunctions();
-#endif
-
-	void RefreshOpacity() const;
-	void RefreshColor() const;
-	void RefreshExtrudeColor() const;
-	void RefreshBevelColor() const;
-
-	void RefreshAlignment();
-	void RefreshMotionDesignFont();
-	void RefreshHorizontalAlignment();
-	void RefreshVerticalAlignment();
-	void RefreshUnlit();
-	void RefreshMasked();
-	void RefreshGeometry();
-	void RefreshLayout();
-	void RefreshGradientValues();
-	AVALANCHETEXT_API void RefreshGradient();
-	void RefreshMaskValues();
-	void RefreshMask();
-	void RefreshMaterialGeometryParameterValues();
-
-	void RefreshFormatting();
-
-	void SetMaterialGeometryParameterValues(UMaterialInstanceDynamic* InMaterial, const FVector& TextScaleFactor) const;
-	void SetMaskedMaterialValues(UMaterialInstanceDynamic* InMaterial);
-
-	//~ Begin UObject
-#if WITH_EDITOR
-	virtual void PostEditChangeProperty(FPropertyChangedEvent& PropertyChangedEvent) override;
-	//~ End UObject
-#endif
-
-	//~ Begin USceneComponent
-#if WITH_EDITOR
-	virtual void PostEditComponentMove(bool bFinished) override;
-#endif
-	virtual void OnUpdateTransform(EUpdateTransformFlags UpdateTransformFlags, ETeleportType Teleport) override;
-	virtual void OnAttachmentChanged() override;
-	virtual void OnRegister() override;
-	//~ End USceneComponent
-
 	//~ Begin UText3DComponent
-	virtual void FormatText(FText& InOutText) const override;
+	virtual void Serialize(FArchive& InArchive) override;
 	//~ End UText3DComponent
 
-	UPROPERTY(EditAnywhere, Getter="IsEnforcingUpperCase", Setter="SetEnforceUpperCase", Category="Text", meta=(DisplayAfter="Text", AllowPrivateAccess="true"))
-	bool bEnforceUpperCase = false;
-
-	UPROPERTY(EditAnywhere, Getter="GetMotionDesignFont", Setter="SetMotionDesignFont", Category="Text", meta=(DisplayAfter="Text", AllowPrivateAccess="true"))
-	FAvaFont MotionDesignFont;
-
-	UPROPERTY(EditAnywhere, Getter="GetAlignment", Setter="SetAlignment", Category="Layout", meta=(DisplayAfter="VerticalAlignment", AllowPrivateAccess="true"))
-	FAvaTextAlignment Alignment;
-
-	UPROPERTY(EditAnywhere, Getter="GetColoringStyle", Setter="SetColoringStyle", Category="Style", meta=(DisplayAfter="Text", AllowPrivateAccess="true"))
-	EAvaTextColoringStyle ColoringStyle = EAvaTextColoringStyle::Solid;
-
-	UPROPERTY(EditAnywhere, Getter="GetColor", Setter="SetColor", Category="Style", meta=(EditCondition="ColoringStyle == EAvaTextColoringStyle::Solid", EditConditionHides, DisplayAfter="Text", AllowPrivateAccess="true"))
-	FLinearColor Color = FLinearColor::White;
-
-	UPROPERTY(EditAnywhere, Getter="GetExtrudeColor", Setter="SetExtrudeColor", Category="Style", meta=(EditCondition="ColoringStyle == EAvaTextColoringStyle::Solid && Extrude > 0.0", EditConditionHides, DisplayAfter="Text", AllowPrivateAccess="true"))
-	FLinearColor ExtrudeColor = FLinearColor::Gray;
-
-	UPROPERTY(EditAnywhere, Getter="GetBevelColor", Setter="SetBevelColor", Category="Style", meta=(EditCondition="ColoringStyle == EAvaTextColoringStyle::Solid && Bevel > 0.0", EditConditionHides, DisplayAfter="Text", AllowPrivateAccess="true"))
-	FLinearColor BevelColor = FLinearColor::Gray;
-
-	UPROPERTY(EditAnywhere, Getter="GetGradientSettings", Setter="SetGradientSettings", Category="Style", meta=(EditCondition="ColoringStyle == EAvaTextColoringStyle::Gradient", EditConditionHides, DisplayAfter="Text", AllowPrivateAccess="true"))
-	FAvaLinearGradientSettings GradientSettings;
-
-	UPROPERTY(EditAnywhere, Getter="GetMainTexture", Setter="SetMainTexture", Category="Style", meta=(EditCondition="ColoringStyle == EAvaTextColoringStyle::FromTexture", EditConditionHides, DisplayAfter="Text", AllowPrivateAccess="true"))
-	TObjectPtr<UTexture2D> MainTexture;
-
-	UPROPERTY(EditAnywhere, Getter="GetTiling", Setter="SetTiling", Category="Style", meta=(Delta=0.1, EditCondition="ColoringStyle == EAvaTextColoringStyle::FromTexture", EditConditionHides, DisplayAfter="Text", AllowPrivateAccess="true"))
-	FVector2D Tiling;
-
-	UPROPERTY(EditAnywhere, Getter="GetTranslucencyStyle", Setter="SetTranslucencyStyle", Category="Style", meta=(EditCondition="ColoringStyle != EAvaTextColoringStyle::CustomMaterial", EditConditionHides,DisplayAfter="Text", AllowPrivateAccess="true"))
-	EAvaTextTranslucency TranslucencyStyle = EAvaTextTranslucency::None;
-
-	UPROPERTY(EditAnywhere, Getter="GetOpacity", Setter="SetOpacity", Category="Style", meta=(ClampMin=0.0, ClampMax=1.0, EditCondition="TranslucencyStyle == EAvaTextTranslucency::Translucent && ColoringStyle != EAvaTextColoringStyle::CustomMaterial", EditConditionHides, DisplayAfter="Text", AllowPrivateAccess="true"))
-	float Opacity = 1.f;
-
-	UPROPERTY(EditAnywhere, Getter="GetMaskOrientation", Setter="SetMaskOrientation", Category="Style", meta=(EditCondition="TranslucencyStyle == EAvaTextTranslucency::GradientMask", EditConditionHides, DisplayAfter="Text", AllowPrivateAccess="true"))
-	EAvaMaterialMaskOrientation MaskOrientation = EAvaMaterialMaskOrientation::LeftRight;
-
-	UPROPERTY(EditAnywhere, Getter="GetMaskSmoothness", Setter="SetMaskSmoothness", Category="Style", meta=(ClampMin=0.0, ClampMax=1.0, EditCondition="TranslucencyStyle == EAvaTextTranslucency::GradientMask", EditConditionHides, DisplayAfter="Text", AllowPrivateAccess="true"))
-	float MaskSmoothness = 0.1f;
-
-	UPROPERTY(EditAnywhere, Getter="GetMaskOffset", Setter="SetMaskOffset", Category="Style", meta=(ClampMin=0.0, ClampMax=1.0, EditCondition="TranslucencyStyle == EAvaTextTranslucency::GradientMask", EditConditionHides, DisplayAfter="Text", AllowPrivateAccess="true"))
-	float MaskOffset = 0.5f;
-
-	UPROPERTY(EditAnywhere, Getter="GetMaskRotation", Setter="SetMaskRotation", Category="Style", meta=(ClampMin=0.0, ClampMax=1.0, EditCondition="MaskOrientation == EAvaMaterialMaskOrientation::Custom && TranslucencyStyle == EAvaTextTranslucency::GradientMask", EditConditionHides, DisplayAfter="Text", AllowPrivateAccess="true"))
-	float MaskRotation = 0.0f;
-
-	UPROPERTY(EditAnywhere, Getter="GetIsUnlit", Setter="SetIsUnlit", Category="Style", meta=(EditCondition="ColoringStyle != EAvaTextColoringStyle::CustomMaterial", EditConditionHides, DisplayAfter="Text", AllowPrivateAccess="true"))
-	bool bIsUnlit = true;
-
-#if WITH_EDITOR
-	static inline TMap<FName, TFunction<void(UAvaText3DComponent*)>> PropertyChangeFunctions;
+private:
+#if WITH_EDITORONLY_DATA
+	UE_DEPRECATED(5.6, "Use SetFont instead")
+	UPROPERTY()
+	FAvaFont MotionDesignFont_DEPRECATED;
+
+	UE_DEPRECATED(5.6, "Use Material Extension instead")
+	UPROPERTY()
+	EAvaTextColoringStyle ColoringStyle_DEPRECATED = EAvaTextColoringStyle::Solid;
+
+	UE_DEPRECATED(5.6, "Use Material Extension instead")
+	UPROPERTY()
+	FLinearColor Color_DEPRECATED = FLinearColor::White;
+
+	UE_DEPRECATED(5.6, "Use Material Extension instead")
+	UPROPERTY()
+	FLinearColor ExtrudeColor_DEPRECATED = FLinearColor::Gray;
+
+	UE_DEPRECATED(5.6, "Use Material Extension instead")
+	UPROPERTY()
+	FLinearColor BevelColor_DEPRECATED = FLinearColor::Gray;
+
+	UE_DEPRECATED(5.6, "Use Material Extension instead")
+	UPROPERTY()
+	FAvaLinearGradientSettings GradientSettings_DEPRECATED;
+
+	UE_DEPRECATED(5.6, "Use Material Extension instead")
+	UPROPERTY()
+	TObjectPtr<UTexture2D> MainTexture_DEPRECATED;
+
+	UE_DEPRECATED(5.6, "Use Material Extension instead")
+	UPROPERTY()
+	FVector2D Tiling_DEPRECATED = FVector2D::One();
+
+	UE_DEPRECATED(5.6, "Use Material Extension instead")
+	UPROPERTY()
+	EAvaTextTranslucency TranslucencyStyle_DEPRECATED = EAvaTextTranslucency::None;
+
+	UE_DEPRECATED(5.6, "Use Material Extension instead")
+	UPROPERTY()
+	float Opacity_DEPRECATED = 1.f;
+
+	UE_DEPRECATED(5.6, "Use Material Extension instead")
+	UPROPERTY()
+	EAvaMaterialMaskOrientation MaskOrientation_DEPRECATED = EAvaMaterialMaskOrientation::LeftRight;
+
+	UE_DEPRECATED(5.6, "Use Material Extension instead")
+	UPROPERTY()
+	float MaskSmoothness_DEPRECATED = 0.1f;
+
+	UE_DEPRECATED(5.6, "Use Material Extension instead")
+	UPROPERTY()
+	float MaskOffset_DEPRECATED = 0.5f;
+
+	UE_DEPRECATED(5.6, "Use Material Extension instead")
+	UPROPERTY()
+    float MaskRotation_DEPRECATED = 0.0f;
+
+	UE_DEPRECATED(5.6, "Use Material Extension instead")
+	UPROPERTY()
+	bool bIsUnlit_DEPRECATED = true;
 #endif
 };
diff --git a/Engine/Plugins/Experimental/Avalanche/Source/AvalancheText/Public/AvaTextActor.h b/Engine/Plugins/Experimental/Avalanche/Source/AvalancheText/Public/AvaTextActor.h
--- a/Engine/Plugins/Experimental/Avalanche/Source/AvalancheText/Public/AvaTextActor.h
+++ b/Engine/Plugins/Experimental/Avalanche/Source/AvalancheText/Public/AvaTextActor.h
@@ -6,7 +6,7 @@
 #include "GameFramework/Actor.h"
 #include "AvaTextActor.generated.h"
 
-class UAvaTextCharacterTransform;
+class UText3DCharacterTransform;
 struct FAvaColorChangeData;
 
 UCLASS(MinimalAPI, ClassGroup = (Text3D), DisplayName = "Motion Design Text", meta = (ComponentWrapperClass))
@@ -23,13 +23,13 @@ public:
 	}
 
 	AVALANCHETEXT_API FAvaColorChangeData GetColorData() const;
-
 	AVALANCHETEXT_API void SetColorData(const FAvaColorChangeData& InColorData);
 
 private:
 	UPROPERTY(Category = "Text", VisibleAnywhere, BlueprintReadOnly, meta = (AllowPrivateAccess = "true"))
 	TObjectPtr<UText3DComponent> Text3DComponent;
 
-	UPROPERTY(Category = "Layout", EditAnywhere, BlueprintReadWrite, meta = (AllowPrivateAccess = "true"))
-	TObjectPtr<UAvaTextCharacterTransform> Text3DCharacterTransform;
+	UE_DEPRECATED(5.6, "Use Transform Extension instead")
+	UPROPERTY()
+	TObjectPtr<UText3DCharacterTransform> Text3DCharacterTransform;
 };
diff --git a/Engine/Plugins/Experimental/Avalanche/Source/AvalancheText/Public/AvaTextDefs.h b/Engine/Plugins/Experimental/Avalanche/Source/AvalancheText/Public/AvaTextDefs.h
--- a/Engine/Plugins/Experimental/Avalanche/Source/AvalancheText/Public/AvaTextDefs.h
+++ b/Engine/Plugins/Experimental/Avalanche/Source/AvalancheText/Public/AvaTextDefs.h
@@ -5,15 +5,6 @@
 #include "Text3DComponent.h"
 #include "AvaTextDefs.generated.h"
 
-UENUM()
-enum class EAvaOutlineType : uint8 
-{
-	Invalid     UMETA(Hidden),
-	None,
-	AddOutline,
-	OutlineOnly
-};
-
 UENUM()
 enum class EAvaTextColoringStyle : uint8 
 {
@@ -41,23 +32,6 @@ enum class EAvaMaterialMaskOrientation : uint8
 	Custom
 };
 
-USTRUCT(BlueprintType)
-struct FAvaTextAlignment
-{
-	GENERATED_BODY()
-    
-	UPROPERTY(EditAnywhere, BlueprintReadWrite, Category = "Motion Design")
-	EText3DHorizontalTextAlignment HorizontalAlignment = EText3DHorizontalTextAlignment::Left;
-
-	UPROPERTY(EditAnywhere, BlueprintReadWrite, Category = "Motion Design")
-	EText3DVerticalTextAlignment VerticalAlignment = EText3DVerticalTextAlignment::FirstLine;
-	
-	bool operator==(const FAvaTextAlignment& Other) const
-	{
-		return HorizontalAlignment == Other.HorizontalAlignment && VerticalAlignment == Other.VerticalAlignment;
-	}
-};
-
 UENUM()
 enum class EAvaGradientDirection : uint8 
 {
@@ -101,113 +75,3 @@ struct FAvaLinearGradientSettings
 			Rotation == Other.Rotation;
 	}
 };
-
-/**
- * TODO: use this for UAvaText3DComponent in place of separate mask settings
- * Main Blocker from using this now is that it needs custom tracks/sections/property track editor to be keyable for Sequencer
- */
-USTRUCT(BlueprintType)
-struct FAvaMaterialMaskSettings
-{
-	GENERATED_BODY()
-	
-	UPROPERTY(EditAnywhere, BlueprintReadWrite, Category = "Text")
-	EAvaMaterialMaskOrientation Orientation = EAvaMaterialMaskOrientation::LeftRight;
-	
-	UPROPERTY(Interp, EditAnywhere, BlueprintReadWrite, Category = "Text", meta = (ClampMin = 0.0, ClampMax = 1.0))
-	float Smoothness = 0.1f;
-	
-	UPROPERTY(Interp, EditAnywhere, BlueprintReadWrite, Category = "Text", meta = (ClampMin = 0.0, ClampMax = 1.0))
-	float Offset = 0.5f;
-
-	UPROPERTY(Interp, EditAnywhere, BlueprintReadWrite, Category = "Text", meta = (ClampMin = 0.0, ClampMax = 1.0, EditCondition = "Orientation==EAvaMaterialMaskOrientation::Custom", EditConditionHides))
-	float Rotation = 0.0f;
-
-	bool operator==(const FAvaMaterialMaskSettings& Other) const
-	{
-		return
-			Orientation == Other.Orientation &&
-			Smoothness == Other.Smoothness &&
-			Offset == Other.Offset &&
-			Rotation == Other.Rotation;
-	}
-};
-
-UENUM()
-enum class EAvaTextLength : uint8 
-{
-	Unlimited,
-	CharacterCount UMETA(DisplayName = "Maximum Character Count")
-};
-
-UENUM()
-enum class EAvaTextCase : uint8 
-{
-	Regular,
-	UpperCase,
-	LowerCase
-};
-
-USTRUCT(BlueprintType)
-struct FAvaTextField
-{
-	GENERATED_BODY()
-	
-	UPROPERTY(EditAnywhere, BlueprintReadWrite, Category = "Motion Design", meta = (MultiLine = true))
-	FText Text;
-
-	UPROPERTY(EditAnywhere, BlueprintReadWrite, Category = "Motion Design")
-	EAvaTextLength MaximumLength = EAvaTextLength::Unlimited;
-
-	UPROPERTY(EditAnywhere, BlueprintReadWrite, Category = "Motion Design", meta = (EditCondition = "MaximumLength==EAvaTextLength::CharacterCount", EditConditionHides))
-	int32 Count = 20;
-
-	UPROPERTY(EditAnywhere, BlueprintReadWrite, Category = "Motion Design")
-	EAvaTextCase TextCase = EAvaTextCase::Regular;
-};
-
-/**
- * Flags to list the desired features of a Text Material
- */
-enum class EAvaTextMaterialFeatures : uint8
-{
-	None         = 0,
-	GradientMask = 1 << 0,
-	Unlit        = 1 << 1,
-	Translucent  = 1 << 2
-};
-ENUM_CLASS_FLAGS(EAvaTextMaterialFeatures)
-
-struct FAvaTextMaterialSettings
-{
-	/** The coloring style used for the material (e.g. Solid, Gradient, Textured) */
-	EAvaTextColoringStyle ColoringStyle = EAvaTextColoringStyle::Solid;
-
-	/** The features of the material (e.g. Masked, Unlit) */
-	EAvaTextMaterialFeatures MaterialFeatures = EAvaTextMaterialFeatures::Unlit;
-
-	FAvaTextMaterialSettings() = default;
-	
-	FAvaTextMaterialSettings(const EAvaTextColoringStyle InColoringStyle, const EAvaTextMaterialFeatures InMaterialFeatures)
-		: ColoringStyle(InColoringStyle)
-		, MaterialFeatures(InMaterialFeatures)
-	{
-	}
-
-	bool operator==(const FAvaTextMaterialSettings& Other) const
-	{
-		return
-			ColoringStyle == Other.ColoringStyle &&
-			MaterialFeatures == Other.MaterialFeatures;
-	}
-
-	bool operator!=(const FAvaTextMaterialSettings& Other) const
-	{
-		return !(*this == Other);
-	}
-	
-	friend uint32 GetTypeHash(const FAvaTextMaterialSettings& InMaterialSettings)
-	{
-		return HashCombineFast(GetTypeHash(InMaterialSettings.ColoringStyle), GetTypeHash(InMaterialSettings.MaterialFeatures));
-	}
-};
diff --git a/Engine/Plugins/Experimental/Avalanche/Source/AvalancheText/Public/Font/AvaFont.h b/Engine/Plugins/Experimental/Avalanche/Source/AvalancheText/Public/Font/AvaFont.h
--- a/Engine/Plugins/Experimental/Avalanche/Source/AvalancheText/Public/Font/AvaFont.h
+++ b/Engine/Plugins/Experimental/Avalanche/Source/AvalancheText/Public/Font/AvaFont.h
@@ -8,7 +8,7 @@
 
 #include "AvaFont.generated.h"
 
-USTRUCT(BlueprintType, meta=(DisplayName="Motion Design Font"))
+USTRUCT()
 struct FAvaFont
 {
 	friend class UAvaText3DComponent;
diff --git a/Engine/Plugins/Experimental/Avalanche/Source/AvalancheText/Public/Font/AvaFontObject.h b/Engine/Plugins/Experimental/Avalanche/Source/AvalancheText/Public/Font/AvaFontObject.h
--- a/Engine/Plugins/Experimental/Avalanche/Source/AvalancheText/Public/Font/AvaFontObject.h
+++ b/Engine/Plugins/Experimental/Avalanche/Source/AvalancheText/Public/Font/AvaFontObject.h
@@ -27,8 +27,6 @@ struct FAvaSystemFontMetrics
 	bool bIsBold;
 
 	bool bIsItalic;
-
-	FString FamilyName;
 };
 
 struct FSystemFontsRetrieveParams
@@ -55,7 +53,7 @@ private:
 	TArray<FString> FontFacePaths;
 };
 
-UCLASS(BlueprintType)
+UCLASS()
 class AVALANCHETEXT_API UAvaFontObject : public UObject
 {
 	GENERATED_BODY()
@@ -69,9 +67,6 @@ private:
 
 	FAvaSystemFontMetrics Metrics;
 
-	UPROPERTY(Transient)
-	TArray<TObjectPtr<UFontFace>> FontFaces;
-
 	FSystemFontsRetrieveParams FontParameters;
 
 public:
diff --git a/Engine/Plugins/Experimental/Avalanche/Source/AvalancheTextEditor/AvalancheTextEditor.Build.cs b/Engine/Plugins/Experimental/Avalanche/Source/AvalancheTextEditor/AvalancheTextEditor.Build.cs
--- a/Engine/Plugins/Experimental/Avalanche/Source/AvalancheTextEditor/AvalancheTextEditor.Build.cs
+++ b/Engine/Plugins/Experimental/Avalanche/Source/AvalancheTextEditor/AvalancheTextEditor.Build.cs
@@ -21,6 +21,7 @@ public class AvalancheTextEditor : ModuleRules
 			{
 				"Avalanche",
 				"AvalancheComponentVisualizers",
+				"AvalancheEditorCore",
 				"AvalancheInteractiveTools",
 				"AvalancheLevelViewport",
 				"AvalancheText",
@@ -28,12 +29,10 @@ public class AvalancheTextEditor : ModuleRules
 				"Core",
 				"CoreUObject",
 				"DynamicMaterialEditor",
-				"EditorScriptingUtilities",
-				"EditorSubsystem",
 				"Engine",
-				"FreeType2",
-				"InputCore",
+				"EditorScriptingUtilities",
 				"InteractiveToolsFramework",
+				"InputCore",
 				"Projects",
 				"PropertyEditor",
 				"Slate",
@@ -42,33 +41,5 @@ public class AvalancheTextEditor : ModuleRules
 				"UnrealEd",
 			}
 		);
-
-		if (Target.Platform.IsInGroup(UnrealPlatformGroup.Windows))
-		{
-			PrivateDependencyModuleNames.Add("DirectX");
-
-			PublicAdditionalLibraries.AddRange(new string[] {
-				Path.Combine(Target.WindowsPlatform.DirectXLibDir, "dwrite.lib")
-			});
-		}
-
-		// Needed for FreeType/Font functionality
-		AddEngineThirdPartyPrivateStaticDependencies(Target,
-			"zlib",
-			"UElibPNG"
-		);
-
-		if (Target.Type == TargetRules.TargetType.Editor)
-		{
-			PrivateDependencyModuleNames.AddRange(new string[]
-			{
-				"AvalancheEditorCore"
-			});
-		}
-
-		if (Target.Platform == UnrealTargetPlatform.Mac)
-		{
-			PublicFrameworks.AddRange(new string[] { "Foundation", "CoreText", "ApplicationServices"});
-		}
 	}
 }
diff --git a/Engine/Plugins/Experimental/Avalanche/Source/AvalancheTextEditor/Private/AvaTextEditorModule.cpp b/Engine/Plugins/Experimental/Avalanche/Source/AvalancheTextEditor/Private/AvaTextEditorModule.cpp
--- a/Engine/Plugins/Experimental/Avalanche/Source/AvalancheTextEditor/Private/AvaTextEditorModule.cpp
+++ b/Engine/Plugins/Experimental/Avalanche/Source/AvalancheTextEditor/Private/AvaTextEditorModule.cpp
@@ -10,11 +10,6 @@
 #include "ColorPicker/AvaViewportColorPickerAdapter.h"
 #include "DMObjectMaterialProperty.h"
 #include "DetailsView/AvaLinearGradientSettingsCustomization.h"
-#include "DetailsView/AvaTextAlignmentCustomization.h"
-#include "DetailsView/AvaTextComponentCustomization.h"
-#include "DetailsView/AvaTextFieldCustomization.h"
-#include "Font/AvaFontDetailsCustomization.h"
-#include "Font/AvaFontManagerSubsystem.h"
 #include "Framework/Application/SlateApplication.h"
 #include "IAvalancheComponentVisualizersModule.h"
 #include "IDynamicMaterialEditorModule.h"
@@ -69,20 +64,8 @@ void FAvaTextEditorModule::RegisterCustomLayouts()
 	FPropertyEditorModule& PropertyModule = FModuleManager::GetModuleChecked<FPropertyEditorModule>(PropertyEditor);
 
 	// Text
-	PropertyModule.RegisterCustomPropertyTypeLayout(FAvaFont::StaticStruct()->GetFName()
-		, FOnGetPropertyTypeCustomizationInstance::CreateStatic(&FAvaFontDetailsCustomization::MakeInstance));
-
-	PropertyModule.RegisterCustomPropertyTypeLayout(FAvaTextAlignment::StaticStruct()->GetFName()
-		, FOnGetPropertyTypeCustomizationInstance::CreateStatic(&FAvaTextAlignmentCustomization::MakeInstance));
-
 	PropertyModule.RegisterCustomPropertyTypeLayout(FAvaLinearGradientSettings::StaticStruct()->GetFName()
 		, FOnGetPropertyTypeCustomizationInstance::CreateStatic(&FAvaLinearGradientSettingsCustomization::MakeInstance));
-
-	PropertyModule.RegisterCustomPropertyTypeLayout(FAvaTextField::StaticStruct()->GetFName()
-		, FOnGetPropertyTypeCustomizationInstance::CreateStatic(&FAvaTextFieldCustomization::MakeInstance));
-
-	PropertyModule.RegisterCustomClassLayout(UAvaText3DComponent::StaticClass()->GetFName()
-		, FOnGetDetailCustomizationInstance::CreateStatic(&FAvaTextComponentCustomization::MakeInstance));
 }
 
 void FAvaTextEditorModule::UnregisterCustomLayouts()
@@ -94,10 +77,7 @@ void FAvaTextEditorModule::UnregisterCustomLayouts()
 		FPropertyEditorModule& PropertyModule = FModuleManager::GetModuleChecked<FPropertyEditorModule>(PropertyEditor);
 
 		// Text
-		PropertyModule.UnregisterCustomPropertyTypeLayout(FAvaFont::StaticStruct()->GetFName());
-		PropertyModule.UnregisterCustomPropertyTypeLayout(FAvaTextAlignment::StaticStruct()->GetFName());
 		PropertyModule.UnregisterCustomPropertyTypeLayout(FAvaLinearGradientSettings::StaticStruct()->GetFName());
-		PropertyModule.UnregisterCustomPropertyTypeLayout(FAvaTextField::StaticStruct()->GetFName());
 	}
 }
 
diff --git a/Engine/Plugins/Experimental/Avalanche/Source/AvalancheTextEditor/Private/Visualizer/AvaTextVisualizer.cpp b/Engine/Plugins/Experimental/Avalanche/Source/AvalancheTextEditor/Private/Visualizer/AvaTextVisualizer.cpp
--- a/Engine/Plugins/Experimental/Avalanche/Source/AvalancheTextEditor/Private/Visualizer/AvaTextVisualizer.cpp
+++ b/Engine/Plugins/Experimental/Avalanche/Source/AvalancheTextEditor/Private/Visualizer/AvaTextVisualizer.cpp
@@ -6,6 +6,8 @@
 #include "AvaText3DComponent.h"
 #include "EditorViewportClient.h"
 #include "Engine/Texture2D.h"
+#include "Extensions/Text3DDefaultLayoutExtension.h"
+#include "Extensions/Text3DDefaultMaterialExtension.h"
 #include "IAvalancheComponentVisualizersModule.h"
 #include "LevelEditor/AvaLevelEditorUtils.h"
 #include "Modules/ModuleManager.h"
@@ -45,12 +47,11 @@ FAvaTextVisualizer::FAvaTextVisualizer()
 {
 	using namespace UE::AvaCore;
 
-	HasMaxHeightProperty        = GetProperty<UAvaText3DComponent>(GET_MEMBER_NAME_CHECKED(UAvaText3DComponent, bHasMaxHeight));
-	HasMaxWidthProperty         = GetProperty<UAvaText3DComponent>(GET_MEMBER_NAME_CHECKED(UAvaText3DComponent, bHasMaxWidth));
-	MaxHeightProperty           = GetProperty<UAvaText3DComponent>(GET_MEMBER_NAME_CHECKED(UAvaText3DComponent, MaxHeight));
-	MaxWidthProperty            = GetProperty<UAvaText3DComponent>(GET_MEMBER_NAME_CHECKED(UAvaText3DComponent, MaxWidth));
-	ScaleProportionallyProperty = GetProperty<UAvaText3DComponent>(GET_MEMBER_NAME_CHECKED(UAvaText3DComponent, bScaleProportionally));
-	GradientProperty            = GetProperty<UAvaText3DComponent>(GET_MEMBER_NAME_CHECKED(UAvaText3DComponent, GradientSettings));
+	HasMaxHeightProperty        = GetProperty<UText3DDefaultLayoutExtension>(UText3DDefaultLayoutExtension::GetUseMaxHeightPropertyName());
+	HasMaxWidthProperty         = GetProperty<UText3DDefaultLayoutExtension>(UText3DDefaultLayoutExtension::GetUseMaxWidthPropertyName());
+	MaxHeightProperty           = GetProperty<UText3DDefaultLayoutExtension>(UText3DDefaultLayoutExtension::GetMaxHeightPropertyName());
+	MaxWidthProperty            = GetProperty<UText3DDefaultLayoutExtension>(UText3DDefaultLayoutExtension::GetMaxWidthPropertyName());
+	ScaleProportionallyProperty = GetProperty<UText3DDefaultLayoutExtension>(UText3DDefaultLayoutExtension::GetScaleProportionallyPropertyName());
 
 	ResetEditingFlags();
 	bShowGradientControls = false;
@@ -72,8 +73,7 @@ TMap<UObject*, TArray<FProperty*>> FAvaTextVisualizer::GatherEditableProperties(
 				MaxWidthProperty,
 				HasMaxHeightProperty,
 				MaxHeightProperty,
-				ScaleProportionallyProperty,
-				GradientProperty
+				ScaleProportionallyProperty
 			}
 		}};
 	}
@@ -557,49 +557,17 @@ bool FAvaTextVisualizer::HandleInputDeltaInternal(FEditorViewportClient* InViewp
 					{
 						bHasBeenModified = true;
 					}
-
-					if (AActor* const Text3DActor = AvaTextComponent->GetOwner())
+					
+					if (UText3DDefaultMaterialExtension* MaterialExtension = AvaTextComponent->GetCastedMaterialExtension<UText3DDefaultMaterialExtension>())
 					{
-						Text3DActor->Modify();
-						AvaTextComponent->Modify();
-
-						if (AvaTextComponent->GradientSettings.Direction != EAvaGradientDirection::Custom)
-						{
-							if (AvaTextComponent->GradientSettings.Direction == EAvaGradientDirection::Vertical)
-							{
-								AvaTextComponent->GradientSettings.Rotation = 0;
-							}
-							else if (AvaTextComponent->GradientSettings.Direction == EAvaGradientDirection::Horizontal)
-							{
-								AvaTextComponent->GradientSettings.Rotation = 0.75;
-							}
-
-							InitialGradientRotation = AvaTextComponent->GradientSettings.Rotation;
-							
-							AvaTextComponent->GradientSettings.Direction = EAvaGradientDirection::Custom;
-
-							const FPropertyEditorModule& PropertyModule = FModuleManager::LoadModuleChecked<FPropertyEditorModule>("PropertyEditor");
-
-							for (const FName& DetailViewId : FAvaLevelEditorUtils::GetDetailsViewNames())
-							{
-								if (const TSharedPtr<IDetailsView> DetailView = PropertyModule.FindDetailView(DetailViewId))
-								{
-									DetailView->ForceRefresh();
-								}
-							}
-						}
-						else
-						{
-							AvaTextComponent->GradientSettings.Rotation = GradientRot;
-						}
-
-						AvaTextComponent->RefreshGradient();
+						MaterialExtension->Modify();
+						MaterialExtension->SetGradientRotation(GradientRot);
+						return true;
 					}
-					return true;
 				}
 			}
 		}
-		
+
 		EndEditing();
 		return true;
 	}
@@ -612,16 +580,18 @@ bool FAvaTextVisualizer::HandleInputDeltaInternal(FEditorViewportClient* InViewp
 			{
 				if (GetViewportWidgetAxisList(InViewportClient) & EAxisList::YZ)
 				{
-					if (AActor* const Text3DActor = AvaTextComponent->GetOwner())
+					const AActor* Text3DActor = AvaTextComponent->GetOwner();
+					UText3DDefaultMaterialExtension* MaterialExtension = AvaTextComponent->GetCastedMaterialExtension<UText3DDefaultMaterialExtension>();
+
+					if (Text3DActor && MaterialExtension)
 					{
-						Text3DActor->Modify();
-						AvaTextComponent->Modify();
+						MaterialExtension->Modify();
 
 						float GradientOffset = InitialGradientOffset;
 						
 						const float OffsetIncr = InAccumulatedTranslation.Length() * FTextVisualizerStatics::GradientOffsetHandleScale;
 
-						FVector GradientDir = Text3DActor->GetActorUpVector().RotateAngleAxis(-AvaTextComponent->GradientSettings.Rotation * 360, Text3DActor->GetActorForwardVector());
+						FVector GradientDir = Text3DActor->GetActorUpVector().RotateAngleAxis(-MaterialExtension->GetGradientRotation() * 360, Text3DActor->GetActorForwardVector());
 						GradientDir.Normalize();
 
 						const float DeltaDotGradDir = InAccumulatedTranslation.Dot(GradientDir);
@@ -634,8 +604,7 @@ bool FAvaTextVisualizer::HandleInputDeltaInternal(FEditorViewportClient* InViewp
 							bHasBeenModified = true;
 						}
 						
-						AvaTextComponent->GradientSettings.Offset = GradientOffset;
-						AvaTextComponent->RefreshGradient();
+						MaterialExtension->SetGradientOffset(GradientOffset);
 						return true;
 					}
 				}
@@ -654,15 +623,13 @@ bool FAvaTextVisualizer::HandleInputDeltaInternal(FEditorViewportClient* InViewp
 			{
 				if (GetViewportWidgetAxisList(InViewportClient) & EAxisList::Z)
 				{
-					if (AActor* const Text3DActor = AvaTextComponent->GetOwner())
+					if (UText3DDefaultMaterialExtension* MaterialExtension = AvaTextComponent->GetCastedMaterialExtension<UText3DDefaultMaterialExtension>())
 					{
-						Text3DActor->Modify();
-						AvaTextComponent->Modify();
+						MaterialExtension->Modify();
 
 						float GradientSmoothness = InitialGradientSmoothness;
 						
 						const float SmoothnessIncr = InAccumulatedScale.Z * FTextVisualizerStatics::GradientSmoothnessHandleScale;
-						
 						GradientSmoothness += SmoothnessIncr;
 						GradientSmoothness = FMath::Clamp(GradientSmoothness, 0.0, 1.0);
 
@@ -671,8 +638,7 @@ bool FAvaTextVisualizer::HandleInputDeltaInternal(FEditorViewportClient* InViewp
 							bHasBeenModified = true;
 						}
 
-						AvaTextComponent->GradientSettings.Smoothness = GradientSmoothness;
-						AvaTextComponent->RefreshGradient();
+						MaterialExtension->SetGradientSmoothness(GradientSmoothness);
 						return true;
 					}
 				}
@@ -727,7 +693,7 @@ bool FAvaTextVisualizer::ResetValue(FEditorViewportClient* InViewportClient, HHi
 		AvaTextComponent->Modify();
 
 		AvaTextComponent->SetHasMaxWidth(false);
-		AvaTextComponent->SetMaxWidth(GetDefault<UAvaText3DComponent>()->MaxWidth);
+		AvaTextComponent->SetMaxWidth(GetDefault<UText3DDefaultLayoutExtension>()->GetMaxWidth());
 
 		FAvaVisualizerBase::NotifyPropertyModified(AvaTextComponent, MaxWidthProperty, EPropertyChangeType::ValueSet);
 	}
@@ -737,7 +703,7 @@ bool FAvaTextVisualizer::ResetValue(FEditorViewportClient* InViewportClient, HHi
 		AvaTextComponent->Modify();
 
 		AvaTextComponent->SetHasMaxHeight(false);
-		AvaTextComponent->SetMaxHeight(GetDefault<UAvaText3DComponent>()->MaxHeight);
+		AvaTextComponent->SetMaxHeight(GetDefault<UText3DDefaultLayoutExtension>()->GetMaxHeight());
 
 		FAvaVisualizerBase::NotifyPropertyModified(AvaTextComponent, MaxHeightProperty, EPropertyChangeType::ValueSet);
 	}
@@ -846,11 +812,11 @@ void FAvaTextVisualizer::GetTextActorGradientControlsLocations(const UAvaText3DC
 	using namespace UE::Ava::Private;
 	if (IsValid(InAvaTextComponent))
 	{
-		if (const AActor* const Text3DActor = InAvaTextComponent->GetOwner())
-		{
-			const AActor* TextAsActor = Cast<AActor>(Text3DActor);
-			check(TextAsActor);
+		const AActor* const Text3DActor = InAvaTextComponent->GetOwner();
+		const UText3DDefaultMaterialExtension* MaterialExtension = InAvaTextComponent->GetCastedMaterialExtension<UText3DDefaultMaterialExtension>();
 
+		if (Text3DActor && MaterialExtension)
+		{
 			FVector BoundsOrigin;
 			FVector BoundsExtent;
 			
@@ -858,12 +824,12 @@ void FAvaTextVisualizer::GetTextActorGradientControlsLocations(const UAvaText3DC
 			const float GradientLineHalfLength = FMath::Min(BoundsExtent.Length(), FTextVisualizerStatics::GradientHandleMaxLength);
 			FVector GradientCenterLocation = BoundsOrigin;
 
-			const FVector GradientDir = InAvaTextComponent->GetGradientDirection();
+			const FVector GradientDir = MaterialExtension->GetGradientDirection();
 			OutGradientStartLocation = GradientCenterLocation - GradientDir * GradientLineHalfLength;
 			OutGradientEndLocation = GradientCenterLocation + GradientDir * GradientLineHalfLength;	
 
 			// we need to move gradient center along its direction, base on offset amount
-			const float GradientAsNormalizedPos = InAvaTextComponent->GradientSettings.Offset * 2.0 - 1.0;
+			const float GradientAsNormalizedPos = MaterialExtension->GetGradientOffset() * 2.0 - 1.0;
 
 			GradientCenterLocation -= GradientDir * GradientAsNormalizedPos * GradientLineHalfLength;
 			OutGradientCenterLocation = GradientCenterLocation;
@@ -963,13 +929,16 @@ void FAvaTextVisualizer::DrawVisualizationNotEditing(const UActorComponent* InCo
 		DrawMaxTextHeightHandle(AvaTextComponent, InView, InPDI, Inactive);
 	}
 	
-	if (AvaTextComponent->GetColoringStyle() == EAvaTextColoringStyle::Gradient)
+	if (const UText3DDefaultMaterialExtension* MaterialExtension = AvaTextComponent->GetCastedMaterialExtension<UText3DDefaultMaterialExtension>())
 	{
-		DrawEditGradientButton(AvaTextComponent, InView, InPDI, InOutIconIndex, bShowGradientControls ? Enabled : Disabled);
-
-		if (bShowGradientControls)
+		if (MaterialExtension->GetStyle() == EText3DMaterialStyle::Gradient)
 		{
-			DrawGradientHandles(AvaTextComponent, InView, InPDI);
+			DrawEditGradientButton(AvaTextComponent, InView, InPDI, InOutIconIndex, bShowGradientControls ? Enabled : Disabled);
+
+			if (bShowGradientControls)
+			{
+				DrawGradientHandles(AvaTextComponent, InView, InPDI);
+			}
 		}
 	}
 }
@@ -1001,15 +970,18 @@ void FAvaTextVisualizer::DrawVisualizationEditing(const UActorComponent* InCompo
 
 	DrawMaxTextWidthHandle(AvaTextComponent, InView, InPDI, Inactive);
 	DrawMaxTextHeightHandle(AvaTextComponent, InView, InPDI, Inactive);
-	
-	if (AvaTextComponent->GetColoringStyle() == EAvaTextColoringStyle::Gradient)
-	{
-		DrawEditGradientButton(AvaTextComponent, InView, InPDI, InOutIconIndex, bShowGradientControls ? Enabled : Disabled);
-		++InOutIconIndex;
 
-		if (bShowGradientControls)
+	if (const UText3DDefaultMaterialExtension* MaterialExtension = TextComponent->GetCastedMaterialExtension<UText3DDefaultMaterialExtension>())
+	{
+		if (MaterialExtension->GetStyle() == EText3DMaterialStyle::Gradient)
 		{
-			DrawGradientHandles(AvaTextComponent, InView, InPDI);
+			DrawEditGradientButton(AvaTextComponent, InView, InPDI, InOutIconIndex, bShowGradientControls ? Enabled : Disabled);
+			++InOutIconIndex;
+
+			if (bShowGradientControls)
+			{
+				DrawGradientHandles(AvaTextComponent, InView, InPDI);
+			}
 		}
 	}
 }
@@ -1028,10 +1000,13 @@ void FAvaTextVisualizer::StoreInitialValues()
 	InitialMaxWidth = TextComponent->GetMaxWidth();
 	InitialMaxHeight = TextComponent->GetMaxHeight();
 	bInitialScaleProportionally = TextComponent->ScalesProportionally();
-	
-	InitialGradientRotation = TextComponent->GetGradientSettings().Rotation;
-	InitialGradientOffset = TextComponent->GetGradientSettings().Offset;
-	InitialGradientSmoothness = TextComponent->GetGradientSettings().Smoothness;
+
+	if (UText3DDefaultMaterialExtension* MaterialExtension = TextComponent->GetCastedMaterialExtension<UText3DDefaultMaterialExtension>())
+	{
+		InitialGradientRotation = MaterialExtension->GetGradientRotation();
+		InitialGradientOffset = MaterialExtension->GetGradientOffset();
+		InitialGradientSmoothness = MaterialExtension->GetGradientSmoothness();
+	}
 }
 
 void FAvaTextVisualizer::DrawMaxTextSizeVisualization(const UText3DComponent* InTextComponent, const FSceneView* InView, 
@@ -1112,9 +1087,13 @@ void FAvaTextVisualizer::DrawGradientHandles(const UAvaText3DComponent* InAvaTex
 	GetTextActorGradientControlsLocations(InAvaTextComponent, GradientCenter, GradientLineStart, GradientLineEnd);
 	
 	InPDI->DrawLine(GradientLineStart, GradientLineEnd, FLinearColor::White,SDPG_Foreground );
-	
-	DrawGradientLineStartHandle(InAvaTextComponent, InView, InPDI, GradientLineStart, InAvaTextComponent->GetGradientSettings().ColorB);
-	DrawGradientLineEndHandle(InAvaTextComponent, InView, InPDI, GradientLineEnd, InAvaTextComponent->GetGradientSettings().ColorA);
+
+	if (UText3DDefaultMaterialExtension* MaterialExtension = TextComponent->GetCastedMaterialExtension<UText3DDefaultMaterialExtension>())
+	{
+		DrawGradientLineStartHandle(InAvaTextComponent, InView, InPDI, GradientLineStart, MaterialExtension->GetGradientColorB());
+		DrawGradientLineEndHandle(InAvaTextComponent, InView, InPDI, GradientLineEnd, MaterialExtension->GetGradientColorA());
+	}
+
 	DrawGradientCenterHandle(InAvaTextComponent, InView, InPDI, GradientCenter);
 	DrawGradientSmoothnessHandle(InAvaTextComponent, InView, InPDI, GradientCenter + FVector(0, FTextVisualizerStatics::GradientSmoothnessHandleOffset, 0));
 	
diff --git a/Engine/Plugins/Experimental/Avalanche/Source/AvalancheTextEditor/Private/Visualizer/AvaTextVisualizer.h b/Engine/Plugins/Experimental/Avalanche/Source/AvalancheTextEditor/Private/Visualizer/AvaTextVisualizer.h
--- a/Engine/Plugins/Experimental/Avalanche/Source/AvalancheTextEditor/Private/Visualizer/AvaTextVisualizer.h
+++ b/Engine/Plugins/Experimental/Avalanche/Source/AvalancheTextEditor/Private/Visualizer/AvaTextVisualizer.h
@@ -133,7 +133,6 @@ protected:
 	FProperty* MaxWidthProperty;
 	FProperty* MaxHeightProperty;
 	FProperty* ScaleProportionallyProperty;
-	FProperty* GradientProperty;
 
 	TWeakObjectPtr<UAvaText3DComponent> TextComponent;
 	mutable FBox Bounds;
diff --git a/Engine/Plugins/Experimental/PropertyAnimator/Source/PropertyAnimator/Private/Presets/PropertyAnimatorPresetTextLocation.cpp b/Engine/Plugins/Experimental/PropertyAnimator/Source/PropertyAnimator/Private/Presets/PropertyAnimatorPresetTextLocation.cpp
--- a/Engine/Plugins/Experimental/PropertyAnimator/Source/PropertyAnimator/Private/Presets/PropertyAnimatorPresetTextLocation.cpp
+++ b/Engine/Plugins/Experimental/PropertyAnimator/Source/PropertyAnimator/Private/Presets/PropertyAnimatorPresetTextLocation.cpp
@@ -15,7 +15,7 @@ void UPropertyAnimatorPresetTextLocation::GetPresetProperties(const AActor* InAc
 		return;
 	}
 
-	USceneComponent* TextRootComponent = TextComponent->GetChildComponent(1);
+	USceneComponent* TextRootComponent = TextComponent->GetChildComponent(0);
 	if (!TextRootComponent)
 	{
 		return;
diff --git a/Engine/Plugins/Experimental/PropertyAnimator/Source/PropertyAnimator/Private/Presets/PropertyAnimatorPresetTextRotation.cpp b/Engine/Plugins/Experimental/PropertyAnimator/Source/PropertyAnimator/Private/Presets/PropertyAnimatorPresetTextRotation.cpp
--- a/Engine/Plugins/Experimental/PropertyAnimator/Source/PropertyAnimator/Private/Presets/PropertyAnimatorPresetTextRotation.cpp
+++ b/Engine/Plugins/Experimental/PropertyAnimator/Source/PropertyAnimator/Private/Presets/PropertyAnimatorPresetTextRotation.cpp
@@ -15,7 +15,7 @@ void UPropertyAnimatorPresetTextRotation::GetPresetProperties(const AActor* InAc
 		return;
 	}
 
-	USceneComponent* TextRootComponent = TextComponent->GetChildComponent(1);
+	USceneComponent* TextRootComponent = TextComponent->GetChildComponent(0);
 	if (!TextRootComponent)
 	{
 		return;
diff --git a/Engine/Plugins/Experimental/Text3D/Source/Text3D/Private/Extensions/Text3DDefaultLayoutExtension.cpp b/Engine/Plugins/Experimental/Text3D/Source/Text3D/Private/Extensions/Text3DDefaultLayoutExtension.cpp
new file mode 100644
--- /dev/null
+++ b/Engine/Plugins/Experimental/Text3D/Source/Text3D/Private/Extensions/Text3DDefaultLayoutExtension.cpp
@@ -0,0 +1,410 @@
+// Copyright Epic Games, Inc. All Rights Reserved.
+
+#include "Extensions/Text3DDefaultLayoutExtension.h"
+
+#include "Engine/Font.h"
+#include "Fonts/SlateTextShaper.h"
+#include "Framework/Text/PlainTextLayoutMarshaller.h"
+#include "GeometryBuilders/Text3DGlyphMeshBuilder.h"
+#include "GeometryBuilders/Text3DGlyph.h"
+#include "LayoutBuilders/Text3DLayout.h"
+#include "LayoutBuilders/Text3DLayoutShaper.h"
+#include "LayoutBuilders/Text3DShapedGlyphText.h"
+#include "Logs/Text3DLogs.h"
+#include "Styling/StyleDefaults.h"
+#include "Subsystems/Text3DEngineSubsystem.h"
+#include "Text3DComponent.h"
+
+FName UText3DDefaultLayoutExtension::GetUseMaxWidthPropertyName()
+{
+	return GET_MEMBER_NAME_CHECKED(UText3DDefaultLayoutExtension, bUseMaxWidth);
+}
+
+FName UText3DDefaultLayoutExtension::GetUseMaxHeightPropertyName()
+{
+	return GET_MEMBER_NAME_CHECKED(UText3DDefaultLayoutExtension, bUseMaxHeight);
+}
+
+FName UText3DDefaultLayoutExtension::GetMaxHeightPropertyName()
+{
+	return GET_MEMBER_NAME_CHECKED(UText3DDefaultLayoutExtension, MaxHeight);
+}
+
+FName UText3DDefaultLayoutExtension::GetMaxWidthPropertyName()
+{
+	return GET_MEMBER_NAME_CHECKED(UText3DDefaultLayoutExtension, MaxWidth);
+}
+
+FName UText3DDefaultLayoutExtension::GetScaleProportionallyPropertyName()
+{
+	return GET_MEMBER_NAME_CHECKED(UText3DDefaultLayoutExtension, bScaleProportionally);
+}
+
+void UText3DDefaultLayoutExtension::SetKerning(const float Value)
+{
+	if (FMath::IsNearlyEqual(Kerning, Value))
+	{
+		return;
+	}
+	
+	Kerning = Value;
+	OnLayoutOptionsChanged();
+}
+
+void UText3DDefaultLayoutExtension::SetLineSpacing(const float Value)
+{
+	if (FMath::IsNearlyEqual(LineSpacing, Value))
+	{
+		return;
+	}
+	
+	LineSpacing = Value;
+	OnLayoutOptionsChanged();
+}
+
+void UText3DDefaultLayoutExtension::SetWordSpacing(const float Value)
+{
+	if (FMath::IsNearlyEqual(WordSpacing, Value))
+	{
+		return;
+	}
+
+	WordSpacing = Value;
+	OnLayoutOptionsChanged();
+}
+
+void UText3DDefaultLayoutExtension::SetHorizontalAlignment(const EText3DHorizontalTextAlignment Value)
+{
+	if (HorizontalAlignment == Value)
+	{
+		return;
+	}
+
+	HorizontalAlignment = Value;
+	OnLayoutOptionsChanged();
+}
+
+void UText3DDefaultLayoutExtension::SetVerticalAlignment(const EText3DVerticalTextAlignment Value)
+{
+	if (VerticalAlignment == Value)
+	{
+		return;
+	}
+	
+	VerticalAlignment = Value;
+	OnLayoutOptionsChanged();
+}
+
+void UText3DDefaultLayoutExtension::SetUseMaxWidth(const bool Value)
+{
+	if (bUseMaxWidth == Value)
+	{
+		return;
+	}
+
+	bUseMaxWidth = Value;
+	OnLayoutOptionsChanged();
+}
+
+void UText3DDefaultLayoutExtension::SetMaxWidth(const float Value)
+{
+	const float NewValue = FMath::Max(1.0f, Value);
+	if (FMath::IsNearlyEqual(MaxWidth, NewValue))
+	{
+		return;
+	}
+	
+	MaxWidth = NewValue;
+	OnLayoutOptionsChanged();
+}
+
+void UText3DDefaultLayoutExtension::SetMaxWidthBehavior(const EText3DMaxWidthHandling Value)
+{
+	if (MaxWidthBehavior == Value)
+	{
+		return;
+	}
+
+	MaxWidthBehavior = Value;
+	OnLayoutOptionsChanged();
+}
+
+void UText3DDefaultLayoutExtension::SetUseMaxHeight(const bool Value)
+{
+	if (bUseMaxHeight == Value)
+	{
+		return;
+	}
+	
+	bUseMaxHeight = Value;
+	OnLayoutOptionsChanged();
+}
+
+void UText3DDefaultLayoutExtension::SetMaxHeight(const float Value)
+{
+	const float NewValue = FMath::Max(1.0f, Value);
+	if (FMath::IsNearlyEqual(MaxHeight, NewValue))
+	{
+		return;
+	}
+
+	MaxHeight = NewValue;
+	OnLayoutOptionsChanged();
+}
+
+void UText3DDefaultLayoutExtension::SetScaleProportionally(const bool Value)
+{
+	if (bScaleProportionally == Value)
+	{
+		return;
+	}
+
+	bScaleProportionally = Value;
+	OnLayoutOptionsChanged();
+}
+
+UText3DDefaultLayoutExtension::UText3DDefaultLayoutExtension()
+{
+	ShapedText = MakeShared<FText3DShapedGlyphText>();
+}
+
+float UText3DDefaultLayoutExtension::GetTextHeight() const
+{
+	return ShapedText->Lines.Num() * ShapedText->LineHeight + (ShapedText->Lines.Num() - 1) * LineSpacing;
+}
+
+FVector UText3DDefaultLayoutExtension::GetTextScale() const
+{
+	return TextScale;
+}
+
+EText3DExtensionResult UText3DDefaultLayoutExtension::PreRendererUpdate(EText3DRendererFlags InFlag)
+{
+	if (InFlag != EText3DRendererFlags::Geometry && InFlag != EText3DRendererFlags::Layout)
+	{
+		return EText3DExtensionResult::Active;
+	}
+
+	const UText3DComponent* Text3DComponent = GetText3DComponent();
+	
+	if (!EnumHasAnyFlags(Text3DComponent->GetUpdateFlags(), EText3DRendererFlags::Geometry))
+	{
+		ShapedText->Kerning = Kerning;
+		ShapedText->WordSpacing = WordSpacing;
+		ShapedText->MaxWidth = MaxWidth;
+		ShapedText->bWrap = bUseMaxWidth && MaxWidthBehavior == EText3DMaxWidthHandling::WrapAndScale;
+	}
+	else
+	{
+		UText3DEngineSubsystem* Subsystem = UText3DEngineSubsystem::Get();
+
+		const uint32 TypefaceIndex = Text3DComponent->GetTypeFaceIndex();
+		FCachedFontData& CachedFontData = Subsystem->GetCachedFontData(Text3DComponent->GetFont(), TypefaceIndex);
+		const FT_Face Face = CachedFontData.GetFreeTypeFace(TypefaceIndex);
+
+		if (!Face)
+		{
+			UE_LOG(LogText3D, Error, TEXT("Failed to load font face data '%s' in Text3D layout extension"), *CachedFontData.GetFontName());
+			return EText3DExtensionResult::Failed;
+		}
+
+		ShapedText->Reset();
+		ShapedText->LineHeight = Face->size->metrics.height * FontInverseScale;
+		ShapedText->FontAscender = Face->size->metrics.ascender * FontInverseScale;
+		ShapedText->FontDescender = Face->size->metrics.descender * FontInverseScale;
+		ShapedText->Kerning = Kerning;
+		ShapedText->WordSpacing = WordSpacing;
+		ShapedText->MaxWidth = MaxWidth;
+		ShapedText->bWrap = bUseMaxWidth && MaxWidthBehavior == EText3DMaxWidthHandling::WrapAndScale;
+
+		constexpr int32 AdjustedFontSize = 48; // Magic number that makes font scale consistent with previous implementation
+		FSlateFontInfo FontInfo(Text3DComponent->GetFont(), AdjustedFontSize);
+		FontInfo.CompositeFont = FStyleDefaults::GetFontInfo().CompositeFont;
+		FontInfo.TypefaceFontName = Text3DComponent->GetTypeface();
+
+		FTextBlockStyle Style;
+		Style.SetFont(FontInfo);
+
+		if (!TextLayout.IsValid())
+		{
+			TextLayout = MakeShared<FText3DLayout>(Style);
+		}
+
+		if (!TextLayoutMarshaller.IsValid())
+		{
+			TextLayoutMarshaller = FPlainTextLayoutMarshaller::Create();
+		}
+
+		FText3DLayoutShaper::Get()->ShapeBidirectionalText(Style, Text3DComponent->GetFormattedText().ToString(), TextLayout, TextLayoutMarshaller, ShapedText->Lines);
+
+		const int32 ApproximateGlyphNum = Algo::TransformAccumulate(ShapedText->Lines,
+			[](const FShapedGlyphLine& InLine)
+			{
+				return InLine.GlyphsToRender.Num();
+			},
+			0);
+
+		ShapedText->GlyphIndexToFontFace.Reserve(ApproximateGlyphNum);
+
+		for (int32 LineIndex = 0; LineIndex < ShapedText->Lines.Num(); ++LineIndex)
+		{
+			for (const FShapedGlyphEntry& GlyphEntry : ShapedText->Lines[LineIndex].GlyphsToRender)
+			{
+				if (!GlyphEntry.FontFaceData.IsValid())
+				{
+					// Add as nullptr if not already in map
+					if (!ShapedText->GlyphIndexToFontFace.Contains(GlyphEntry.GlyphIndex))
+					{
+						ShapedText->GlyphIndexToFontFace.Add(GlyphEntry.GlyphIndex, nullptr);
+					}
+
+					continue;
+				}
+
+				const TSharedPtr<FFreeTypeFace> FontFacePtr = GlyphEntry.FontFaceData->FontFace.Pin();
+
+				if (FontFacePtr.IsValid())
+				{
+					ShapedText->GlyphIndexToFontFace.FindOrAdd(GlyphEntry.GlyphIndex, FontFacePtr.Get());
+				}
+			}
+		}
+	}
+
+	ShapedText->CalculateWidth();
+	CalculateTextScale();
+
+	return EText3DExtensionResult::Finished;
+}
+
+EText3DExtensionResult UText3DDefaultLayoutExtension::PostRendererUpdate(EText3DRendererFlags InFlag)
+{
+	return EText3DExtensionResult::Active;
+}
+
+void UText3DDefaultLayoutExtension::CalculateTextScale()
+{
+	FVector Scale(1.0f, 1.0f, 1.0f);
+
+	float TextMaxWidth = 0.0f;
+	for (const FShapedGlyphLine& ShapedLine : ShapedText->Lines)
+	{
+		TextMaxWidth = FMath::Max(TextMaxWidth, ShapedLine.Width);
+	}
+
+	if (bUseMaxWidth && TextMaxWidth > MaxWidth && TextMaxWidth > 0.0f)
+	{
+		Scale.Y *= MaxWidth / TextMaxWidth;
+		if (bScaleProportionally)
+		{
+			Scale.Z = Scale.Y;
+		}
+	}
+
+	const float TotalHeight = GetTextHeight();
+	if (bUseMaxHeight && TotalHeight > MaxHeight && TotalHeight > 0.0f)
+	{
+		Scale.Z *= MaxHeight / TotalHeight;
+		if (bScaleProportionally)
+		{
+			Scale.Y = Scale.Z;
+		}
+	}
+
+	if (bScaleProportionally)
+	{
+		Scale.X = Scale.Y;
+	}
+
+	TextScale = Scale;
+}
+
+FVector UText3DDefaultLayoutExtension::GetLineLocation(int32 LineIndex) const
+{
+	float HorizontalOffset = 0.0f;
+	float VerticalOffset = 0.0f;
+	if (LineIndex < 0 || LineIndex >= ShapedText->Lines.Num())
+	{
+		return FVector();
+	}
+
+	const FShapedGlyphLine& ShapedLine = ShapedText->Lines[LineIndex];
+
+	if (HorizontalAlignment == EText3DHorizontalTextAlignment::Center)
+	{
+		HorizontalOffset = -ShapedLine.Width * 0.5f;
+	}
+	else if (HorizontalAlignment == EText3DHorizontalTextAlignment::Right)
+	{
+		HorizontalOffset = -ShapedLine.Width;
+	}
+
+	const float TotalHeight = GetTextHeight();
+	if (VerticalAlignment != EText3DVerticalTextAlignment::FirstLine)
+	{
+		// First align it to Top
+		VerticalOffset -= ShapedText->FontAscender;
+
+		if (VerticalAlignment == EText3DVerticalTextAlignment::Center)
+		{
+			VerticalOffset += TotalHeight * 0.5f;
+		}
+		else if (VerticalAlignment == EText3DVerticalTextAlignment::Bottom)
+		{
+			VerticalOffset += TotalHeight + ShapedText->FontDescender;
+		}
+	}
+
+	VerticalOffset -= LineIndex * (ShapedText->LineHeight + LineSpacing);
+
+	return FVector(0.0f, HorizontalOffset, VerticalOffset);
+}
+
+TSharedPtr<FText3DShapedGlyphText> UText3DDefaultLayoutExtension::GetShapedText() const
+{
+	return ShapedText;
+}
+
+#if WITH_EDITOR
+void UText3DDefaultLayoutExtension::PostEditChangeProperty(FPropertyChangedEvent& InEvent)
+{
+	Super::PostEditChangeProperty(InEvent);
+
+	static const TSet<FName> LayoutPropertyNames
+	{
+		GET_MEMBER_NAME_CHECKED(UText3DDefaultLayoutExtension, Kerning),
+		GET_MEMBER_NAME_CHECKED(UText3DDefaultLayoutExtension, LineSpacing),
+		GET_MEMBER_NAME_CHECKED(UText3DDefaultLayoutExtension, WordSpacing),
+		GET_MEMBER_NAME_CHECKED(UText3DDefaultLayoutExtension, HorizontalAlignment),
+		GET_MEMBER_NAME_CHECKED(UText3DDefaultLayoutExtension, VerticalAlignment),
+		GET_MEMBER_NAME_CHECKED(UText3DDefaultLayoutExtension, bUseMaxWidth),
+		GET_MEMBER_NAME_CHECKED(UText3DDefaultLayoutExtension, MaxWidth),
+		GET_MEMBER_NAME_CHECKED(UText3DDefaultLayoutExtension, bUseMaxHeight),
+		GET_MEMBER_NAME_CHECKED(UText3DDefaultLayoutExtension, MaxHeight),
+		GET_MEMBER_NAME_CHECKED(UText3DDefaultLayoutExtension, bScaleProportionally),
+	};
+
+	if (LayoutPropertyNames.Contains(InEvent.GetMemberPropertyName()))
+	{
+		OnLayoutOptionsChanged();
+	}
+}
+#endif
+
+void UText3DDefaultLayoutExtension::OnLayoutOptionsChanged()
+{
+	MaxWidth = FMath::Max(1.0f, MaxWidth);
+	MaxHeight = FMath::Max(1.0f, MaxHeight);
+
+	EText3DRendererFlags Flags = EText3DRendererFlags::Material;
+	if (MaxWidthBehavior == EText3DMaxWidthHandling::WrapAndScale)
+	{
+		EnumAddFlags(Flags, EText3DRendererFlags::Geometry);
+	}
+	else
+	{
+		EnumAddFlags(Flags, EText3DRendererFlags::Layout);
+	}
+
+	RequestUpdate(Flags);
+}
diff --git a/Engine/Plugins/Experimental/Text3D/Source/Text3D/Private/Extensions/Text3DDefaultTransformExtension.cpp b/Engine/Plugins/Experimental/Text3D/Source/Text3D/Private/Extensions/Text3DDefaultTransformExtension.cpp
new file mode 100644
--- /dev/null
+++ b/Engine/Plugins/Experimental/Text3D/Source/Text3D/Private/Extensions/Text3DDefaultTransformExtension.cpp
@@ -0,0 +1,294 @@
+// Copyright Epic Games, Inc. All Rights Reserved.
+
+#include "Extensions/Text3DDefaultTransformExtension.h"
+
+void UText3DDefaultTransformExtension::SetLocationEnabled(bool bEnabled)
+{
+	if (bLocationEnabled != bEnabled)
+	{
+		bLocationEnabled = bEnabled;
+		OnTransformOptionsChanged();
+	}
+}
+
+void UText3DDefaultTransformExtension::SetLocationProgress(float Progress)
+{
+	Progress = FMath::Clamp(Progress, 0.0f, 100.0f);
+	if (!FMath::IsNearlyEqual(LocationProgress, Progress))
+	{
+		LocationProgress = Progress;
+		OnTransformOptionsChanged();
+	}
+}
+
+void UText3DDefaultTransformExtension::SetLocationOrder(EText3DCharacterEffectOrder Order)
+{
+	if (LocationOrder != Order)
+	{
+		LocationOrder = Order;
+		OnTransformOptionsChanged();
+	}
+}
+
+void UText3DDefaultTransformExtension::SetLocationRange(float Range)
+{
+	Range = FMath::Clamp(Range, 0.0f, 100.0f);
+	if (!FMath::IsNearlyEqual(LocationRange, Range))
+	{
+		LocationRange = Range;
+		OnTransformOptionsChanged();
+	}
+}
+
+void UText3DDefaultTransformExtension::SetLocationDistance(const FVector& Distance)
+{
+	if (!(LocationDistance - Distance).IsNearlyZero())
+	{
+		LocationDistance = Distance;
+		OnTransformOptionsChanged();
+	}
+}
+
+void UText3DDefaultTransformExtension::SetScaleEnabled(bool bEnabled)
+{
+	if (bScaleEnabled != bEnabled)
+	{
+		bScaleEnabled = bEnabled;
+		OnTransformOptionsChanged();
+	}
+}
+
+void UText3DDefaultTransformExtension::SetScaleProgress(float Progress)
+{
+	Progress = FMath::Clamp(Progress, 0.0f, 100.0f);
+	if (!FMath::IsNearlyEqual(ScaleProgress, Progress))
+	{
+		ScaleProgress = Progress;
+		OnTransformOptionsChanged();
+	}
+}
+
+void UText3DDefaultTransformExtension::SetScaleOrder(EText3DCharacterEffectOrder Order)
+{
+	if (ScaleOrder != Order)
+	{
+		ScaleOrder = Order;
+		OnTransformOptionsChanged();
+	}
+}
+
+void UText3DDefaultTransformExtension::SetScaleRange(float Range)
+{
+	Range = FMath::Clamp(Range, 0.0f, 100.0f);
+	if (!FMath::IsNearlyEqual(ScaleRange, Range))
+	{
+		ScaleRange = Range;
+		OnTransformOptionsChanged();
+	}
+}
+
+void UText3DDefaultTransformExtension::SetScaleBegin(const FVector& Value)
+{
+	if (!(ScaleBegin - Value).IsNearlyZero())
+	{
+		ScaleBegin = Value;
+		OnTransformOptionsChanged();
+	}
+}
+
+void UText3DDefaultTransformExtension::SetScaleEnd(const FVector& Value)
+{
+	if (!(ScaleEnd - Value).IsNearlyZero())
+	{
+		ScaleEnd = Value;
+		OnTransformOptionsChanged();
+	}
+}
+
+void UText3DDefaultTransformExtension::SetRotationEnabled(bool bEnabled)
+{
+	if (bRotationEnabled != bEnabled)
+	{
+		bRotationEnabled = bEnabled;
+		OnTransformOptionsChanged();
+	}
+}
+
+void UText3DDefaultTransformExtension::SetRotationProgress(float Progress)
+{
+	Progress = FMath::Clamp(Progress, 0.0f, 100.0f);
+	if (!FMath::IsNearlyEqual(RotationProgress, Progress))
+	{
+		RotationProgress = Progress;
+		OnTransformOptionsChanged();
+	}
+}
+
+void UText3DDefaultTransformExtension::SetRotationOrder(EText3DCharacterEffectOrder Order)
+{
+	if (RotationOrder != Order)
+	{
+		RotationOrder = Order;
+		OnTransformOptionsChanged();
+	}
+}
+
+void UText3DDefaultTransformExtension::SetRotationRange(float Range)
+{
+	Range = FMath::Clamp(Range, 0.0f, 100.0f);
+	if (!FMath::IsNearlyEqual(RotationRange, Range))
+	{
+		RotationRange = Range;
+		OnTransformOptionsChanged();
+	}
+}
+
+void UText3DDefaultTransformExtension::SetRotationBegin(const FRotator& Value)
+{
+	if (!(RotationBegin - Value).IsNearlyZero())
+	{
+		RotationBegin = Value;
+		OnTransformOptionsChanged();
+	}
+}
+
+void UText3DDefaultTransformExtension::SetRotationEnd(const FRotator& Value)
+{
+	if (!(RotationEnd - Value).IsNearlyZero())
+	{
+		RotationEnd = Value;
+		OnTransformOptionsChanged();
+	}
+}
+
+UText3DDefaultTransformExtension::UText3DDefaultTransformExtension()
+{
+	OnTransformOptionsChanged();
+}
+
+#if WITH_EDITOR
+void UText3DDefaultTransformExtension::PostEditChangeProperty(FPropertyChangedEvent& InEvent)
+{
+	Super::PostEditChangeProperty(InEvent);
+
+	static const TSet<FName> PropertyNames =
+	{
+		GET_MEMBER_NAME_CHECKED(UText3DDefaultTransformExtension, bLocationEnabled),
+		GET_MEMBER_NAME_CHECKED(UText3DDefaultTransformExtension, LocationProgress),
+		GET_MEMBER_NAME_CHECKED(UText3DDefaultTransformExtension, LocationOrder),
+		GET_MEMBER_NAME_CHECKED(UText3DDefaultTransformExtension, LocationRange),
+		GET_MEMBER_NAME_CHECKED(UText3DDefaultTransformExtension, LocationDistance),
+		GET_MEMBER_NAME_CHECKED(UText3DDefaultTransformExtension, bRotationEnabled),
+		GET_MEMBER_NAME_CHECKED(UText3DDefaultTransformExtension, RotationProgress),
+		GET_MEMBER_NAME_CHECKED(UText3DDefaultTransformExtension, RotationOrder),
+		GET_MEMBER_NAME_CHECKED(UText3DDefaultTransformExtension, RotationRange),
+		GET_MEMBER_NAME_CHECKED(UText3DDefaultTransformExtension, RotationBegin),
+		GET_MEMBER_NAME_CHECKED(UText3DDefaultTransformExtension, RotationEnd),
+		GET_MEMBER_NAME_CHECKED(UText3DDefaultTransformExtension, RotationEnd),
+		GET_MEMBER_NAME_CHECKED(UText3DDefaultTransformExtension, bScaleEnabled),
+		GET_MEMBER_NAME_CHECKED(UText3DDefaultTransformExtension, ScaleProgress),
+		GET_MEMBER_NAME_CHECKED(UText3DDefaultTransformExtension, ScaleOrder),
+		GET_MEMBER_NAME_CHECKED(UText3DDefaultTransformExtension, ScaleRange),
+		GET_MEMBER_NAME_CHECKED(UText3DDefaultTransformExtension, ScaleBegin),
+		GET_MEMBER_NAME_CHECKED(UText3DDefaultTransformExtension, ScaleEnd),
+	};
+
+	if (PropertyNames.Contains(InEvent.GetMemberPropertyName()))
+	{
+		OnTransformOptionsChanged();
+	}
+}
+#endif
+
+void UText3DDefaultTransformExtension::OnTransformOptionsChanged()
+{
+	EText3DRendererFlags Flags = EText3DRendererFlags::Material;
+	EnumAddFlags(Flags, EText3DRendererFlags::Layout);
+	RequestUpdate(Flags);
+}
+
+void UText3DDefaultTransformExtension::ApplyTransform(uint32 InGlyphIndex, uint32 InGlyphCount, FTransform& InGlyphTransform)
+{
+	if (bLocationEnabled)
+	{
+		float Effect = GetEffectValue(InGlyphIndex, InGlyphCount, LocationOrder, LocationProgress, LocationRange);
+		FVector Location = LocationDistance * Effect;
+		InGlyphTransform.AddToTranslation(Location);
+	}
+
+	if (bScaleEnabled)
+	{
+		float Effect = GetEffectValue(InGlyphIndex, InGlyphCount, ScaleOrder, ScaleProgress, ScaleRange);
+		FVector Scale = ScaleBegin + ((ScaleEnd - ScaleBegin) * Effect);
+		InGlyphTransform.MultiplyScale3D(Scale);
+	}
+
+	if (bRotationEnabled)
+	{
+		float Effect = GetEffectValue(InGlyphIndex, InGlyphCount, RotationOrder, RotationProgress, RotationRange);
+		FRotator Rotator = RotationBegin + (RotationEnd - RotationBegin) * Effect;
+		InGlyphTransform.ConcatenateRotation(Rotator.Quaternion());
+	}
+}
+
+float UText3DDefaultTransformExtension::GetEffectValue(int32 Index, int32 Total, EText3DCharacterEffectOrder Order, float Progress, float Range)
+{
+	float Effect = 0.0f, Strip = 1.0f;
+	GetLineParameters(Range * 0.01f, Order, Total, Effect, Strip);
+
+	int32 Position = GetEffectPosition(Index, Total, Order);
+	return FMath::Clamp((Progress * 0.01f - Effect * Position) / Strip, 0.0f, 1.0f);
+}
+
+void UText3DDefaultTransformExtension::GetLineParameters(float Range, EText3DCharacterEffectOrder Order, int32 Total, float & EffectOut, float & StripOut)
+{
+	float Segment = 1.0f / Total;
+	EffectOut = (1.0f - Range) / Total;
+	StripOut = FMath::Max(Range, Segment);
+
+	if (Order == EText3DCharacterEffectOrder::FromCenter || Order == EText3DCharacterEffectOrder::ToCenter)
+	{
+		int32 Center = Total / 2.0f - 0.5f;
+		EffectOut *= (float)Total / (Center + 1.0f);
+	}
+
+	if (FMath::IsNearlyZero(EffectOut))
+	{
+		EffectOut = 0.01f;
+	}
+}
+
+int32 UText3DDefaultTransformExtension::GetEffectPosition(int32 Index, int32 Total, EText3DCharacterEffectOrder Order)
+{
+	int32 Center = 0.5f * Total - 0.5f;
+	switch (Order)
+	{
+	case EText3DCharacterEffectOrder::FromCenter:
+		{
+			if (Index > Center)
+			{
+				Index = Total - Index - 1;
+			}
+
+			return Center - Index;
+		}
+
+	case EText3DCharacterEffectOrder::ToCenter:
+		{
+			if (Index > Center)
+			{
+				return Total - Index - 1;
+			}
+
+			return Index;
+		}
+
+	case EText3DCharacterEffectOrder::Opposite:
+		{
+			return Total - Index - 1;
+		}
+	default: ;
+	}
+
+	return Index;
+}
\ No newline at end of file
diff --git a/Engine/Plugins/Experimental/Text3D/Source/Text3D/Private/Extensions/Text3DExtensionBase.cpp b/Engine/Plugins/Experimental/Text3D/Source/Text3D/Private/Extensions/Text3DExtensionBase.cpp
new file mode 100644
--- /dev/null
+++ b/Engine/Plugins/Experimental/Text3D/Source/Text3D/Private/Extensions/Text3DExtensionBase.cpp
@@ -0,0 +1,18 @@
+// Copyright Epic Games, Inc. All Rights Reserved.
+
+#include "Extensions/Text3DExtensionBase.h"
+
+#include "Text3DComponent.h"
+
+void UText3DExtensionBase::RequestUpdate(EText3DRendererFlags InFlags, bool bInImmediate) const
+{
+	if (UText3DComponent* Text3DComponent = GetText3DComponent())
+	{
+		Text3DComponent->RequestUpdate(InFlags, bInImmediate);
+	}
+}
+
+UText3DComponent* UText3DExtensionBase::GetText3DComponent() const
+{
+	return GetTypedOuter<UText3DComponent>();
+}
diff --git a/Engine/Plugins/Experimental/Text3D/Source/Text3D/Private/Renderers/StaticMeshes/Text3DStaticMeshesRenderer.cpp b/Engine/Plugins/Experimental/Text3D/Source/Text3D/Private/Renderers/StaticMeshes/Text3DStaticMeshesRenderer.cpp
--- a/Engine/Plugins/Experimental/Text3D/Source/Text3D/Private/Renderers/StaticMeshes/Text3DStaticMeshesRenderer.cpp
+++ b/Engine/Plugins/Experimental/Text3D/Source/Text3D/Private/Renderers/StaticMeshes/Text3DStaticMeshesRenderer.cpp
@@ -7,6 +7,10 @@
 #include "Components/SceneComponent.h"
 #include "Components/StaticMeshComponent.h"
 #include "Engine/StaticMesh.h"
+#include "Extensions/Text3DGeometryExtensionBase.h"
+#include "Extensions/Text3DLayoutExtensionBase.h"
+#include "Extensions/Text3DMaterialExtensionBase.h"
+#include "Extensions/Text3DRenderingExtensionBase.h"
 #include "LayoutBuilders/Text3DShapedGlyphText.h"
 #include "LayoutBuilders/Text3DLayoutShaper.h"
 #include "Logs/Text3DLogs.h"
@@ -14,26 +18,29 @@
 #include "Subsystems/Text3DEngineSubsystem.h"
 #include "Text3DComponent.h"
 #include "Text3DInternalTypes.h"
+#include "Extensions/Text3DTransformExtensionBase.h"
 
 void UText3DStaticMeshesRenderer::OnCreate()
 {
 	if (!TextRoot)
 	{
 		UText3DComponent* TextComponent = GetText3DComponent();
-		TextRoot = NewObject<USceneComponent>(TextComponent, TEXT("Text3DRoot"), RF_Transient);
+		TextRoot = NewObject<USceneComponent>(TextComponent, TEXT("TextRoot"), RF_Transient);
 		TextRoot->SetupAttachment(TextComponent);
 	}
 }
 
 void UText3DStaticMeshesRenderer::OnUpdate(EText3DRendererFlags InFlags)
 {
-	UText3DComponent* TextComponent = GetText3DComponent();
+	const UText3DComponent* TextComponent = GetText3DComponent();
 
 	if (EnumHasAnyFlags(InFlags, EText3DRendererFlags::Geometry))
 	{
-		TextRoot->SetRelativeScale3D(TextComponent->GetTextScale());
-		
-		const TSharedPtr<FText3DShapedGlyphText> ShapedText = TextComponent->GetShapedText();
+		const UText3DLayoutExtensionBase* LayoutExtension = TextComponent->GetLayoutExtension();
+
+		TextRoot->SetRelativeScale3D(LayoutExtension->GetTextScale());
+
+		const TSharedPtr<FText3DShapedGlyphText> ShapedText = LayoutExtension->GetShapedText();
 
 		// Pre-allocate, avoid new'ing!
 		AllocateGlyphs(Algo::TransformAccumulate(ShapedText->Lines, [](const FShapedGlyphLine& ShapedLine)
@@ -45,16 +52,12 @@ void UText3DStaticMeshesRenderer::OnUpdate(EText3DRendererFlags InFlags)
 		},
 		0));
 
-		UText3DEngineSubsystem* Subsystem = UText3DEngineSubsystem::Get();
-
-		const uint32 TypefaceIndex = TextComponent->GetTypeFaceIndex();
-		FCachedFontData& CachedFontData = Subsystem->GetCachedFontData(TextComponent->GetFont(), TypefaceIndex);
-		FGlyphMeshParameters GlyphMeshParameters{TextComponent->GetExtrude(), TextComponent->GetBevel(), TextComponent->GetBevelType(), TextComponent->GetBevelSegments(), TextComponent->HasOutline(), TextComponent->GetOutlineExpand(), TypefaceIndex};
+		const UText3DGeometryExtensionBase* GeometryExtension = TextComponent->GetGeometryExtension();
 
 		int32 GlyphIndex = 0;
 		for (TConstEnumerateRef<FShapedGlyphLine> ShapedLine : EnumerateRange(ShapedText->Lines))
 		{
-			FVector LineLocation = TextComponent->GetLineLocation(ShapedLine.GetIndex());
+			FVector LineLocation = LayoutExtension->GetLineLocation(ShapedLine.GetIndex());
 
 			for (TConstEnumerateRef<FShapedGlyphEntry> ShapedGlyph : EnumerateRange(ShapedLine->GlyphsToRender))
 			{
@@ -70,7 +73,7 @@ void UText3DStaticMeshesRenderer::OnUpdate(EText3DRendererFlags InFlags)
 				const int32 GlyphId = GlyphIndex++;
 
 				const FFreeTypeFace* FontFace = ShapedText->GlyphIndexToFontFace[ShapedGlyph->GlyphIndex];
-				UStaticMesh* CachedMesh = CachedFontData.GetGlyphMesh(ShapedGlyph->GlyphIndex, GlyphMeshParameters, FontFace);
+				UStaticMesh* CachedMesh = GeometryExtension->GetGlyphMesh(ShapedGlyph->GlyphIndex, FontFace);
 				if (!CachedMesh || FMath::IsNearlyZero(CachedMesh->GetBounds().SphereRadius))
 				{
 					continue;
@@ -80,42 +83,33 @@ void UText3DStaticMeshesRenderer::OnUpdate(EText3DRendererFlags InFlags)
 				{
 					UStaticMeshComponent* StaticMeshComponent = CharacterMeshes[GlyphId];
 					StaticMeshComponent->SetStaticMesh(CachedMesh);
-					StaticMeshComponent->SetVisibility(TextComponent->GetVisibleFlag());
-					StaticMeshComponent->SetHiddenInGame(TextComponent->bHiddenInGame);
-					StaticMeshComponent->SetCastShadow(TextComponent->CastsShadow());
 				}
 				else
 				{
 					// @note: This shouldn't occur, but it does under unknown circumstances (UE-164789) so it should be handled
 					UE_LOG(LogText3D, Error, TEXT("CharacterMesh not found at index %d"), GlyphId);
 				}
-
-				if (CharacterKernings.IsValidIndex(GlyphId))
-				{
-					FTransform Transform;
-					Transform.SetLocation(GlyphLocation);
-					USceneComponent* CharacterKerningComponent = CharacterKernings[GlyphId];
-					CharacterKerningComponent->SetRelativeTransform(Transform);
-				}
-				else
-				{
-					// @note: This shouldn't occur, but it does under unknown circumstances (UE-164789) so it should be handled
-					UE_LOG(LogText3D, Error, TEXT("CharacterKerning not found at index %d"), GlyphId);
-				}
 			}
 		}
+
+		RefreshBounds();
 	}
 
 	if (EnumHasAnyFlags(InFlags, EText3DRendererFlags::Layout))
 	{
-		TextRoot->SetRelativeScale3D(TextComponent->GetTextScale());
+		const UText3DLayoutExtensionBase* LayoutExtension = TextComponent->GetLayoutExtension();
+
+		TextRoot->SetRelativeScale3D(LayoutExtension->GetTextScale());
 
-		const TSharedPtr<FText3DShapedGlyphText> ShapedText = TextComponent->GetShapedText();
+		const TSharedPtr<FText3DShapedGlyphText> ShapedText = LayoutExtension->GetShapedText();
 
+		const TConstArrayView<UText3DTransformExtensionBase*> TransformExtensions = TextComponent->GetTransformExtensions();
+
+		const int32 GlyphCount = CharacterMeshes.Num();
 		int32 GlyphIndex = 0;
 		for (TConstEnumerateRef<FShapedGlyphLine> ShapedLine : EnumerateRange(ShapedText->Lines))
 		{
-			FVector Location = TextComponent->GetLineLocation(ShapedLine.GetIndex());
+			FVector Location = LayoutExtension->GetLineLocation(ShapedLine.GetIndex());
 
 			for (TConstEnumerateRef<FShapedGlyphEntry> ShapedGlyph : EnumerateRange(ShapedLine->GlyphsToRender))
 			{
@@ -132,25 +126,34 @@ void UText3DStaticMeshesRenderer::OnUpdate(EText3DRendererFlags InFlags)
 					GlyphKerningComponent->SetRelativeLocation(CharLocation);
 				}
 
+				if (UStaticMeshComponent* GlyphMeshComponent = GetGlyphMeshComponent(GlyphIndex))
+				{
+					FTransform GlyphTransform = FTransform::Identity;
+
+					for (UText3DTransformExtensionBase* TransformExtension : TransformExtensions)
+					{
+						if (TransformExtension && TransformExtension->GetTargetRange().IsInRange(GlyphIndex))
+						{
+							TransformExtension->ApplyTransform(GlyphIndex, GlyphCount, GlyphTransform);
+						}
+					}
+
+					GlyphMeshComponent->SetRelativeTransform(GlyphTransform);
+				}
+
 				GlyphIndex++;
 			}
 		}
-	}
-
-	if (EnumHasAnyFlags(InFlags, EText3DRendererFlags::Visibility))
-	{
-		for (UStaticMeshComponent* StaticMeshComponent : CharacterMeshes)
-		{
-			StaticMeshComponent->SetHiddenInGame(TextComponent->bHiddenInGame);
-			StaticMeshComponent->SetVisibility(TextComponent->GetVisibleFlag());
-			StaticMeshComponent->SetCastShadow(TextComponent->CastsShadow());
-		}
+		
+		RefreshBounds();
 	}
 
 	if (EnumHasAnyFlags(InFlags, EText3DRendererFlags::Material))
 	{
 		using namespace UE::Text3D::Materials;
 
+		const UText3DMaterialExtensionBase* MaterialExtensionBase = TextComponent->GetMaterialExtension();
+
 		for (UStaticMeshComponent* StaticMeshComponent : CharacterMeshes)
 		{
 			for (int32 GroupIndex = 0; GroupIndex < static_cast<int32>(EText3DGroupType::TypeCount); GroupIndex++)
@@ -162,7 +165,7 @@ void UText3DStaticMeshesRenderer::OnUpdate(EText3DRendererFlags InFlags)
 					continue;
 				}
 
-				UMaterialInterface* Material = TextComponent->GetMaterial(static_cast<EText3DGroupType>(GroupIndex));
+				UMaterialInterface* Material = MaterialExtensionBase->GetMaterial(static_cast<EText3DGroupType>(GroupIndex));
 
 				if (Material != StaticMeshComponent->GetMaterial(MaterialIndex))
 				{
@@ -171,6 +174,19 @@ void UText3DStaticMeshesRenderer::OnUpdate(EText3DRendererFlags InFlags)
 			}
 		}
 	}
+
+	if (EnumHasAnyFlags(InFlags, EText3DRendererFlags::Visibility))
+	{
+		const UText3DRenderingExtensionBase* RenderingExtension = TextComponent->GetRenderingExtension();
+
+		for (UStaticMeshComponent* StaticMeshComponent : CharacterMeshes)
+		{
+			StaticMeshComponent->SetHiddenInGame(TextComponent->bHiddenInGame);
+			StaticMeshComponent->SetVisibility(TextComponent->GetVisibleFlag());
+			StaticMeshComponent->SetCastShadow(RenderingExtension->GetTextCastShadow());
+			StaticMeshComponent->SetCastHiddenShadow(RenderingExtension->GetTextCastHiddenShadow());
+		}
+	}
 }
 
 void UText3DStaticMeshesRenderer::OnClear()
@@ -334,7 +350,7 @@ const TArray<USceneComponent*>& UText3DStaticMeshesRenderer::GetGlyphKerningComp
 
 UStaticMeshComponent* UText3DStaticMeshesRenderer::GetGlyphMeshComponent(int32 Index)
 {
-	if (!CharacterKernings.IsValidIndex(Index))
+	if (!CharacterMeshes.IsValidIndex(Index))
 	{
 		return nullptr;
 	}
diff --git a/Engine/Plugins/Experimental/Text3D/Source/Text3D/Private/Renderers/Text3DRendererBase.cpp b/Engine/Plugins/Experimental/Text3D/Source/Text3D/Private/Renderers/Text3DRendererBase.cpp
--- a/Engine/Plugins/Experimental/Text3D/Source/Text3D/Private/Renderers/Text3DRendererBase.cpp
+++ b/Engine/Plugins/Experimental/Text3D/Source/Text3D/Private/Renderers/Text3DRendererBase.cpp
@@ -2,6 +2,11 @@
 
 #include "Renderers/Text3DRendererBase.h"
 
+#include "Extensions/Text3DGeometryExtensionBase.h"
+#include "Extensions/Text3DLayoutExtensionBase.h"
+#include "Extensions/Text3DMaterialExtensionBase.h"
+#include "Extensions/Text3DRenderingExtensionBase.h"
+#include "Extensions/Text3DTransformExtensionBase.h"
 #include "GameFramework/Actor.h"
 #include "Logs/Text3DLogs.h"
 #include "Text3DComponent.h"
@@ -11,6 +16,11 @@ UText3DComponent* UText3DRendererBase::GetText3DComponent() const
 	return GetTypedOuter<UText3DComponent>();
 }
 
+void UText3DRendererBase::RefreshBounds()
+{
+	CachedBounds = OnCalculateBounds();
+}
+
 void UText3DRendererBase::Create()
 {
 	if (bInitialized)
@@ -42,15 +52,67 @@ void UText3DRendererBase::Update(EText3DRendererFlags InFlags)
 		return;
 	}
 
-	OnUpdate(InFlags);
-	
-	if (EnumHasAnyFlags(InFlags, EText3DRendererFlags::Geometry)
-		|| EnumHasAnyFlags(InFlags, EText3DRendererFlags::Layout))
+	TArray<UText3DExtensionBase*, TInlineAllocator<5>> Extensions
+	{
+		Text3DComponent->GetGeometryExtension(),
+		Text3DComponent->GetLayoutExtension(),
+		Text3DComponent->GetMaterialExtension(),
+		Text3DComponent->GetRenderingExtension()
+	};
+
+	Extensions.Append(Text3DComponent->GetTransformExtensions());
+
+	uint8 Flag = 1 << 0;
+
+	while (Flag < static_cast<uint8>(EText3DRendererFlags::All))
 	{
-		CachedBounds = OnCalculateBounds();
+		const EText3DRendererFlags FlagCasted = static_cast<EText3DRendererFlags>(Flag);
+
+		if (EnumHasAnyFlags(InFlags, FlagCasted))
+		{
+			for (int32 Index = Extensions.Num() - 1; Index >= 0; --Index)
+			{
+				if (UText3DExtensionBase* Extension = Extensions[Index])
+				{
+					const EText3DExtensionResult ExtensionStatus = Extension->PreRendererUpdate(FlagCasted);
+
+					if (ExtensionStatus == EText3DExtensionResult::Failed)
+					{
+						UE_LOG(LogText3D, Error, TEXT("Failed to PRE update Text3D %s extension"), *GetNameSafe(Extension->GetClass()))
+						return;
+					}
+					else if (ExtensionStatus == EText3DExtensionResult::Finished)
+					{
+						Extensions.RemoveAt(Index);
+					}
+				}
+			}
+
+			OnUpdate(FlagCasted);
+
+			for (int32 Index = Extensions.Num() - 1; Index >= 0; --Index)
+			{
+				if (UText3DExtensionBase* Extension = Extensions[Index])
+				{
+					const EText3DExtensionResult ExtensionStatus = Extension->PostRendererUpdate(FlagCasted);
+
+					if (ExtensionStatus == EText3DExtensionResult::Failed)
+					{
+						UE_LOG(LogText3D, Error, TEXT("Failed to POST update Text3D %s extension"), *GetNameSafe(Extension->GetClass()))
+						return;
+					}
+					else if (ExtensionStatus == EText3DExtensionResult::Finished)
+					{
+						Extensions.RemoveAt(Index);
+					}
+				}
+			}
+		}
+
+		Flag <<= 1;
 	}
 
-	UE_LOG(LogText3D, Verbose, TEXT("%s : Text3DRenderer %s Updated with flags %i"), *Text3DComponent->GetOwner()->GetActorNameOrLabel(), *GetName().ToString(), InFlags)
+	UE_LOG(LogText3D, Log, TEXT("%s : Text3DRenderer %s Updated with flags %i"), *Text3DComponent->GetOwner()->GetActorNameOrLabel(), *GetName().ToString(), InFlags)
 }
 
 void UText3DRendererBase::Clear()
@@ -70,7 +132,7 @@ void UText3DRendererBase::Clear()
 	CachedBounds.Reset();
 
 	const AActor* Owner = Text3DComponent->GetOwner();
-	UE_LOG(LogText3D, Verbose, TEXT("%s : Text3DRenderer %s Cleared"), !!Owner ? *Owner->GetActorNameOrLabel() : TEXT("Invalid owner"), *GetName().ToString())
+	UE_LOG(LogText3D, Log, TEXT("%s : Text3DRenderer %s Cleared"), !!Owner ? *Owner->GetActorNameOrLabel() : TEXT("Invalid owner"), *GetName().ToString())
 }
 
 void UText3DRendererBase::Destroy()
diff --git a/Engine/Plugins/Experimental/Text3D/Source/Text3D/Private/Subsystems/Text3DEngineSubsystem.h b/Engine/Plugins/Experimental/Text3D/Source/Text3D/Private/Subsystems/Text3DEngineSubsystem.h
--- a/Engine/Plugins/Experimental/Text3D/Source/Text3D/Private/Subsystems/Text3DEngineSubsystem.h
+++ b/Engine/Plugins/Experimental/Text3D/Source/Text3D/Private/Subsystems/Text3DEngineSubsystem.h
@@ -14,6 +14,7 @@ class FFreeTypeFace;
 class UFont;
 class UMaterial;
 class UStaticMesh;
+class UWorld;
 
 USTRUCT()
 struct FGlyphMeshParameters
@@ -196,7 +197,7 @@ public:
 	// ~End UEngineSubsystem
 
 	// ~Begin FSelfRegisteringExec
-	virtual bool Exec_Dev(class UWorld* InWorld, const TCHAR* InCmd, FOutputDevice& InAr) override;
+	virtual bool Exec_Dev(UWorld* InWorld, const TCHAR* InCmd, FOutputDevice& InAr) override;
 	// ~End FSelfRegisteringExec
 
 	void PrintCache() const;
