[Motion Design] Rundown Editor: fix the page editor so it doesn't dirty the rundown just by selecting different pages.

Saving the page detail view data in the page would automatically dirty the rundown asset even if no changes where made.
Simplest options is to rely solely on the OnFinishedChangingPageDetailProperties.

diff --git a/Engine/Plugins/Experimental/Avalanche/Source/AvalancheMediaEditor/Private/Rundown/DetailsView/SAvaRundownPageDetails.cpp b/Engine/Plugins/Experimental/Avalanche/Source/AvalancheMediaEditor/Private/Rundown/DetailsView/SAvaRundownPageDetails.cpp
--- a/Engine/Plugins/Experimental/Avalanche/Source/AvalancheMediaEditor/Private/Rundown/DetailsView/SAvaRundownPageDetails.cpp
+++ b/Engine/Plugins/Experimental/Avalanche/Source/AvalancheMediaEditor/Private/Rundown/DetailsView/SAvaRundownPageDetails.cpp
@@ -232,20 +232,6 @@ void SAvaRundownPageDetails::RefreshPageDetailView()
 	const FRundownContext RundownContext = GetRundownContext();
 	if (RundownContext.IsValid())
 	{
-		// Save previously edited page -- todo: this may not be necessary.
-		if (ActivePageId != CurrentPageDetailViewPageId
-			&& CurrentPageDetailViewPageId != FAvaRundownPage::InvalidPageId
-			&& CurrentPageDetailViewData.IsValid())
-		{
-			FAvaRundownPage& Page = RundownContext.GetRundown()->GetPage(CurrentPageDetailViewPageId);
-			if (Page.IsValidPage())
-			{
-				RundownContext.GetRundown()->Modify();
-				Page.SetInstancedCommands(CurrentPageDetailViewData->Commands);
-				RundownContext.GetRundown()->GetOnPagesChanged().Broadcast(RundownContext.GetRundown(), Page, EAvaRundownPageChanges::Commands);
-			}
-		}
-		
 		const FAvaRundownPage& ActivePage = RundownContext.GetRundown()->GetPage(ActivePageId);
 		if (ActivePage.IsValidPage())
 		{
