[MotionDesign] Add foreground color for SImage widgets where applicable

diff --git a/Engine/Plugins/Experimental/Avalanche/Source/AvalancheMediaEditor/Private/Broadcast/ChannelGrid/Slate/SAvaBroadcastChannel.cpp b/Engine/Plugins/Experimental/Avalanche/Source/AvalancheMediaEditor/Private/Broadcast/ChannelGrid/Slate/SAvaBroadcastChannel.cpp
--- a/Engine/Plugins/Experimental/Avalanche/Source/AvalancheMediaEditor/Private/Broadcast/ChannelGrid/Slate/SAvaBroadcastChannel.cpp
+++ b/Engine/Plugins/Experimental/Avalanche/Source/AvalancheMediaEditor/Private/Broadcast/ChannelGrid/Slate/SAvaBroadcastChannel.cpp
@@ -101,6 +101,7 @@ void SAvaBroadcastChannel::Construct(const FArguments& InArgs, const TSharedPtr<
 				.ButtonContent()
 				[
 					SNew(SImage)
+					.ColorAndOpacity(FSlateColor::UseForeground())
 					.Image(FAppStyle::Get().GetBrush("Icons.Settings"))
 				]
 			]
@@ -114,6 +115,7 @@ void SAvaBroadcastChannel::Construct(const FArguments& InArgs, const TSharedPtr<
 				.ToolTipText(LOCTEXT("ChannelPinToolTip", "Pin the channel across all profiles."))
 				[
 					SNew(SImage)
+					.ColorAndOpacity(FSlateColor::UseForeground())
 					.Image(this, &SAvaBroadcastChannel::GetChannelPinBrush)
 				]
 			]
@@ -127,6 +129,7 @@ void SAvaBroadcastChannel::Construct(const FArguments& InArgs, const TSharedPtr<
 				.ToolTipText(this, &SAvaBroadcastChannel::GetChannelMaximizeRestoreTooltipText)
 				[
 					SNew(SImage)
+					.ColorAndOpacity(FSlateColor::UseForeground())
 					.Image(this, &SAvaBroadcastChannel::GetChannelMaximizeRestoreBrush)
 				]
 			]
@@ -141,6 +144,7 @@ void SAvaBroadcastChannel::Construct(const FArguments& InArgs, const TSharedPtr<
 				.ToolTipText(LOCTEXT("ChannelRemoveToolTip", "Remove this channel."))
 				[
 					SNew(SImage)
+					.ColorAndOpacity(FSlateColor::UseForeground())
 					.Image(FAppStyle::Get().GetBrush("Icons.X"))
 				]
 			]
@@ -165,6 +169,7 @@ void SAvaBroadcastChannel::Construct(const FArguments& InArgs, const TSharedPtr<
 					.OnGetMenuContent(this, &SAvaBroadcastChannel::GetChannelStatusOptions)
 					[
 						SNew(SImage)
+						.ColorAndOpacity(FSlateColor::UseForeground())
 						.Image(this, &SAvaBroadcastChannel::GetChannelStatusBrush)
 					]				
 				]
@@ -192,6 +197,7 @@ void SAvaBroadcastChannel::Construct(const FArguments& InArgs, const TSharedPtr<
 				.ToolTipText(LOCTEXT("ChannelTypeToggleToolTip", "Toggles between \"Program\" or \"Preview\" channel type."))
 				[
 					SNew(SImage)
+					.ColorAndOpacity(FSlateColor::UseForeground())
 					.Image(this, &SAvaBroadcastChannel::GetChannelTypeBrush)
 				]
 			]
@@ -806,8 +812,11 @@ void SAvaBroadcastChannel::OnChannelQualitySettingsChanged(const FPropertyChange
 	}
 }
 
-TSharedRef<SWidget> SAvaBroadcastChannel::GetQualitySettingsDetails() const
+TSharedRef<SWidget> SAvaBroadcastChannel::GetQualitySettingsDetails()
 {
+	// We need to recreate the details view to update presets from the config the user may have changed
+	CreateQualitySettingsDetails();
+
 	return SNew(SBox)
 		.MaxDesiredHeight(500.f)
 		.MinDesiredWidth(120.f)
diff --git a/Engine/Plugins/Experimental/Avalanche/Source/AvalancheOutliner/Private/Columns/AvaOutlinerLockColumn.cpp b/Engine/Plugins/Experimental/Avalanche/Source/AvalancheOutliner/Private/Columns/AvaOutlinerLockColumn.cpp
--- a/Engine/Plugins/Experimental/Avalanche/Source/AvalancheOutliner/Private/Columns/AvaOutlinerLockColumn.cpp
+++ b/Engine/Plugins/Experimental/Avalanche/Source/AvalancheOutliner/Private/Columns/AvaOutlinerLockColumn.cpp
@@ -24,6 +24,7 @@ SHeaderRow::FColumn::FArguments FAvaOutlinerLockColumn::ConstructHeaderRowColumn
 		.DefaultTooltip(GetColumnDisplayNameText())
 		[
 			SNew(SImage)
+			.ColorAndOpacity(FSlateColor::UseForeground())
 			.Image(FAppStyle::GetBrush("Icons.Lock"))
 		];
 }
diff --git a/Engine/Plugins/Experimental/Avalanche/Source/AvalancheOutliner/Private/Columns/Slate/SAvaOutlinerItemChip.cpp b/Engine/Plugins/Experimental/Avalanche/Source/AvalancheOutliner/Private/Columns/Slate/SAvaOutlinerItemChip.cpp
--- a/Engine/Plugins/Experimental/Avalanche/Source/AvalancheOutliner/Private/Columns/Slate/SAvaOutlinerItemChip.cpp
+++ b/Engine/Plugins/Experimental/Avalanche/Source/AvalancheOutliner/Private/Columns/Slate/SAvaOutlinerItemChip.cpp
@@ -35,6 +35,7 @@ void SAvaOutlinerItemChip::Construct(const FArguments& InArgs
 				.Stretch(EStretch::ScaleToFit)
 				[
 					SNew(SImage)
+					.ColorAndOpacity(FSlateColor::UseForeground())
 					.Image(InItem, &IAvaOutlinerItem::GetIconBrush)
 				]
 			]
diff --git a/Engine/Plugins/Experimental/Avalanche/Source/AvalancheOutliner/Private/Slate/SAvaOutlinerItemFilters.cpp b/Engine/Plugins/Experimental/Avalanche/Source/AvalancheOutliner/Private/Slate/SAvaOutlinerItemFilters.cpp
--- a/Engine/Plugins/Experimental/Avalanche/Source/AvalancheOutliner/Private/Slate/SAvaOutlinerItemFilters.cpp
+++ b/Engine/Plugins/Experimental/Avalanche/Source/AvalancheOutliner/Private/Slate/SAvaOutlinerItemFilters.cpp
@@ -80,6 +80,7 @@ void SAvaOutlinerItemFilters::Construct(const FArguments& InArgs, const TSharedR
 				+ SOverlay::Slot()
 				[
 					SNew(SImage)
+					.ColorAndOpacity(FSlateColor::UseForeground())
 					.Image(FAvaOutlinerStyle::Get().GetBrush(TEXT("AvaOutliner.FilterIcon")))
 					.DesiredSizeOverride(FVector2D(24.f))
 				]
@@ -88,6 +89,7 @@ void SAvaOutlinerItemFilters::Construct(const FArguments& InArgs, const TSharedR
 				.VAlign(VAlign_Bottom)
 				[
 					SNew(SImage)
+					.ColorAndOpacity(FSlateColor::UseForeground())
 					.Image(this, &SAvaOutlinerItemFilters::GetFilterItemMenuButtonIcon)
 					.DesiredSizeOverride(FVector2D(12.f))
 					.RenderTransform(FSlateRenderTransform(FVector2f(2.f, 0.f)))
diff --git a/Engine/Plugins/Experimental/Avalanche/Source/AvalancheSequencer/Private/EaseCurveTool/Widgets/SAvaEaseCurvePreset.cpp b/Engine/Plugins/Experimental/Avalanche/Source/AvalancheSequencer/Private/EaseCurveTool/Widgets/SAvaEaseCurvePreset.cpp
--- a/Engine/Plugins/Experimental/Avalanche/Source/AvalancheSequencer/Private/EaseCurveTool/Widgets/SAvaEaseCurvePreset.cpp
+++ b/Engine/Plugins/Experimental/Avalanche/Source/AvalancheSequencer/Private/EaseCurveTool/Widgets/SAvaEaseCurvePreset.cpp
@@ -61,10 +61,12 @@ void SAvaEaseCurvePreset::Construct(const FArguments& InArgs)
 					[
 						SNew(SImage)
 						.DesiredSizeOverride(FVector2D(FAvaEaseCurveStyle::Get().GetFloat(TEXT("ToolButton.ImageSize"))))
+						.ColorAndOpacity(FSlateColor::UseForeground())
 						.Image(FAppStyle::GetBrush(TEXT("Icons.Plus")))
 					]
 				]
 			]
+			
 			+ SWidgetSwitcher::Slot()
 			[
 				SNew(SHorizontalBox)
@@ -90,6 +92,7 @@ void SAvaEaseCurvePreset::Construct(const FArguments& InArgs)
 					[
 						SNew(SImage)
 						.DesiredSizeOverride(FVector2D(FAvaEaseCurveStyle::Get().GetFloat(TEXT("ToolButton.ImageSize"))))
+						.ColorAndOpacity(FSlateColor::UseForeground())
 						.Image(FAppStyle::GetBrush(TEXT("Icons.X")))
 					]
 				]
diff --git a/Engine/Plugins/Experimental/Avalanche/Source/AvalancheSequencer/Private/EaseCurveTool/Widgets/SAvaEaseCurvePresetComboBox.cpp b/Engine/Plugins/Experimental/Avalanche/Source/AvalancheSequencer/Private/EaseCurveTool/Widgets/SAvaEaseCurvePresetComboBox.cpp
--- a/Engine/Plugins/Experimental/Avalanche/Source/AvalancheSequencer/Private/EaseCurveTool/Widgets/SAvaEaseCurvePresetComboBox.cpp
+++ b/Engine/Plugins/Experimental/Avalanche/Source/AvalancheSequencer/Private/EaseCurveTool/Widgets/SAvaEaseCurvePresetComboBox.cpp
@@ -68,6 +68,7 @@ TSharedRef<SWidget> SAvaEaseCurvePresetComboBox::GenerateSearchRowWidget()
 			[
 				SNew(SImage)
 				.DesiredSizeOverride(ButtonImageSize)
+				.ColorAndOpacity(FSlateColor::UseForeground())
 				.Image(FAppStyle::GetBrush(TEXT("Icons.Refresh")))
 			]
 		];
@@ -87,6 +88,7 @@ TSharedRef<SWidget> SAvaEaseCurvePresetComboBox::GenerateSearchRowWidget()
 				[
 					SNew(SImage)
 					.DesiredSizeOverride(ButtonImageSize)
+					.ColorAndOpacity(FSlateColor::UseForeground())
 					.Image(FAppStyle::GetBrush(TEXT("Icons.FolderOpen")))
 				]
 			];
@@ -108,6 +110,7 @@ TSharedRef<SWidget> SAvaEaseCurvePresetComboBox::GenerateSearchRowWidget()
 				[
 					SNew(SImage)
 					.DesiredSizeOverride(ButtonImageSize)
+					.ColorAndOpacity(FSlateColor::UseForeground())
 					.Image(FAppStyle::GetBrush(TEXT("Icons.Edit")))
 				]
 			];
@@ -128,6 +131,7 @@ TSharedRef<SWidget> SAvaEaseCurvePresetComboBox::GenerateSearchRowWidget()
 				[
 					SNew(SImage)
 					.DesiredSizeOverride(ButtonImageSize)
+					.ColorAndOpacity(FSlateColor::UseForeground())
 					.Image(FAppStyle::GetBrush(TEXT("Icons.Plus")))
 				]
 			];
@@ -180,6 +184,7 @@ void SAvaEaseCurvePresetComboBox::UpdateGroupsContent()
 				.VAlign(VAlign_Center)
 				[
 					SNew(STextBlock)
+					.ColorAndOpacity(FSlateColor::UseForeground())
 					.TextStyle(FAppStyle::Get(), TEXT("HintText"))
 					.Text(InText)
 				];
diff --git a/Engine/Plugins/Experimental/Avalanche/Source/AvalancheSequencer/Private/EaseCurveTool/Widgets/SAvaEaseCurvePresetGroup.cpp b/Engine/Plugins/Experimental/Avalanche/Source/AvalancheSequencer/Private/EaseCurveTool/Widgets/SAvaEaseCurvePresetGroup.cpp
--- a/Engine/Plugins/Experimental/Avalanche/Source/AvalancheSequencer/Private/EaseCurveTool/Widgets/SAvaEaseCurvePresetGroup.cpp
+++ b/Engine/Plugins/Experimental/Avalanche/Source/AvalancheSequencer/Private/EaseCurveTool/Widgets/SAvaEaseCurvePresetGroup.cpp
@@ -86,6 +86,7 @@ void SAvaEaseCurvePresetGroup::Construct(const FArguments& InArgs)
 				[
 					SNew(SImage)
 					.DesiredSizeOverride(FVector2D(10.f))
+					.ColorAndOpacity(FSlateColor::UseForeground())
 					.Image(FAppStyle::GetBrush(TEXT("Icons.Delete")))
 				]
 			]
diff --git a/Engine/Plugins/Experimental/Avalanche/Source/AvalancheSequencer/Private/EaseCurveTool/Widgets/SAvaEaseCurvePresetGroupItem.cpp b/Engine/Plugins/Experimental/Avalanche/Source/AvalancheSequencer/Private/EaseCurveTool/Widgets/SAvaEaseCurvePresetGroupItem.cpp
--- a/Engine/Plugins/Experimental/Avalanche/Source/AvalancheSequencer/Private/EaseCurveTool/Widgets/SAvaEaseCurvePresetGroupItem.cpp
+++ b/Engine/Plugins/Experimental/Avalanche/Source/AvalancheSequencer/Private/EaseCurveTool/Widgets/SAvaEaseCurvePresetGroupItem.cpp
@@ -70,7 +70,7 @@ void SAvaEaseCurvePresetGroupItem::Construct(const FArguments& InArgs, const TSh
 							SNew(STextBlock)
 							.TextStyle(FAppStyle::Get(), TEXT("Menu.Label"))
 							.OverflowPolicy(ETextOverflowPolicy::Ellipsis)
-							.ColorAndOpacity(FStyleColors::White)
+							.ColorAndOpacity(FSlateColor::UseForeground())
 							.Text(FText::FromString(Preset->Name))
 							.ToolTipText(ItemTooltipText)
 						]
@@ -91,6 +91,7 @@ void SAvaEaseCurvePresetGroupItem::Construct(const FArguments& InArgs, const TSh
 								[
 									SNew(SImage)
 									.DesiredSizeOverride(FVector2D(10.f))
+									.ColorAndOpacity(FSlateColor::UseForeground())
 									.Image(FAppStyle::GetBrush(TEXT("Icons.Delete")))
 								]
 							]
diff --git a/Engine/Plugins/Experimental/Avalanche/Source/AvalancheSequencer/Private/StaggerTool/Widgets/SAvaStaggerOperationPoint.cpp b/Engine/Plugins/Experimental/Avalanche/Source/AvalancheSequencer/Private/StaggerTool/Widgets/SAvaStaggerOperationPoint.cpp
--- a/Engine/Plugins/Experimental/Avalanche/Source/AvalancheSequencer/Private/StaggerTool/Widgets/SAvaStaggerOperationPoint.cpp
+++ b/Engine/Plugins/Experimental/Avalanche/Source/AvalancheSequencer/Private/StaggerTool/Widgets/SAvaStaggerOperationPoint.cpp
@@ -116,6 +116,7 @@ TSharedRef<SCheckBox> SAvaStaggerOperationPoint::CreateButton(const float InValu
 			.VAlign(VAlign_Center)
 			[
 				SNew(SImage)
+				.ColorAndOpacity(FSlateColor::UseForeground())
 				.Image(FAppStyle::GetBrush(InImageBrushName))
 			]
 		];
diff --git a/Engine/Plugins/Experimental/Avalanche/Source/AvalancheSequencer/Private/StaggerTool/Widgets/SAvaStaggerTool.cpp b/Engine/Plugins/Experimental/Avalanche/Source/AvalancheSequencer/Private/StaggerTool/Widgets/SAvaStaggerTool.cpp
--- a/Engine/Plugins/Experimental/Avalanche/Source/AvalancheSequencer/Private/StaggerTool/Widgets/SAvaStaggerTool.cpp
+++ b/Engine/Plugins/Experimental/Avalanche/Source/AvalancheSequencer/Private/StaggerTool/Widgets/SAvaStaggerTool.cpp
@@ -103,6 +103,7 @@ TSharedRef<SWidget> SAvaStaggerTool::ConstructApplyRow()
 			[
 				SNew(SImage)
 				.DesiredSizeOverride(FVector2D(ImageSize))
+				.ColorAndOpacity(FSlateColor::UseForeground())
 				.Image(FAppStyle::GetBrush(TEXT("HorizontalAlignment_Left")))
 			]
 		]
@@ -146,6 +147,7 @@ TSharedRef<SWidget> SAvaStaggerTool::ConstructApplyRow()
 					.Padding(0.f, 4.f)
 					[
 						SNew(SImage)
+						.ColorAndOpacity(FSlateColor::UseForeground())
 						.Image_Lambda([this]()
 							{
 								const FName BrushName = (WeakTool.IsValid() && WeakTool.Pin()->IsAutoAppying())
@@ -177,6 +179,7 @@ TSharedRef<SWidget> SAvaStaggerTool::ConstructApplyRow()
 				[
 					SNew(SImage)
 					.DesiredSizeOverride(FVector2D(ImageSize))
+					.ColorAndOpacity(FSlateColor::UseForeground())
 					.Image(FAppStyle::GetBrush(TEXT("PropertyWindow.DiffersFromDefault")))
 				]
 			]
