[Sequencer] Include missed header file change

diff --git a/Engine/Plugins/Experimental/Avalanche/Source/AvalancheMediaEditor/Private/Broadcast/ChannelGrid/Slate/SAvaBroadcastChannel.h b/Engine/Plugins/Experimental/Avalanche/Source/AvalancheMediaEditor/Private/Broadcast/ChannelGrid/Slate/SAvaBroadcastChannel.h
--- a/Engine/Plugins/Experimental/Avalanche/Source/AvalancheMediaEditor/Private/Broadcast/ChannelGrid/Slate/SAvaBroadcastChannel.h
+++ b/Engine/Plugins/Experimental/Avalanche/Source/AvalancheMediaEditor/Private/Broadcast/ChannelGrid/Slate/SAvaBroadcastChannel.h
@@ -117,8 +117,8 @@ protected:
 
 	void OnChannelQualitySettingsChanged(const FPropertyChangedEvent& InEvent);
 
-	TSharedRef<SWidget> GetQualitySettingsDetails() const;
-	
+	TSharedRef<SWidget> GetQualitySettingsDetails();
+
 	FName ChannelName = NAME_None;
 
 	TSharedPtr<SMenuAnchor> ChannelStatusOptions;
