[Motion Design] Broadcast_"Profile"_Not set for Loc

diff --git a/Engine/Plugins/Experimental/Avalanche/Source/AvalancheMedia/Private/Broadcast/AvaBroadcast.cpp b/Engine/Plugins/Experimental/Avalanche/Source/AvalancheMedia/Private/Broadcast/AvaBroadcast.cpp
--- a/Engine/Plugins/Experimental/Avalanche/Source/AvalancheMedia/Private/Broadcast/AvaBroadcast.cpp
+++ b/Engine/Plugins/Experimental/Avalanche/Source/AvalancheMedia/Private/Broadcast/AvaBroadcast.cpp
@@ -156,7 +156,8 @@ FName UAvaBroadcast::CreateProfile(FName InProfileName, bool bMakeCurrentProfile
 {
 	if (InProfileName == NAME_None)
 	{
-		InProfileName = FName(TEXT("Profile"), 0);
+		static const FText DefaultProfileName = LOCTEXT("DefaultProfileName", "Profile");
+		InProfileName = FName(DefaultProfileName.ToString(), 0);
 	}
 	
 	FAvaBroadcastProfile& Profile = CreateProfileInternal(InProfileName);
