[MotionDesign] Broadcast_Settings_Not set for Loc:
- Change to use FText for preset names

diff --git a/Engine/Plugins/Experimental/Avalanche/Source/Avalanche/Private/Viewport/AvaViewportQualitySettings.cpp b/Engine/Plugins/Experimental/Avalanche/Source/Avalanche/Private/Viewport/AvaViewportQualitySettings.cpp
--- a/Engine/Plugins/Experimental/Avalanche/Source/Avalanche/Private/Viewport/AvaViewportQualitySettings.cpp
+++ b/Engine/Plugins/Experimental/Avalanche/Source/Avalanche/Private/Viewport/AvaViewportQualitySettings.cpp
@@ -5,6 +5,9 @@
 
 #define LOCTEXT_NAMESPACE "AvaViewportQualitySettings"
 
+const FText FAvaViewportQualitySettingsPreset::NoLumen = LOCTEXT("PresetNoLumen", "No Lumen");
+const FText FAvaViewportQualitySettingsPreset::Reduced = LOCTEXT("PresetReduced", "Reduced");
+
 FAvaViewportQualitySettings::FAvaViewportQualitySettings()
 {
 	Features = DefaultFeatures();
@@ -137,11 +140,11 @@ FAvaViewportQualitySettings FAvaViewportQualitySettings::Default()
 	return FAvaViewportQualitySettings();
 }
 
-FAvaViewportQualitySettings FAvaViewportQualitySettings::Preset(const FName& InPresetName)
+FAvaViewportQualitySettings FAvaViewportQualitySettings::Preset(const FText& InPresetName)
 {
 	FAvaViewportQualitySettings PresetSettings(true);
 
-	if (InPresetName.IsEqual("No Lumen", ENameCase::IgnoreCase))
+	if (InPresetName.EqualTo(FAvaViewportQualitySettingsPreset::NoLumen))
 	{
 		PresetSettings.EnableFeaturesByName(false, 
 			{
@@ -156,7 +159,7 @@ FAvaViewportQualitySettings FAvaViewportQualitySettings::Preset(const FName& InP
 				"LumenReuseShadowMaps"
 			});
 	}
-	else if (InPresetName.IsEqual("Reduced", ENameCase::IgnoreCase))
+	else if (InPresetName.EqualTo(FAvaViewportQualitySettingsPreset::Reduced))
 	{
 		PresetSettings.EnableFeaturesByName(false,
 			{
@@ -236,11 +239,11 @@ void FAvaViewportQualitySettings::SortFeaturesByDisplayText(TArray<FAvaViewportQ
 		{
 			FText NameTextA;
 			FText TooltipTextA;
-			FAvaViewportQualitySettings::FeatureNameAndTooltipText(FeatureA.Name, NameTextA, TooltipTextA);
+			FeatureNameAndTooltipText(FeatureA.Name, NameTextA, TooltipTextA);
 
 			FText NameTextB;
 			FText TooltipTextB;
-			FAvaViewportQualitySettings::FeatureNameAndTooltipText(FeatureB.Name, NameTextB, TooltipTextB);
+			FeatureNameAndTooltipText(FeatureB.Name, NameTextB, TooltipTextB);
 
 			return NameTextA.CompareToCaseIgnored(NameTextB) < 0;
 		});
diff --git a/Engine/Plugins/Experimental/Avalanche/Source/Avalanche/Public/Viewport/AvaViewportQualitySettings.h b/Engine/Plugins/Experimental/Avalanche/Source/Avalanche/Public/Viewport/AvaViewportQualitySettings.h
--- a/Engine/Plugins/Experimental/Avalanche/Source/Avalanche/Public/Viewport/AvaViewportQualitySettings.h
+++ b/Engine/Plugins/Experimental/Avalanche/Source/Avalanche/Public/Viewport/AvaViewportQualitySettings.h
@@ -6,7 +6,6 @@
 
 class FString;
 class FText;
-
 struct FEngineShowFlags;
 
 USTRUCT(BlueprintType)
@@ -45,12 +44,11 @@ struct AVALANCHE_API FAvaViewportQualitySettings
 {
 	GENERATED_BODY()
 
-public:
 	static TArray<FAvaViewportQualitySettingsFeature> DefaultFeatures();
 	static TArray<FAvaViewportQualitySettingsFeature> AllFeatures(const bool bUseAllFeatures);
 
 	static FAvaViewportQualitySettings Default();
-	static FAvaViewportQualitySettings Preset(const FName& InPresetName);
+	static FAvaViewportQualitySettings Preset(const FText& InPresetName);
 	static FAvaViewportQualitySettings All(const bool bUseAllFeatures);
 
 	static void FeatureNameAndTooltipText(const FString& InFeatureName, FText& OutNameText, FText& OutTooltipText);
@@ -84,3 +82,28 @@ public:
 	UPROPERTY(EditAnywhere, BlueprintReadWrite, EditFixedSize, Category = "Quality", meta = (EditFixedOrder))
 	TArray<FAvaViewportQualitySettingsFeature> Features;
 };
+
+USTRUCT(BlueprintType)
+struct AVALANCHE_API FAvaViewportQualitySettingsPreset
+{
+	GENERATED_BODY()
+
+	static const FText NoLumen;
+	static const FText Reduced;
+
+	FAvaViewportQualitySettingsPreset() {}
+	FAvaViewportQualitySettingsPreset(const FText& InPresetName, const FAvaViewportQualitySettings& InQualitySettings)
+		: PresetName(InPresetName), QualitySettings(InQualitySettings)
+	{}
+
+	bool operator==(const FAvaViewportQualitySettingsPreset& InOther) const
+	{
+		return PresetName.EqualTo(InOther.PresetName);
+	}
+
+	UPROPERTY(EditAnywhere, BlueprintReadWrite, Category = "Preset")
+	FText PresetName;
+
+	UPROPERTY(EditAnywhere, BlueprintReadWrite, Category = "Preset")
+	FAvaViewportQualitySettings QualitySettings;
+};
diff --git a/Engine/Plugins/Experimental/Avalanche/Source/AvalancheEditor/Private/AvaEditorSettings.cpp b/Engine/Plugins/Experimental/Avalanche/Source/AvalancheEditor/Private/AvaEditorSettings.cpp
--- a/Engine/Plugins/Experimental/Avalanche/Source/AvalancheEditor/Private/AvaEditorSettings.cpp
+++ b/Engine/Plugins/Experimental/Avalanche/Source/AvalancheEditor/Private/AvaEditorSettings.cpp
@@ -3,7 +3,6 @@
 #include "AvaEditorSettings.h"
 #include "ISettingsModule.h"
 #include "Modules/ModuleManager.h"
-#include "Viewport/AvaCineCameraActor.h"
 
 #define LOCTEXT_NAMESPACE "AvaEditorSettings"
 
@@ -31,31 +30,63 @@ void UAvaEditorSettings::OpenEditorSettingsWindow() const
 	SettingsModule.ShowViewer(GetContainerName(), CategoryName, SectionName);
 }
 
-void UAvaEditorSettings::PostInitProperties()
+void UAvaEditorSettings::PostLoad()
 {
-	Super::PostInitProperties();
+	Super::PostLoad();
 
-	if (const FName PresetNameNoLumen = TEXT("No Lumen");
-		!ViewportQualityPresets.Contains(PresetNameNoLumen))
+	if (!ViewportQualityPresets_DEPRECATED.IsEmpty())
 	{
-		ViewportQualityPresets.Add(PresetNameNoLumen, FAvaViewportQualitySettings::Preset(PresetNameNoLumen));
+		ViewportQualitySettingsPresets.Empty(ViewportQualityPresets_DEPRECATED.Num());
+
+		for (const TPair<FName, FAvaViewportQualitySettings>& Preset : ViewportQualityPresets_DEPRECATED)
+		{
+			const FString PresetName = Preset.Key.ToString();
+			const FText PresetText = FText::AsLocalizable_Advanced(TEXT(LOCTEXT_NAMESPACE), PresetName, PresetName);
+			ViewportQualitySettingsPresets.Add(FAvaViewportQualitySettingsPreset(PresetText, Preset.Value));
+		}
+
+		ViewportQualityPresets_DEPRECATED.Empty();
 	}
 
-	if (const FName PresetNameReduced = TEXT("Reduced");
-		!ViewportQualityPresets.Contains(PresetNameReduced))
+	EnsureDefaultViewportQualityPresets();
+	MaintainViewportQualityPresetsIntegrity();
+}
+
+void UAvaEditorSettings::EnsureDefaultViewportQualityPresets()
+{
+	auto AddPresetIfNotFound = [this](const FText& InPresetName)
 	{
-		ViewportQualityPresets.Add(PresetNameReduced, FAvaViewportQualitySettings::Preset(PresetNameReduced));
-	}
+		FAvaViewportQualitySettingsPreset* const FoundPreset = ViewportQualitySettingsPresets.FindByPredicate(
+			[InPresetName](const FAvaViewportQualitySettingsPreset& InPreset)
+			{
+				return InPreset.PresetName.EqualTo(InPresetName);
+			});
+		if (!FoundPreset)
+		{
+			ViewportQualitySettingsPresets.Add(FAvaViewportQualitySettingsPreset(InPresetName
+				, FAvaViewportQualitySettings::Preset(InPresetName)));
+		}
+	};
+
+	AddPresetIfNotFound(FAvaViewportQualitySettingsPreset::NoLumen);
+	AddPresetIfNotFound(FAvaViewportQualitySettingsPreset::Reduced);
+}
+
+void UAvaEditorSettings::MaintainViewportQualityPresetsIntegrity()
+{
+	ViewportQualitySettingsPresets.RemoveAll([](const FAvaViewportQualitySettingsPreset& InPreset)
+		{
+			return InPreset.PresetName.IsEmptyOrWhitespace();
+		});
 
 	// Verifying integrity and sorting will ensure any new/removed entries are handled correctly.
 	DefaultViewportQualitySettings.VerifyIntegrity();
 	DefaultViewportQualitySettings.SortFeaturesByDisplayText();
 
-	for (TPair<FName, FAvaViewportQualitySettings>& Preset : ViewportQualityPresets)
+	for (FAvaViewportQualitySettingsPreset& Preset : ViewportQualitySettingsPresets)
 	{
-		FAvaViewportQualitySettings& Settings = Preset.Value;
-		Settings.VerifyIntegrity();
-		Settings.SortFeaturesByDisplayText();
+		Preset.QualitySettings.VerifyIntegrity();
+		Preset.QualitySettings.SortFeaturesByDisplayText();
 	}
 }
 
diff --git a/Engine/Plugins/Experimental/Avalanche/Source/AvalancheEditor/Private/AvaEditorSettings.h b/Engine/Plugins/Experimental/Avalanche/Source/AvalancheEditor/Private/AvaEditorSettings.h
--- a/Engine/Plugins/Experimental/Avalanche/Source/AvalancheEditor/Private/AvaEditorSettings.h
+++ b/Engine/Plugins/Experimental/Avalanche/Source/AvalancheEditor/Private/AvaEditorSettings.h
@@ -56,13 +56,22 @@ public:
 	UPROPERTY(Config, EditAnywhere, BlueprintReadWrite, Category = "Quality")
 	FAvaViewportQualitySettings DefaultViewportQualitySettings = FAvaViewportQualitySettings(true);
 
+	/** Viewport quality settings user presets. */
+	UPROPERTY(Meta = (DeprecatedProperty, DeprecationMessage = "This property type has changed. Use the updated ViewportQualityPresets instead."))
+	TMap<FName, FAvaViewportQualitySettings> ViewportQualityPresets_DEPRECATED;
+
 	/** Viewport quality settings user presets. */
 	UPROPERTY(Config, EditAnywhere, BlueprintReadWrite, Category = "Quality")
-	TMap<FName, FAvaViewportQualitySettings> ViewportQualityPresets;
+	TArray<FAvaViewportQualitySettingsPreset> ViewportQualitySettingsPresets;
 
 	//~ Begin UObject
-	virtual void PostInitProperties() override;
+	virtual void PostLoad() override;
 	//~ End UObject
 
 	void OpenEditorSettingsWindow() const;
+
+private:
+	void EnsureDefaultViewportQualityPresets();
+
+	void MaintainViewportQualityPresetsIntegrity();
 };
diff --git a/Engine/Plugins/Experimental/Avalanche/Source/AvalancheEditor/Private/DetailView/Customizations/AvaViewportQualitySettingsPropertyTypeCustomization.cpp b/Engine/Plugins/Experimental/Avalanche/Source/AvalancheEditor/Private/DetailView/Customizations/AvaViewportQualitySettingsPropertyTypeCustomization.cpp
--- a/Engine/Plugins/Experimental/Avalanche/Source/AvalancheEditor/Private/DetailView/Customizations/AvaViewportQualitySettingsPropertyTypeCustomization.cpp
+++ b/Engine/Plugins/Experimental/Avalanche/Source/AvalancheEditor/Private/DetailView/Customizations/AvaViewportQualitySettingsPropertyTypeCustomization.cpp
@@ -6,16 +6,21 @@
 #include "DetailWidgetRow.h"
 #include "IDetailChildrenBuilder.h"
 #include "Internationalization/Text.h"
+#include "ISettingsModule.h"
+#include "Modules/ModuleManager.h"
 #include "PropertyCustomizationHelpers.h"
 #include "PropertyHandle.h"
 #include "Styling/AppStyle.h"
 #include "Templates/SharedPointer.h"
 #include "Viewport/AvaViewportQualitySettings.h"
+#include "Widgets/Images/SImage.h"
 #include "Widgets/Input/SButton.h"
+#include "Widgets/Input/SComboButton.h"
 #include "Widgets/Layout/SBorder.h"
 #include "Widgets/Layout/SBox.h"
 #include "Widgets/Layout/SSeparator.h"
 #include "Widgets/Layout/SWrapBox.h"
+#include "Widgets/SBoxPanel.h"
 #include "Widgets/Text/STextBlock.h"
 
 #define LOCTEXT_NAMESPACE "AvaViewportQualitySettingsPropertyTypeCustomization"
@@ -53,9 +58,32 @@ void FAvaViewportQualitySettingsPropertyTypeCustomization::CustomizeChildren(TSh
 				SNew(SBorder)
 				.BorderImage(FAppStyle::GetBrush(TEXT("ToolPanel.GroupBorder")))
 				[
-					SAssignNew(PresetsWrapBox, SWrapBox)
-					.UseAllottedSize(true)
-					.Orientation(EOrientation::Orient_Horizontal)
+					SNew(SHorizontalBox)
+					+ SHorizontalBox::Slot()
+					.AutoWidth()
+					.VAlign(VAlign_Center)
+					[
+						SNew(SComboButton)
+						.ToolTipText(LOCTEXT("EditorSettingToolTip", "Open the Motion Design editor settings to edit viewport quality presets"))
+						.ButtonStyle(FAppStyle::Get(), TEXT("SimpleButton"))
+						.ContentPadding(FMargin(2.f))
+						.HasDownArrow(false)
+						.MenuPlacement(MenuPlacement_BelowAnchor)
+						.OnComboBoxOpened(this, &FAvaViewportQualitySettingsPropertyTypeCustomization::OpenEditorSettings)
+						.ButtonContent()
+						[
+							SNew(SImage)
+							.ColorAndOpacity(FSlateColor::UseForeground())
+							.Image(FAppStyle::Get().GetBrush(TEXT("Icons.Settings")))
+						]
+					]
+					+ SHorizontalBox::Slot()
+					.Padding(FMargin(3.f, 0.f, 0.f, 0.f))
+					[
+						SAssignNew(PresetsWrapBox, SWrapBox)
+						.UseAllottedSize(true)
+						.Orientation(Orient_Horizontal)
+					]
 				]
 			];
 
@@ -103,11 +131,12 @@ void FAvaViewportQualitySettingsPropertyTypeCustomization::RefreshPresets()
 
 	PresetsWrapBox->ClearChildren();
 
-	auto AddSlotToWrapBox = [this](const FName InName, FOnClicked InOnClicked = FOnClicked(), const TAttribute<bool>& InIsEnabled = true)
+	auto AddSlotToWrapBox = [this](const FText& InDisplayName
+		, FOnClicked InOnClicked = FOnClicked(), const TAttribute<bool>& InIsEnabled = true)
 	{
 		if (!InOnClicked.IsBound())
 		{
-			InOnClicked = FOnClicked::CreateSP(this, &FAvaViewportQualitySettingsPropertyTypeCustomization::HandlePresetButtonClick, InName);
+			InOnClicked = FOnClicked::CreateSP(this, &FAvaViewportQualitySettingsPropertyTypeCustomization::HandlePresetButtonClick, InDisplayName);
 		}
 		PresetsWrapBox->AddSlot()
 			[
@@ -115,27 +144,27 @@ void FAvaViewportQualitySettingsPropertyTypeCustomization::RefreshPresets()
 				.Padding(2.f)
 				[
 					SNew(SButton)
-					.ButtonStyle(FAppStyle::Get(), "SimpleButton")
+					.ButtonStyle(FAppStyle::Get(), TEXT("SimpleButton"))
 					.OnClicked(InOnClicked)
 					.IsEnabled(InIsEnabled)
 					[
 						SNew(STextBlock)
 						.TextStyle(FAppStyle::Get(), TEXT("SmallText"))
-						.Text(FText::FromName(InName))
+						.Text(InDisplayName)
 					]
 				]
 			];
 	};
 
-	AddSlotToWrapBox(TEXT("Defaults")
+	AddSlotToWrapBox(LOCTEXT("PresetDefaults", "Defaults")
 		, FOnClicked::CreateSP(this, &FAvaViewportQualitySettingsPropertyTypeCustomization::HandleDefaultsButtonClick)
 		, TAttribute<bool>::CreateSP(this, &FAvaViewportQualitySettingsPropertyTypeCustomization::IsDefaultsButtonEnabled));
 
-	AddSlotToWrapBox(TEXT("All")
+	AddSlotToWrapBox(LOCTEXT("PresetAll", "All")
 		, FOnClicked::CreateSP(this, &FAvaViewportQualitySettingsPropertyTypeCustomization::HandleEnableAllButtonClick)
 		, TAttribute<bool>::CreateSP(this, &FAvaViewportQualitySettingsPropertyTypeCustomization::IsAllButtonEnabled));
 
-	AddSlotToWrapBox(TEXT("None")
+	AddSlotToWrapBox(LOCTEXT("PresetNone", "None")
 		, FOnClicked::CreateSP(this, &FAvaViewportQualitySettingsPropertyTypeCustomization::HandleDisableAllButtonClick)
 		, TAttribute<bool>::CreateSP(this, &FAvaViewportQualitySettingsPropertyTypeCustomization::IsNoneButtonEnabled));
 
@@ -143,14 +172,14 @@ void FAvaViewportQualitySettingsPropertyTypeCustomization::RefreshPresets()
 		.Padding(5.f, 0.f)
 		[
 			SNew(SSeparator)
-			.Orientation(EOrientation::Orient_Vertical)
+			.Orientation(Orient_Vertical)
 		];
 
-	for (const TPair<FName, FAvaViewportQualitySettings>& Preset : UAvaEditorSettings::Get()->ViewportQualityPresets)
+	for (const FAvaViewportQualitySettingsPreset& Preset : UAvaEditorSettings::Get()->ViewportQualitySettingsPresets)
 	{
-		AddSlotToWrapBox(Preset.Key
+		AddSlotToWrapBox(Preset.PresetName
 			, FOnClicked()
-			, TAttribute<bool>::CreateSP(this, &FAvaViewportQualitySettingsPropertyTypeCustomization::IsPresetButtonEnabled, Preset.Key));
+			, TAttribute<bool>::CreateSP(this, &FAvaViewportQualitySettingsPropertyTypeCustomization::IsPresetButtonEnabled, Preset.PresetName));
 	}
 }
 
@@ -191,15 +220,21 @@ FReply FAvaViewportQualitySettingsPropertyTypeCustomization::HandleDisableAllBut
 	return FReply::Handled();
 }
 
-FReply FAvaViewportQualitySettingsPropertyTypeCustomization::HandlePresetButtonClick(const FName InNewPresetName)
+FReply FAvaViewportQualitySettingsPropertyTypeCustomization::HandlePresetButtonClick(const FText InPresetName)
 {
 	UAvaEditorSettings* const EditorSettings = UAvaEditorSettings::Get();
-	if (!EditorSettings->ViewportQualityPresets.Contains(InNewPresetName))
+
+	FAvaViewportQualitySettingsPreset* const FoundPreset = EditorSettings->ViewportQualitySettingsPresets.FindByPredicate(
+		[InPresetName](const FAvaViewportQualitySettingsPreset& InPreset)
+		{
+			return InPreset.PresetName.EqualTo(InPresetName);
+		});
+	if (!FoundPreset)
 	{
 		return FReply::Unhandled();
 	}
 
-	GetStructRef() = EditorSettings->ViewportQualityPresets[InNewPresetName];
+	GetStructRef() = FoundPreset->QualitySettings;
 
 	return FReply::Handled();
 }
@@ -235,19 +270,31 @@ bool FAvaViewportQualitySettingsPropertyTypeCustomization::IsNoneButtonEnabled()
 	return false;
 }
 
-bool FAvaViewportQualitySettingsPropertyTypeCustomization::IsPresetButtonEnabled(const FName InPresetName) const
+bool FAvaViewportQualitySettingsPropertyTypeCustomization::IsPresetButtonEnabled(const FText InPresetName) const
 {
 	const FAvaViewportQualitySettings& StructRef = GetStructRef();
 
-	for (const TPair<FName, FAvaViewportQualitySettings>& Preset : UAvaEditorSettings::Get()->ViewportQualityPresets)
+	for (const FAvaViewportQualitySettingsPreset& Preset : UAvaEditorSettings::Get()->ViewportQualitySettingsPresets)
 	{
-		if (Preset.Key == InPresetName)
+		if (Preset.PresetName.EqualTo(InPresetName))
 		{
-			return Preset.Value != StructRef;
+			return Preset.QualitySettings != StructRef;
 		}
 	}
 
 	return true;
 }
 
+void FAvaViewportQualitySettingsPropertyTypeCustomization::OpenEditorSettings() const
+{
+	if (ISettingsModule* const SettingsModule = FModuleManager::GetModulePtr<ISettingsModule>(TEXT("Settings")))
+	{
+		const UAvaEditorSettings* const Settings = GetDefault<UAvaEditorSettings>();
+		if (ensure(IsValid(Settings)))
+		{
+			SettingsModule->ShowViewer(Settings->GetContainerName(), Settings->GetCategoryName(), Settings->GetSectionName());
+		}
+	}
+}
+
 #undef LOCTEXT_NAMESPACE
diff --git a/Engine/Plugins/Experimental/Avalanche/Source/AvalancheEditor/Private/DetailView/Customizations/AvaViewportQualitySettingsPropertyTypeCustomization.h b/Engine/Plugins/Experimental/Avalanche/Source/AvalancheEditor/Private/DetailView/Customizations/AvaViewportQualitySettingsPropertyTypeCustomization.h
--- a/Engine/Plugins/Experimental/Avalanche/Source/AvalancheEditor/Private/DetailView/Customizations/AvaViewportQualitySettingsPropertyTypeCustomization.h
+++ b/Engine/Plugins/Experimental/Avalanche/Source/AvalancheEditor/Private/DetailView/Customizations/AvaViewportQualitySettingsPropertyTypeCustomization.h
@@ -6,11 +6,8 @@
 #include "Templates/SharedPointer.h"
 #include "Widgets/Layout/SWrapBox.h"
 
-class FName;
 class FReply;
-class FStructOnScope;
 struct FAvaViewportQualitySettings;
-struct FPropertyChangedEvent;
 
 class FAvaViewportQualitySettingsPropertyTypeCustomization : public IPropertyTypeCustomization
 {
@@ -30,12 +27,14 @@ protected:
 	FReply HandleDefaultsButtonClick();
 	FReply HandleEnableAllButtonClick();
 	FReply HandleDisableAllButtonClick();
-	FReply HandlePresetButtonClick(const FName InPresetName);
+	FReply HandlePresetButtonClick(const FText InPresetName);
 
 	bool IsDefaultsButtonEnabled() const;
 	bool IsAllButtonEnabled() const;
 	bool IsNoneButtonEnabled() const;
-	bool IsPresetButtonEnabled(const FName InPresetName) const;
+	bool IsPresetButtonEnabled(const FText InPresetName) const;
+
+	void OpenEditorSettings() const;
 
 	TSharedPtr<IPropertyHandle> StructPropertyHandle;
 
