Motion Design: PIE Camera fix
- When starting PIE, the camera darken overlay and pan/zoom controller are disabled.

diff --git a/Engine/Plugins/Experimental/Avalanche/Source/AvalancheLevelViewport/Private/ViewportClient/AvaLevelViewportClient.cpp b/Engine/Plugins/Experimental/Avalanche/Source/AvalancheLevelViewport/Private/ViewportClient/AvaLevelViewportClient.cpp
--- a/Engine/Plugins/Experimental/Avalanche/Source/AvalancheLevelViewport/Private/ViewportClient/AvaLevelViewportClient.cpp
+++ b/Engine/Plugins/Experimental/Avalanche/Source/AvalancheLevelViewport/Private/ViewportClient/AvaLevelViewportClient.cpp
@@ -453,7 +453,7 @@ FSceneView* FAvaLevelViewportClient::CalcSceneView(FSceneViewFamily* ViewFamily,
 		}
 	}
 
-	const bool bIsZoomPossible = FAvaCameraZoomController::IsCameraZoomPossible();
+	const bool bIsZoomPossible = ZoomController->CanZoom();
 
 	if (!bIsZoomPossible)
 	{
@@ -566,7 +566,7 @@ void FAvaLevelViewportClient::TrackingStopped()
 
 bool FAvaLevelViewportClient::InputKey(const FInputKeyEventArgs& EventArgs)
 {
-	const bool bIsZoomPossible = FAvaCameraZoomController::IsCameraZoomPossible();
+	const bool bIsZoomPossible = ZoomController->CanZoom();
 
 	if (bIsZoomPossible)
 	{
diff --git a/Engine/Plugins/Experimental/Avalanche/Source/AvalancheLevelViewport/Private/Widgets/SAvaLevelViewportCameraBounds.cpp b/Engine/Plugins/Experimental/Avalanche/Source/AvalancheLevelViewport/Private/Widgets/SAvaLevelViewportCameraBounds.cpp
--- a/Engine/Plugins/Experimental/Avalanche/Source/AvalancheLevelViewport/Private/Widgets/SAvaLevelViewportCameraBounds.cpp
+++ b/Engine/Plugins/Experimental/Avalanche/Source/AvalancheLevelViewport/Private/Widgets/SAvaLevelViewportCameraBounds.cpp
@@ -3,6 +3,7 @@
 #include "Widgets/SAvaLevelViewportCameraBounds.h"
 #include "AvaViewportSettings.h"
 #include "AvaVisibleArea.h"
+#include "Editor/EditorEngine.h"
 #include "Interaction/AvaSnapOperation.h"
 #include "ViewportClient/AvaLevelViewportClient.h"
 
@@ -22,8 +23,14 @@ int32 SAvaLevelViewportCameraBounds::OnPaint(const FPaintArgs& InPaintArgs, cons
 	{
 		if (TSharedPtr<FAvaLevelViewportClient> AvaLevelViewportClient = AvaLevelViewportClientWeak.Pin())
 		{
-			DrawCameraBounds(InPaintArgs, InAllottedGeometry, InMyCullingRect, OutDrawElements, InLayerId,
-				AvaViewportSettings->CameraBoundsShadeColor);
+			UWorld* World = AvaLevelViewportClient->GetViewportWorld();
+
+			// Do not draw camera frames during PIE.
+			if (World && World->WorldType == EWorldType::Editor)
+			{
+				DrawCameraBounds(InPaintArgs, InAllottedGeometry, InMyCullingRect, OutDrawElements, InLayerId,
+					AvaViewportSettings->CameraBoundsShadeColor);
+			}
 		}
 	}
 
diff --git a/Engine/Plugins/Experimental/Avalanche/Source/AvalancheViewport/Private/Interaction/AvaCameraZoomController.cpp b/Engine/Plugins/Experimental/Avalanche/Source/AvalancheViewport/Private/Interaction/AvaCameraZoomController.cpp
--- a/Engine/Plugins/Experimental/Avalanche/Source/AvalancheViewport/Private/Interaction/AvaCameraZoomController.cpp
+++ b/Engine/Plugins/Experimental/Avalanche/Source/AvalancheViewport/Private/Interaction/AvaCameraZoomController.cpp
@@ -5,6 +5,7 @@
 #include "AvaViewportUtils.h"
 #include "AvaVisibleArea.h"
 #include "Camera/CameraComponent.h"
+#include "Editor/EditorEngine.h"
 #include "EditorViewportClient.h"
 #include "Framework/Application/SlateApplication.h"
 #include "ViewportClient/IAvaViewportClient.h"
@@ -80,6 +81,23 @@ FAvaCameraZoomController::FAvaCameraZoomController(TSharedRef<IAvaViewportClient
 	Reset();
 }
 
+bool FAvaCameraZoomController::CanZoom() const
+{
+	if (IsCameraZoomPossible())
+	{
+		if (TSharedPtr<IAvaViewportClient> AvaViewportClient = GetViewportClient())
+		{
+			if (UWorld* World = AvaViewportClient->GetViewportWorld())
+			{
+				// Zoom should not be enabled during PIE.
+				return World->WorldType == EWorldType::Editor;
+			}
+		}
+	}
+
+	return false;
+}
+
 void FAvaCameraZoomController::SetZoomLevel(uint8 InZoomLevel)
 {
 	ZoomLevel = FMath::Clamp(InZoomLevel, 0, UE::AvaViewport::Private::ZoomLevels.Num() - 1);
@@ -429,7 +447,7 @@ float FAvaCameraZoomController::GetFOV() const
 {
 	const float DefaultFOV = GetDefaultFOV();
 
-	if (!IsCameraZoomPossible())
+	if (!CanZoom())
 	{
 		return DefaultFOV;
 	}
@@ -444,7 +462,7 @@ float FAvaCameraZoomController::GetFOV() const
 
 FVector2f FAvaCameraZoomController::GetCameraProjectionOffset() const
 {
-	if (!CachedZoomedVisibleArea.IsValid() || !IsCameraZoomPossible())
+	if (!CachedZoomedVisibleArea.IsValid() || !CanZoom())
 	{
 		return FVector2f::ZeroVector;
 	}
@@ -472,7 +490,7 @@ const FAvaVisibleArea& FAvaCameraZoomController::GetCachedVisibleArea() const
 
 const FAvaVisibleArea& FAvaCameraZoomController::GetCachedZoomedVisibleArea() const
 {
-	if (!IsCameraZoomPossible())
+	if (!CanZoom())
 	{
 		return CachedVisibleArea;
 	}
diff --git a/Engine/Plugins/Experimental/Avalanche/Source/AvalancheViewport/Public/Interaction/AvaCameraZoomController.h b/Engine/Plugins/Experimental/Avalanche/Source/AvalancheViewport/Public/Interaction/AvaCameraZoomController.h
--- a/Engine/Plugins/Experimental/Avalanche/Source/AvalancheViewport/Public/Interaction/AvaCameraZoomController.h
+++ b/Engine/Plugins/Experimental/Avalanche/Source/AvalancheViewport/Public/Interaction/AvaCameraZoomController.h
@@ -21,6 +21,8 @@ public:
 		return AvaViewportClientWeak.Pin();
 	}
 
+	AVALANCHEVIEWPORT_API bool CanZoom() const;
+
 	uint8 GetZoomLevel() const
 	{
 		return ZoomLevel;
