Motion Design: added option for user to discard pre-animated state

diff --git a/Engine/Plugins/Experimental/Avalanche/Source/AvalancheSequencer/Private/AvaSequencer.cpp b/Engine/Plugins/Experimental/Avalanche/Source/AvalancheSequencer/Private/AvaSequencer.cpp
--- a/Engine/Plugins/Experimental/Avalanche/Source/AvalancheSequencer/Private/AvaSequencer.cpp
+++ b/Engine/Plugins/Experimental/Avalanche/Source/AvalancheSequencer/Private/AvaSequencer.cpp
@@ -31,6 +31,7 @@
 #include "MVVM/Selection/Selection.h"
 #include "MVVM/ViewModels/SequencerEditorViewModel.h"
 #include "MVVM/Views/SOutlinerView.h"
+#include "Misc/MessageDialog.h"
 #include "Misc/NotifyHook.h"
 #include "Misc/TextFilter.h"
 #include "MovieScene.h"
@@ -221,6 +222,9 @@ void FAvaSequencer::BindCommands()
 		SequencerAction->MapAction(CommandList);
 	}
 
+	CommandList->MapAction(AvaSequencerCommands.ApplyCurrentState
+		, FExecuteAction::CreateRaw(this, &FAvaSequencer::ApplyCurrentState));
+
 	CommandList->MapAction(AvaSequencerCommands.FixBindingPaths
 		, FExecuteAction::CreateRaw(this, &FAvaSequencer::FixBindingPaths));
 
@@ -1136,6 +1140,19 @@ UObject* FAvaSequencer::FindObjectToPossess(UObject* InParentObject, const FMovi
 	return nullptr;
 }
 
+void FAvaSequencer::ApplyCurrentState()
+{
+	const EAppReturnType::Type Response = FMessageDialog::Open(EAppMsgType::YesNo
+		, LOCTEXT("ApplyStateMessage", "Are you sure you want to discard the currently saved pre-animated state, and apply the current state? (This cannot be undone)")
+		, LOCTEXT("ApplyStateTitle", "Apply Current Animated State to World"));
+
+	if (Response == EAppReturnType::Yes)
+	{
+		TSharedRef<ISequencer> Sequencer = GetSequencer();
+		Sequencer->PreAnimatedState.DiscardPreAnimatedState();
+	}
+}
+
 void FAvaSequencer::FixBindingPaths()
 {
 	UAvaSequence* const Sequence   = GetViewedSequence();
diff --git a/Engine/Plugins/Experimental/Avalanche/Source/AvalancheSequencer/Private/AvaSequencer.h b/Engine/Plugins/Experimental/Avalanche/Source/AvalancheSequencer/Private/AvaSequencer.h
--- a/Engine/Plugins/Experimental/Avalanche/Source/AvalancheSequencer/Private/AvaSequencer.h
+++ b/Engine/Plugins/Experimental/Avalanche/Source/AvalancheSequencer/Private/AvaSequencer.h
@@ -163,6 +163,9 @@ public:
  	 */
 	static UObject* FindObjectToPossess(UObject* InParentObject, const FMovieScenePossessable& InPossessable);
 
+	/** Resets the pre-animated state of the sequencer, keeping the current state of the world even when sequencer is saved/closed */
+	void ApplyCurrentState();
+
 	void FixBindingPaths();
 
 	void FixInvalidBindings();
diff --git a/Engine/Plugins/Experimental/Avalanche/Source/AvalancheSequencer/Private/Commands/AvaSequencerCommands.cpp b/Engine/Plugins/Experimental/Avalanche/Source/AvalancheSequencer/Private/Commands/AvaSequencerCommands.cpp
--- a/Engine/Plugins/Experimental/Avalanche/Source/AvalancheSequencer/Private/Commands/AvaSequencerCommands.cpp
+++ b/Engine/Plugins/Experimental/Avalanche/Source/AvalancheSequencer/Private/Commands/AvaSequencerCommands.cpp
@@ -24,6 +24,12 @@ void FAvaSequencerCommands::RegisterCommands()
 		, EUserInterfaceActionType::Button
 		, FInputChord());
 
+	UI_COMMAND(ApplyCurrentState
+		, "Apply Current State"
+		, "Applies the current state of the sequence to the world by resetting the pre-animated state of the sequence."
+		, EUserInterfaceActionType::Button
+		, FInputChord());
+
 	UI_COMMAND(FixBindingPaths
 		, "Fix Binding Paths"
 		, "Fixes the base path of the bindings, that typically breaks when duplicating or renaming an asset"
diff --git a/Engine/Plugins/Experimental/Avalanche/Source/AvalancheSequencer/Private/Commands/AvaSequencerCommands.h b/Engine/Plugins/Experimental/Avalanche/Source/AvalancheSequencer/Private/Commands/AvaSequencerCommands.h
--- a/Engine/Plugins/Experimental/Avalanche/Source/AvalancheSequencer/Private/Commands/AvaSequencerCommands.h
+++ b/Engine/Plugins/Experimental/Avalanche/Source/AvalancheSequencer/Private/Commands/AvaSequencerCommands.h
@@ -24,6 +24,9 @@ public:
 
 	TSharedPtr<FUICommandInfo> StopSelected;
 
+	/** Applies the current state of the sequence to the world by resetting the pre-animated state of the sequence */
+	TSharedPtr<FUICommandInfo> ApplyCurrentState;
+
 	TSharedPtr<FUICommandInfo> FixBindingPaths;
 
 	TSharedPtr<FUICommandInfo> FixInvalidBindings;
diff --git a/Engine/Plugins/Experimental/Avalanche/Source/AvalancheSequencer/Private/Customization/AvaSequenceCustomization.cpp b/Engine/Plugins/Experimental/Avalanche/Source/AvalancheSequencer/Private/Customization/AvaSequenceCustomization.cpp
--- a/Engine/Plugins/Experimental/Avalanche/Source/AvalancheSequencer/Private/Customization/AvaSequenceCustomization.cpp
+++ b/Engine/Plugins/Experimental/Avalanche/Source/AvalancheSequencer/Private/Customization/AvaSequenceCustomization.cpp
@@ -124,6 +124,7 @@ TSharedRef<SWidget> FAvaSequenceCustomization::MakePlaybackMenu() const
 		const FAvaEaseCurveToolCommands& EaseCurveToolCommands = FAvaEaseCurveToolCommands::Get();
 
 		FToolMenuSection& SequencerSection = ToolMenu->FindOrAddSection(TEXT("AvaSequencerActions"), LOCTEXT("AvaSequencerActions", "Sequence"));
+		SequencerSection.AddMenuEntry(AvaSequencerCommands.ApplyCurrentState);
 		SequencerSection.AddMenuEntry(AvaSequencerCommands.OpenStaggerTool);
 
 		FToolMenuSection& EaseCurveToolSection = ToolMenu->FindOrAddSection(TEXT("AvaEaseCurveToolActions"), LOCTEXT("AvaEaseCurveToolActions", "Ease Curve Tool"));
