MotionDesign : Text3D
- Replacing AvaText3D by Text3D in toolbar, setting defaults in factory to match motion design defaults
- Text visualizer applies to Text3D instead of AvaText3D
- AvaText3D is not allowed to be placed in level except if spawned through code
- Adding text sections for details panel (text, geometry, layout, style, rendering)
- Adding character customization for details panel instead of displaying entire array, user can target a specific character based on the index and adjust properties
- Removing unused customization

diff --git a/Engine/Plugins/Experimental/Avalanche/Source/AvalancheText/Private/AvaText3DComponent.cpp b/Engine/Plugins/Experimental/Avalanche/Source/AvalancheText/Private/AvaText3DComponent.cpp
--- a/Engine/Plugins/Experimental/Avalanche/Source/AvalancheText/Private/AvaText3DComponent.cpp
+++ b/Engine/Plugins/Experimental/Avalanche/Source/AvalancheText/Private/AvaText3DComponent.cpp
@@ -21,25 +21,9 @@
 UAvaText3DComponent::UAvaText3DComponent()
 {
 	SetExtrude(0.0f);
-	SetBevel(0.0f);
-	SetBevelType(EText3DBevelType::Convex);
-	SetBevelSegments(8);
-
-	SetHorizontalAlignment(EText3DHorizontalTextAlignment::Left);
-	SetVerticalAlignment(EText3DVerticalTextAlignment::FirstLine);
-
-	SetHasMaxWidth(false);
-	SetHasMaxHeight(false);
 	SetScaleProportionally(false);
-
-	SetKerning(0.f);
-	SetWordSpacing(0.f);
-	SetLineSpacing(0.f);
-
 	SetMaxWidth(100.f);
 	SetMaxHeight(100.f);
-
-	bEnforceUpperCase = false;
 }
 
 void UAvaText3DComponent::Serialize(FArchive& InArchive)
diff --git a/Engine/Plugins/Experimental/Avalanche/Source/AvalancheText/Public/AvaTextActor.h b/Engine/Plugins/Experimental/Avalanche/Source/AvalancheText/Public/AvaTextActor.h
--- a/Engine/Plugins/Experimental/Avalanche/Source/AvalancheText/Public/AvaTextActor.h
+++ b/Engine/Plugins/Experimental/Avalanche/Source/AvalancheText/Public/AvaTextActor.h
@@ -9,7 +9,10 @@
 class UText3DCharacterTransform;
 struct FAvaColorChangeData;
 
-UCLASS(MinimalAPI, ClassGroup = (Text3D), DisplayName = "Motion Design Text", meta = (ComponentWrapperClass))
+/**
+ * This actor is getting replaced by AText3DActor, do not use anymore
+ */
+UCLASS(MinimalAPI, NotPlaceable, ClassGroup = (Text3D), DisplayName = "Motion Design Text", meta = (ComponentWrapperClass))
 class AAvaTextActor : public AActor
 {
 	GENERATED_BODY()
diff --git a/Engine/Plugins/Experimental/Avalanche/Source/AvalancheTextEditor/Private/AvaTextActorFactory.cpp b/Engine/Plugins/Experimental/Avalanche/Source/AvalancheTextEditor/Private/AvaTextActorFactory.cpp
--- a/Engine/Plugins/Experimental/Avalanche/Source/AvalancheTextEditor/Private/AvaTextActorFactory.cpp
+++ b/Engine/Plugins/Experimental/Avalanche/Source/AvalancheTextEditor/Private/AvaTextActorFactory.cpp
@@ -1,13 +1,33 @@
 // Copyright Epic Games, Inc. All Rights Reserved.
 
 #include "AvaTextActorFactory.h"
-#include "AvaTextActor.h"
+
 #include "EngineAnalytics.h"
 #include "Subsystems/PlacementSubsystem.h"
+#include "Text3DActor.h"
+#include "Text3DComponent.h"
 
 UAvaTextActorFactory::UAvaTextActorFactory()
 {
-	NewActorClass = AAvaTextActor::StaticClass();
+	NewActorClass = AText3DActor::StaticClass();
+}
+
+void UAvaTextActorFactory::PostSpawnActor(UObject* InAsset, AActor* InNewActor)
+{
+	Super::PostSpawnActor(InAsset, InNewActor);
+
+	/*
+	 * Text3D is different from AvaText3D, set its properties to expected default values
+	 * Cannot set them on Text3D to avoid affecting other projects
+	 */ 
+	if (const AText3DActor* Text3DActor = Cast<AText3DActor>(InNewActor))
+	{
+		UText3DComponent* Component = Text3DActor->GetText3DComponent();
+		Component->SetExtrude(0.0f);
+		Component->SetScaleProportionally(false);
+		Component->SetMaxWidth(100.f);
+		Component->SetMaxHeight(100.f);
+	}
 }
 
 void UAvaTextActorFactory::PostPlaceAsset(TArrayView<const FTypedElementHandle> InHandle, const FAssetPlacementInfo& InPlacementInfo, const FPlacementOptions& InPlacementOptions)
diff --git a/Engine/Plugins/Experimental/Avalanche/Source/AvalancheTextEditor/Private/AvaTextActorFactory.h b/Engine/Plugins/Experimental/Avalanche/Source/AvalancheTextEditor/Private/AvaTextActorFactory.h
--- a/Engine/Plugins/Experimental/Avalanche/Source/AvalancheTextEditor/Private/AvaTextActorFactory.h
+++ b/Engine/Plugins/Experimental/Avalanche/Source/AvalancheTextEditor/Private/AvaTextActorFactory.h
@@ -14,6 +14,7 @@ public:
 	UAvaTextActorFactory();
 
 	//~ Begin UActorFactory
+	virtual void PostSpawnActor(UObject* InAsset, AActor* InNewActor) override;
 	virtual void PostPlaceAsset(TArrayView<const FTypedElementHandle> InHandle, const FAssetPlacementInfo& InPlacementInfo, const FPlacementOptions& InPlacementOptions) override;
 	//~ End UActorFactory
 };
diff --git a/Engine/Plugins/Experimental/Avalanche/Source/AvalancheTextEditor/Private/AvaTextEditorModule.cpp b/Engine/Plugins/Experimental/Avalanche/Source/AvalancheTextEditor/Private/AvaTextEditorModule.cpp
--- a/Engine/Plugins/Experimental/Avalanche/Source/AvalancheTextEditor/Private/AvaTextEditorModule.cpp
+++ b/Engine/Plugins/Experimental/Avalanche/Source/AvalancheTextEditor/Private/AvaTextEditorModule.cpp
@@ -1,6 +1,7 @@
 // Copyright Epic Games, Inc. All Rights Reserved.
 
 #include "AvaTextEditorModule.h"
+
 #include "AvaInteractiveToolsDelegates.h"
 #include "AvaText3DComponent.h"
 #include "AvaTextActor.h"
@@ -9,7 +10,6 @@
 #include "ColorPicker/AvaViewportColorPickerActorClassRegistry.h"
 #include "ColorPicker/AvaViewportColorPickerAdapter.h"
 #include "DMObjectMaterialProperty.h"
-#include "DetailsView/AvaLinearGradientSettingsCustomization.h"
 #include "Framework/Application/SlateApplication.h"
 #include "IAvalancheComponentVisualizersModule.h"
 #include "IDynamicMaterialEditorModule.h"
@@ -20,10 +20,9 @@
 
 void FAvaTextEditorModule::StartupModule()
 {
-	FCoreDelegates::OnPostEngineInit.AddRaw(this, &FAvaTextEditorModule::PostEngineInit);
+	FCoreDelegates::OnPostEngineInit.AddRaw(this, &FAvaTextEditorModule::RegisterComponentVisualizers);
 
 	FAvaTextEditorCommands::Register();
-	RegisterCustomLayouts();
 
 	RegisterDynamicMaterialPropertyGenerator();
 
@@ -35,49 +34,15 @@ void FAvaTextEditorModule::ShutdownModule()
 	FCoreDelegates::OnPostEngineInit.RemoveAll(this);
 
 	FAvaTextEditorCommands::Unregister();
-
-	if (UObjectInitialized() && !IsEngineExitRequested())
-	{
-		UnregisterCustomLayouts();
-	}
-}
-
-void FAvaTextEditorModule::PostEngineInit()
-{
-	if (FSlateApplication::IsInitialized())
-	{
-		RegisterComponentVisualizers();
-	}
 }
 
 void FAvaTextEditorModule::RegisterComponentVisualizers()
 {
-	IAvalancheComponentVisualizersModule& AvaComponentVisualizerModule = IAvalancheComponentVisualizersModule::Get();
-
-	AvaComponentVisualizerModule.RegisterComponentVisualizer<UAvaText3DComponent, FAvaTextVisualizer>(&Visualizers);
-}
-
-void FAvaTextEditorModule::RegisterCustomLayouts()
-{
-	static FName PropertyEditor("PropertyEditor");
-
-	FPropertyEditorModule& PropertyModule = FModuleManager::GetModuleChecked<FPropertyEditorModule>(PropertyEditor);
-
-	// Text
-	PropertyModule.RegisterCustomPropertyTypeLayout(FAvaLinearGradientSettings::StaticStruct()->GetFName()
-		, FOnGetPropertyTypeCustomizationInstance::CreateStatic(&FAvaLinearGradientSettingsCustomization::MakeInstance));
-}
-
-void FAvaTextEditorModule::UnregisterCustomLayouts()
-{
-	static FName PropertyEditor("PropertyEditor");
-
-	if (FModuleManager::Get().IsModuleLoaded(PropertyEditor))
+	if (FSlateApplication::IsInitialized())
 	{
-		FPropertyEditorModule& PropertyModule = FModuleManager::GetModuleChecked<FPropertyEditorModule>(PropertyEditor);
+		IAvalancheComponentVisualizersModule& AvaComponentVisualizerModule = IAvalancheComponentVisualizersModule::Get();
 
-		// Text
-		PropertyModule.UnregisterCustomPropertyTypeLayout(FAvaLinearGradientSettings::StaticStruct()->GetFName());
+		AvaComponentVisualizerModule.RegisterComponentVisualizer<UText3DComponent, FAvaTextVisualizer>(&Visualizers);
 	}
 }
 
diff --git a/Engine/Plugins/Experimental/Avalanche/Source/AvalancheTextEditor/Private/AvaTextEditorModule.h b/Engine/Plugins/Experimental/Avalanche/Source/AvalancheTextEditor/Private/AvaTextEditorModule.h
--- a/Engine/Plugins/Experimental/Avalanche/Source/AvalancheTextEditor/Private/AvaTextEditorModule.h
+++ b/Engine/Plugins/Experimental/Avalanche/Source/AvalancheTextEditor/Private/AvaTextEditorModule.h
@@ -25,12 +25,7 @@ public:
 private:
 	TArray<TSharedPtr<FComponentVisualizer>> Visualizers;
 
-	void PostEngineInit();
-
 	void RegisterComponentVisualizers();
 
-	void RegisterCustomLayouts();
-	void UnregisterCustomLayouts();
-
 	void RegisterDynamicMaterialPropertyGenerator();
 };
diff --git a/Engine/Plugins/Experimental/Avalanche/Source/AvalancheTextEditor/Private/DetailsView/AvaLinearGradientSettingsCustomization.cpp b/Engine/Plugins/Experimental/Avalanche/Source/AvalancheTextEditor/Private/DetailsView/AvaLinearGradientSettingsCustomization.cpp
deleted file mode 100644
diff --git a/Engine/Plugins/Experimental/Avalanche/Source/AvalancheTextEditor/Private/DetailsView/AvaLinearGradientSettingsCustomization.h b/Engine/Plugins/Experimental/Avalanche/Source/AvalancheTextEditor/Private/DetailsView/AvaLinearGradientSettingsCustomization.h
deleted file mode 100644
diff --git a/Engine/Plugins/Experimental/Avalanche/Source/AvalancheTextEditor/Private/Tool/AvaTextActorTool.cpp b/Engine/Plugins/Experimental/Avalanche/Source/AvalancheTextEditor/Private/Tool/AvaTextActorTool.cpp
--- a/Engine/Plugins/Experimental/Avalanche/Source/AvalancheTextEditor/Private/Tool/AvaTextActorTool.cpp
+++ b/Engine/Plugins/Experimental/Avalanche/Source/AvalancheTextEditor/Private/Tool/AvaTextActorTool.cpp
@@ -2,15 +2,15 @@
 
 #include "Tool/AvaTextActorTool.h"
 #include "AvaInteractiveToolsSettings.h"
-#include "AvaTextActor.h"
 #include "AvaTextActorFactory.h"
 #include "AvaTextEditorCommands.h"
 #include "Builders/AvaInteractiveToolsToolBuilder.h"
 #include "IAvalancheComponentVisualizersModule.h"
+#include "Text3DActor.h"
 
 UAvaTextActorTool::UAvaTextActorTool()
 {
-	ActorClass = AAvaTextActor::StaticClass();
+	ActorClass = AText3DActor::StaticClass();
 }
 
 void UAvaTextActorTool::OnRegisterTool(IAvalancheInteractiveToolsModule* InAITModule)
diff --git a/Engine/Plugins/Experimental/Avalanche/Source/AvalancheTextEditor/Private/Visualizer/AvaTextVisualizer.cpp b/Engine/Plugins/Experimental/Avalanche/Source/AvalancheTextEditor/Private/Visualizer/AvaTextVisualizer.cpp
--- a/Engine/Plugins/Experimental/Avalanche/Source/AvalancheTextEditor/Private/Visualizer/AvaTextVisualizer.cpp
+++ b/Engine/Plugins/Experimental/Avalanche/Source/AvalancheTextEditor/Private/Visualizer/AvaTextVisualizer.cpp
@@ -68,18 +68,21 @@ UActorComponent* FAvaTextVisualizer::GetEditedComponent() const
 
 TMap<UObject*, TArray<FProperty*>> FAvaTextVisualizer::GatherEditableProperties(UObject* InObject) const
 {
-	if (UAvaText3DComponent* Component = Cast<UAvaText3DComponent>(InObject))
+	if (const UText3DComponent* Component = Cast<UText3DComponent>(InObject))
 	{
-		return {{
-			Component,
-			{
-				HasMaxHeightProperty,
-				MaxWidthProperty,
-				HasMaxHeightProperty,
-				MaxHeightProperty,
-				ScaleProportionallyProperty
-			}
-		}};
+		if (UText3DDefaultLayoutExtension* DefaultLayoutExtension = Component->GetCastedLayoutExtension<UText3DDefaultLayoutExtension>())
+		{
+			return {{
+				DefaultLayoutExtension,
+				{
+					HasMaxHeightProperty,
+					MaxWidthProperty,
+					HasMaxHeightProperty,
+					MaxHeightProperty,
+					ScaleProportionallyProperty
+				}
+			}};
+		}
 	}
 
 	return {};
@@ -99,25 +102,25 @@ bool FAvaTextVisualizer::VisProxyHandleClick(FEditorViewportClient* InViewportCl
 		{
 			EndEditing();
 
-			if (UAvaText3DComponent* const AvaTextComponent = Cast<UAvaText3DComponent>(const_cast<UActorComponent*>(MaxHeightProxy->Component.Get())))
+			if (UText3DComponent* const Text3DComponent = Cast<UText3DComponent>(const_cast<UActorComponent*>(MaxHeightProxy->Component.Get())))
 			{
 				FScopedTransaction Transaction(LOCTEXT("ToggleText3DHasMaxHeight", "Toggle Text 3D Has Max Height"));
 
-				if (AActor* const Text3DActor = AvaTextComponent->GetOwner())
+				if (AActor* const Text3DActor = Text3DComponent->GetOwner())
 				{
 					Text3DActor->Modify();
-					AvaTextComponent->Modify();
+					Text3DComponent->Modify();
 
 					// If max height is off, we will toggle it on later.
-					if (!AvaTextComponent->HasMaxHeight() && AvaTextComponent->GetMaxHeight() == GetDefault<UAvaText3DComponent>()->GetMaxHeight())
+					if (!Text3DComponent->HasMaxHeight() && Text3DComponent->GetMaxHeight() == GetDefault<UText3DComponent>()->GetMaxHeight())
 					{
 						// Set a sensible default value
-						Bounds = GetComponentBounds(AvaTextComponent);
+						Bounds = GetComponentBounds(Text3DComponent);
 						float NewMaxHeight = (Bounds.Max.Z - Bounds.Min.Z) * 1.25f;
-						AvaTextComponent->SetMaxHeight(NewMaxHeight);
+						Text3DComponent->SetMaxHeight(NewMaxHeight);
 					}					
 
-					AvaTextComponent->SetHasMaxHeight(!AvaTextComponent->HasMaxHeight());
+					Text3DComponent->SetHasMaxHeight(!Text3DComponent->HasMaxHeight());
 				}
 			}
 
@@ -130,11 +133,11 @@ bool FAvaTextVisualizer::VisProxyHandleClick(FEditorViewportClient* InViewportCl
 		{
 			EndEditing();
 
-			if (UAvaText3DComponent* AvaTextComponent = Cast<UAvaText3DComponent>(const_cast<UActorComponent*>(MaxHeightHandleProxy->Component.Get())))
+			if (UText3DComponent* Text3DComponent = Cast<UText3DComponent>(const_cast<UActorComponent*>(MaxHeightHandleProxy->Component.Get())))
 			{
-				TextComponent = AvaTextComponent;
+				TextComponent = Text3DComponent;
 				bEditingHeight = true;
-				StartEditing(InViewportClient, AvaTextComponent);
+				StartEditing(InViewportClient, Text3DComponent);
 			}
 
 			return true;
@@ -146,25 +149,25 @@ bool FAvaTextVisualizer::VisProxyHandleClick(FEditorViewportClient* InViewportCl
 		{
 			EndEditing();
 
-			if (UAvaText3DComponent* const AvaTextComponent = Cast<UAvaText3DComponent>(const_cast<UActorComponent*>(MaxWidthProxy->Component.Get())))
+			if (UText3DComponent* const Text3DComponent = Cast<UText3DComponent>(const_cast<UActorComponent*>(MaxWidthProxy->Component.Get())))
 			{
 				FScopedTransaction Transaction(LOCTEXT("ToggleText3DHasMaxWidth", "Toggle Text 3D Has Max Width"));
 
-				if (AActor* const Text3DActor = AvaTextComponent->GetOwner())
+				if (AActor* const Text3DActor = Text3DComponent->GetOwner())
 				{
 					Text3DActor->Modify();
-					AvaTextComponent->Modify();
+					Text3DComponent->Modify();
 
 					// If max width is off, we will toggle it on later.
-					if (!AvaTextComponent->HasMaxWidth() && AvaTextComponent->GetMaxWidth() == GetDefault<UAvaText3DComponent>()->GetMaxWidth())
+					if (!Text3DComponent->HasMaxWidth() && Text3DComponent->GetMaxWidth() == GetDefault<UText3DComponent>()->GetMaxWidth())
 					{
 						// Set a sensible default value
-						Bounds = GetComponentBounds(AvaTextComponent);
+						Bounds = GetComponentBounds(Text3DComponent);
 						float NewMaxWidth = (Bounds.Max.Y - Bounds.Min.Y) * 1.25f;
-						AvaTextComponent->SetMaxWidth(NewMaxWidth);
+						Text3DComponent->SetMaxWidth(NewMaxWidth);
 					}
 
-					AvaTextComponent->SetHasMaxWidth(!AvaTextComponent->HasMaxWidth());
+					Text3DComponent->SetHasMaxWidth(!Text3DComponent->HasMaxWidth());
 				}
 			}
 
@@ -177,11 +180,11 @@ bool FAvaTextVisualizer::VisProxyHandleClick(FEditorViewportClient* InViewportCl
 		{
 			EndEditing();
 
-			if (UAvaText3DComponent* const AvaTextComponent = Cast<UAvaText3DComponent>(const_cast<UActorComponent*>(MaxWidthHandleProxy->Component.Get())))
+			if (UText3DComponent* const Text3DComponent = Cast<UText3DComponent>(const_cast<UActorComponent*>(MaxWidthHandleProxy->Component.Get())))
 			{
-				TextComponent = AvaTextComponent;
+				TextComponent = Text3DComponent;
 				bEditingWidth = true;
-				StartEditing(InViewportClient, AvaTextComponent);
+				StartEditing(InViewportClient, Text3DComponent);
 			}
 
 			return true;
@@ -193,15 +196,15 @@ bool FAvaTextVisualizer::VisProxyHandleClick(FEditorViewportClient* InViewportCl
 		{
 			EndEditing();
 
-			if (UAvaText3DComponent* const AvaTextComponent = Cast<UAvaText3DComponent>(const_cast<UActorComponent*>(ScaleProportionallyProxy->Component.Get())))
+			if (UText3DComponent* const Text3DComponent = Cast<UText3DComponent>(const_cast<UActorComponent*>(ScaleProportionallyProxy->Component.Get())))
 			{
 				FScopedTransaction Transaction(LOCTEXT("ToggleText3DScaleProportionally", "Toggle Text 3D Scale Proportionally"));
 
-				if (AActor* const Text3DActor = AvaTextComponent->GetOwner())
+				if (AActor* const Text3DActor = Text3DComponent->GetOwner())
 				{
 					Text3DActor->Modify();
-					AvaTextComponent->Modify();
-					AvaTextComponent->SetScaleProportionally(!AvaTextComponent->ScalesProportionally());
+					Text3DComponent->Modify();
+					Text3DComponent->SetScaleProportionally(!Text3DComponent->ScalesProportionally());
 				}
 			}
 			
@@ -222,13 +225,13 @@ bool FAvaTextVisualizer::VisProxyHandleClick(FEditorViewportClient* InViewportCl
 		{
 			EndEditing();
 
-			if (UAvaText3DComponent* AvaTextComponent = Cast<UAvaText3DComponent>(const_cast<UActorComponent*>(GradientStartHandleProxy->Component.Get())))
+			if (UText3DComponent* Text3DComponent = Cast<UText3DComponent>(const_cast<UActorComponent*>(GradientStartHandleProxy->Component.Get())))
 			{
-				TextComponent = AvaTextComponent;
+				TextComponent = Text3DComponent;
 				bEditingGradientRotation_StartHandle = true;
-				GradientEditBeginLocation_StartHandle = GetGradientStartHandleLocation(AvaTextComponent);
-				GradientEditBeginLocation_Center = GetGradientCenterHandleLocation(AvaTextComponent);
-				StartEditing(InViewportClient, AvaTextComponent);
+				GradientEditBeginLocation_StartHandle = GetGradientStartHandleLocation(Text3DComponent);
+				GradientEditBeginLocation_Center = GetGradientCenterHandleLocation(Text3DComponent);
+				StartEditing(InViewportClient, Text3DComponent);
 			}
 
 			return true;
@@ -240,13 +243,13 @@ bool FAvaTextVisualizer::VisProxyHandleClick(FEditorViewportClient* InViewportCl
 		{
 			EndEditing();
 
-			if (UAvaText3DComponent* AvaTextComponent = Cast<UAvaText3DComponent>(const_cast<UActorComponent*>(GradientEndHandleProxy->Component.Get())))
+			if (UText3DComponent* Text3DComponent = Cast<UText3DComponent>(const_cast<UActorComponent*>(GradientEndHandleProxy->Component.Get())))
 			{
-				TextComponent = AvaTextComponent;
+				TextComponent = Text3DComponent;
 				bEditingGradientRotation_EndHandle = true;
-				GradientEditBeginLocation_EndHandle = GetGradientEndHandleLocation(AvaTextComponent);
-				GradientEditBeginLocation_Center = GetGradientCenterHandleLocation(AvaTextComponent);
-				StartEditing(InViewportClient, AvaTextComponent);
+				GradientEditBeginLocation_EndHandle = GetGradientEndHandleLocation(Text3DComponent);
+				GradientEditBeginLocation_Center = GetGradientCenterHandleLocation(Text3DComponent);
+				StartEditing(InViewportClient, Text3DComponent);
 			}
 
 			return true;
@@ -258,11 +261,11 @@ bool FAvaTextVisualizer::VisProxyHandleClick(FEditorViewportClient* InViewportCl
 		{
 			EndEditing();
 
-			if (UAvaText3DComponent* AvaTextComponent = Cast<UAvaText3DComponent>(const_cast<UActorComponent*>(GradientCenterHandleProxy->Component.Get())))
+			if (UText3DComponent* Text3DComponent = Cast<UText3DComponent>(const_cast<UActorComponent*>(GradientCenterHandleProxy->Component.Get())))
 			{
-				TextComponent = AvaTextComponent;
+				TextComponent = Text3DComponent;
 				bEditingGradientOffset = true;
-				StartEditing(InViewportClient, AvaTextComponent);
+				StartEditing(InViewportClient, Text3DComponent);
 			}
 
 			return true;
@@ -274,11 +277,11 @@ bool FAvaTextVisualizer::VisProxyHandleClick(FEditorViewportClient* InViewportCl
 		{
 			EndEditing();
 
-			if (UAvaText3DComponent* AvaTextComponent = Cast<UAvaText3DComponent>(const_cast<UActorComponent*>(GradientSmoothnessHandleProxy->Component.Get())))
+			if (UText3DComponent* Text3DComponent = Cast<UText3DComponent>(const_cast<UActorComponent*>(GradientSmoothnessHandleProxy->Component.Get())))
 			{
-				TextComponent = AvaTextComponent;
+				TextComponent = Text3DComponent;
 				bEditingGradientSmoothness = true;
-				StartEditing(InViewportClient, AvaTextComponent);
+				StartEditing(InViewportClient, Text3DComponent);
 				return true;
 			}			
 		}
@@ -289,11 +292,11 @@ bool FAvaTextVisualizer::VisProxyHandleClick(FEditorViewportClient* InViewportCl
 		{
 			EndEditing();
 
-			if (UAvaText3DComponent* AvaTextComponent = Cast<UAvaText3DComponent>(const_cast<UActorComponent*>(KerningHandleProxy->Component.Get())))
+			if (UText3DComponent* Text3DComponent = Cast<UText3DComponent>(const_cast<UActorComponent*>(KerningHandleProxy->Component.Get())))
 			{
-				TextComponent = AvaTextComponent;
+				TextComponent = Text3DComponent;
 				EditingKerningIndex = KerningHandleProxy->Index;
-				StartEditing(InViewportClient, AvaTextComponent);
+				StartEditing(InViewportClient, Text3DComponent);
 				return true;
 			}			
 		}
@@ -302,56 +305,56 @@ bool FAvaTextVisualizer::VisProxyHandleClick(FEditorViewportClient* InViewportCl
 	return false;
 }
 
-FVector FAvaTextVisualizer::GetGradientSmoothnessHandleLocation(const UAvaText3DComponent* InAvaTextComponent) const
+FVector FAvaTextVisualizer::GetGradientSmoothnessHandleLocation(const UText3DComponent* InTextComponent) const
 {
 	using namespace UE::Ava::Private;
 	
-	return GetGradientCenterHandleLocation(InAvaTextComponent) + FVector(0, FTextVisualizerStatics::GradientSmoothnessHandleOffset, 0);
+	return GetGradientCenterHandleLocation(InTextComponent) + FVector(0, FTextVisualizerStatics::GradientSmoothnessHandleOffset, 0);
 }
 
 bool FAvaTextVisualizer::GetWidgetLocation(const FEditorViewportClient* ViewportClient, FVector& OutLocation) const
 {
-	if (const UAvaText3DComponent* const AvaTextComponent = Cast<UAvaText3DComponent>(GetEditedComponent()))
+	if (const UText3DComponent* const Text3DComponent = Cast<UText3DComponent>(GetEditedComponent()))
 	{
 		if (bEditingWidth)
 		{
-			OutLocation = GetWidthHandleLocation(AvaTextComponent);
+			OutLocation = GetWidthHandleLocation(Text3DComponent);
 			return true;
 		}
 
 		if (bEditingHeight)
 		{
-			OutLocation = GetHeightHandleLocation(AvaTextComponent);
+			OutLocation = GetHeightHandleLocation(Text3DComponent);
 			return true;
 		}
 
 		if (bEditingGradientRotation_StartHandle)
 		{
-			OutLocation = GetGradientStartHandleLocation(AvaTextComponent);
+			OutLocation = GetGradientStartHandleLocation(Text3DComponent);
 			return true;
 		}
 		
 		if (bEditingGradientRotation_EndHandle)
 		{
-			OutLocation = GetGradientEndHandleLocation(AvaTextComponent);
+			OutLocation = GetGradientEndHandleLocation(Text3DComponent);
 			return true;
 		}
 		
 		if (bEditingGradientOffset)
 		{
-			OutLocation = GetGradientCenterHandleLocation(AvaTextComponent);
+			OutLocation = GetGradientCenterHandleLocation(Text3DComponent);
 			return true;
 		}
 
 		if (bEditingGradientSmoothness)
 		{
-			OutLocation = GetGradientSmoothnessHandleLocation(AvaTextComponent);
+			OutLocation = GetGradientSmoothnessHandleLocation(Text3DComponent);
 			return true;
 		}
 
 		if (EditingKerningIndex != INDEX_NONE)
 		{
-			OutLocation = GetKerningHandleLocation(AvaTextComponent, EditingKerningIndex);
+			OutLocation = GetKerningHandleLocation(Text3DComponent, EditingKerningIndex);
 			return true;
 		}
 	}
@@ -461,7 +464,7 @@ bool FAvaTextVisualizer::HandleInputDeltaInternal(FEditorViewportClient* InViewp
 	
 	if (bEditingWidth)
 	{
-		if (UAvaText3DComponent* const AvaTextComponent = Cast<UAvaText3DComponent>(GetEditedComponent()))
+		if (UText3DComponent* const Text3DComponent = Cast<UText3DComponent>(GetEditedComponent()))
 		{
 			if (GetViewportWidgetMode(InViewportClient) == UE::Widget::WM_Translate)
 			{
@@ -469,7 +472,7 @@ bool FAvaTextVisualizer::HandleInputDeltaInternal(FEditorViewportClient* InViewp
 				{
 					float MaxWidth = InitialMaxWidth;
 
-					switch (AvaTextComponent->GetHorizontalAlignment())
+					switch (Text3DComponent->GetHorizontalAlignment())
 					{
 						case EText3DHorizontalTextAlignment::Left:
 							MaxWidth += InAccumulatedTranslation.Y;
@@ -491,11 +494,11 @@ bool FAvaTextVisualizer::HandleInputDeltaInternal(FEditorViewportClient* InViewp
 						bHasBeenModified = true;
 					}
 
-					if (AActor* const Text3DActor = AvaTextComponent->GetOwner())
+					if (AActor* const Text3DActor = Text3DComponent->GetOwner())
 					{
 						Text3DActor->Modify();
-						AvaTextComponent->Modify();
-						AvaTextComponent->SetMaxWidth(MaxWidth);
+						Text3DComponent->Modify();
+						Text3DComponent->SetMaxWidth(MaxWidth);
 					}
 
 					return true;
@@ -508,7 +511,7 @@ bool FAvaTextVisualizer::HandleInputDeltaInternal(FEditorViewportClient* InViewp
 	}
 	else if (bEditingHeight)
 	{
-		if (UAvaText3DComponent* const AvaTextComponent = Cast<UAvaText3DComponent>(GetEditedComponent()))
+		if (UText3DComponent* const Text3DComponent = Cast<UText3DComponent>(GetEditedComponent()))
 		{
 			if (GetViewportWidgetMode(InViewportClient) == UE::Widget::WM_Translate)
 			{
@@ -516,7 +519,7 @@ bool FAvaTextVisualizer::HandleInputDeltaInternal(FEditorViewportClient* InViewp
 				{
 					float MaxHeight = InitialMaxHeight;
 
-					switch (AvaTextComponent->GetVerticalAlignment())
+					switch (Text3DComponent->GetVerticalAlignment())
 					{
 						case EText3DVerticalTextAlignment::FirstLine:
 							MaxHeight += InAccumulatedTranslation.Z / TopHeightFraction;
@@ -542,11 +545,11 @@ bool FAvaTextVisualizer::HandleInputDeltaInternal(FEditorViewportClient* InViewp
 						bHasBeenModified = true;
 					}
 
-					if (AActor* const Text3DActor = AvaTextComponent->GetOwner())
+					if (AActor* const Text3DActor = Text3DComponent->GetOwner())
 					{
 						Text3DActor->Modify();
-						AvaTextComponent->Modify();
-						AvaTextComponent->SetMaxHeight(MaxHeight);
+						Text3DComponent->Modify();
+						Text3DComponent->SetMaxHeight(MaxHeight);
 					}
 
 					return true;
@@ -559,7 +562,7 @@ bool FAvaTextVisualizer::HandleInputDeltaInternal(FEditorViewportClient* InViewp
 	}
 	else if (bEditingGradientRotation_EndHandle || bEditingGradientRotation_StartHandle)
 	{
-		if (UAvaText3DComponent* const AvaTextComponent = Cast<UAvaText3DComponent>(GetEditedComponent()))
+		if (UText3DComponent* const Text3DComponent = Cast<UText3DComponent>(GetEditedComponent()))
 		{
 			if (GetViewportWidgetMode(InViewportClient) == UE::Widget::WM_Translate)
 			{
@@ -591,7 +594,7 @@ bool FAvaTextVisualizer::HandleInputDeltaInternal(FEditorViewportClient* InViewp
 						bHasBeenModified = true;
 					}
 					
-					if (UText3DDefaultMaterialExtension* MaterialExtension = AvaTextComponent->GetCastedMaterialExtension<UText3DDefaultMaterialExtension>())
+					if (UText3DDefaultMaterialExtension* MaterialExtension = Text3DComponent->GetCastedMaterialExtension<UText3DDefaultMaterialExtension>())
 					{
 						MaterialExtension->Modify();
 						MaterialExtension->SetGradientRotation(GradientRot);
@@ -606,14 +609,14 @@ bool FAvaTextVisualizer::HandleInputDeltaInternal(FEditorViewportClient* InViewp
 	}
 	else if (bEditingGradientOffset)
 	{
-		if (UAvaText3DComponent* const AvaTextComponent = Cast<UAvaText3DComponent>(GetEditedComponent()))
+		if (UText3DComponent* const Text3DComponent = Cast<UText3DComponent>(GetEditedComponent()))
 		{
 			if (GetViewportWidgetMode(InViewportClient) == UE::Widget::WM_Translate)
 			{
 				if (GetViewportWidgetAxisList(InViewportClient) & EAxisList::YZ)
 				{
-					const AActor* Text3DActor = AvaTextComponent->GetOwner();
-					UText3DDefaultMaterialExtension* MaterialExtension = AvaTextComponent->GetCastedMaterialExtension<UText3DDefaultMaterialExtension>();
+					const AActor* Text3DActor = Text3DComponent->GetOwner();
+					UText3DDefaultMaterialExtension* MaterialExtension = Text3DComponent->GetCastedMaterialExtension<UText3DDefaultMaterialExtension>();
 
 					if (Text3DActor && MaterialExtension)
 					{
@@ -649,13 +652,13 @@ bool FAvaTextVisualizer::HandleInputDeltaInternal(FEditorViewportClient* InViewp
 	}
 	else if (bEditingGradientSmoothness)
 	{
-		if (UAvaText3DComponent* const AvaTextComponent = Cast<UAvaText3DComponent>(GetEditedComponent()))
+		if (UText3DComponent* const Text3DComponent = Cast<UText3DComponent>(GetEditedComponent()))
 		{
 			if (GetViewportWidgetMode(InViewportClient) == UE::Widget::WM_Scale)
 			{
 				if (GetViewportWidgetAxisList(InViewportClient) & EAxisList::Z)
 				{
-					if (UText3DDefaultMaterialExtension* MaterialExtension = AvaTextComponent->GetCastedMaterialExtension<UText3DDefaultMaterialExtension>())
+					if (UText3DDefaultMaterialExtension* MaterialExtension = Text3DComponent->GetCastedMaterialExtension<UText3DDefaultMaterialExtension>())
 					{
 						MaterialExtension->SetFlags(RF_Transactional);
 						MaterialExtension->Modify();
@@ -683,13 +686,13 @@ bool FAvaTextVisualizer::HandleInputDeltaInternal(FEditorViewportClient* InViewp
 	}
 	else if (EditingKerningIndex != INDEX_NONE)
 	{
-		if (const UAvaText3DComponent* const AvaTextComponent = Cast<UAvaText3DComponent>(GetEditedComponent()))
+		if (const UText3DComponent* const Text3DComponent = Cast<UText3DComponent>(GetEditedComponent()))
 		{
 			if (GetViewportWidgetMode(InViewportClient) == UE::Widget::WM_Translate)
 			{
 				if (GetViewportWidgetAxisList(InViewportClient) & EAxisList::Y)
 				{
-					if (UText3DDefaultCharacter* Character = AvaTextComponent->GetCastedCharacter<UText3DDefaultCharacter>(EditingKerningIndex))
+					if (UText3DDefaultCharacter* Character = Text3DComponent->GetCastedCharacter<UText3DDefaultCharacter>(EditingKerningIndex))
 					{
 						Character->SetFlags(RF_Transactional);
 						Character->Modify();
@@ -720,12 +723,12 @@ void FAvaTextVisualizer::TrackingStopped(FEditorViewportClient* InViewportClient
 {
 	if (bEditingGradientRotation_EndHandle || bEditingGradientRotation_StartHandle)
 	{
-		if (const UAvaText3DComponent* const AvaTextComponent = Cast<UAvaText3DComponent>(GetEditedComponent()))
+		if (const UText3DComponent* const Text3DComponent = Cast<UText3DComponent>(GetEditedComponent()))
 		{
 			// handles positions need to be refreshed every time we release the mouse, so that clockwise vs. counterclockwise interaction works properly
-			GradientEditBeginLocation_StartHandle = GetGradientStartHandleLocation(AvaTextComponent);
-			GradientEditBeginLocation_EndHandle = GetGradientEndHandleLocation(AvaTextComponent);
-			GradientEditBeginLocation_Center = GetGradientCenterHandleLocation(AvaTextComponent);
+			GradientEditBeginLocation_StartHandle = GetGradientStartHandleLocation(Text3DComponent);
+			GradientEditBeginLocation_EndHandle = GetGradientEndHandleLocation(Text3DComponent);
+			GradientEditBeginLocation_Center = GetGradientCenterHandleLocation(Text3DComponent);
 		}
 	}
 
@@ -742,49 +745,49 @@ bool FAvaTextVisualizer::ResetValue(FEditorViewportClient* InViewportClient, HHi
 		return Super::ResetValue(InViewportClient, HitProxy);
 	}
 
-	UAvaText3DComponent* AvaTextComponent = TextComponent.Get();
+	UText3DComponent* Text3DComponent = TextComponent.Get();
 
-	if (!AvaTextComponent)
+	if (!Text3DComponent)
 	{
 		return Super::ResetValue(InViewportClient, HitProxy);
 	}
 
-	AvaTextComponent->SetFlags(RF_Transactional);
+	Text3DComponent->SetFlags(RF_Transactional);
 
 	if (HitProxy->IsA(HAvaTextMaxTextWidthProxy::StaticGetType()))
 	{
 		FScopedTransaction Transaction(NSLOCTEXT("AvaText3DVisualizer", "VisualizerResetValue", "Visualizer Reset Value"));
-		AvaTextComponent->Modify();
+		Text3DComponent->Modify();
 
-		AvaTextComponent->SetHasMaxWidth(false);
-		AvaTextComponent->SetMaxWidth(GetDefault<UText3DDefaultLayoutExtension>()->GetMaxWidth());
+		Text3DComponent->SetHasMaxWidth(false);
+		Text3DComponent->SetMaxWidth(GetDefault<UText3DDefaultLayoutExtension>()->GetMaxWidth());
 
-		FAvaVisualizerBase::NotifyPropertyModified(AvaTextComponent, MaxWidthProperty, EPropertyChangeType::ValueSet);
+		FAvaVisualizerBase::NotifyPropertyModified(Text3DComponent, MaxWidthProperty, EPropertyChangeType::ValueSet);
 	}
 	else if (HitProxy->IsA(HAvaTextMaxTextHeightProxy::StaticGetType()))
 	{
 		FScopedTransaction Transaction(NSLOCTEXT("AvaTextVisualizer", "VisualizerResetValue", "Visualizer Reset Value"));
-		AvaTextComponent->Modify();
+		Text3DComponent->Modify();
 
-		AvaTextComponent->SetHasMaxHeight(false);
-		AvaTextComponent->SetMaxHeight(GetDefault<UText3DDefaultLayoutExtension>()->GetMaxHeight());
+		Text3DComponent->SetHasMaxHeight(false);
+		Text3DComponent->SetMaxHeight(GetDefault<UText3DDefaultLayoutExtension>()->GetMaxHeight());
 
-		FAvaVisualizerBase::NotifyPropertyModified(AvaTextComponent, MaxHeightProperty, EPropertyChangeType::ValueSet);
+		FAvaVisualizerBase::NotifyPropertyModified(Text3DComponent, MaxHeightProperty, EPropertyChangeType::ValueSet);
 	}
 	else if (HitProxy->IsA(HAvaTextScaleProportionallyProxy::StaticGetType()))
 	{
 		FScopedTransaction Transaction(NSLOCTEXT("AvaTextVisualizer", "VisualizerResetValue", "Visualizer Reset Value"));
-		AvaTextComponent->Modify();
+		Text3DComponent->Modify();
 
-		AvaTextComponent->SetScaleProportionally(false);
+		Text3DComponent->SetScaleProportionally(false);
 
-		FAvaVisualizerBase::NotifyPropertyModified(AvaTextComponent, ScaleProportionallyProperty, EPropertyChangeType::ValueSet);
+		FAvaVisualizerBase::NotifyPropertyModified(Text3DComponent, ScaleProportionallyProperty, EPropertyChangeType::ValueSet);
 	}
 	else if (HitProxy->IsA(HAvaTextCharacterKerningHandleProxy::StaticGetType()))
 	{
 		FScopedTransaction Transaction(NSLOCTEXT("AvaTextVisualizer", "VisualizerResetValue", "Visualizer Reset Value"));
 
-		if (UText3DDefaultCharacter* Character = AvaTextComponent->GetCastedCharacter<UText3DDefaultCharacter>(EditingKerningIndex))
+		if (UText3DDefaultCharacter* Character = Text3DComponent->GetCastedCharacter<UText3DDefaultCharacter>(EditingKerningIndex))
 		{
 			Character->Modify();
 			const UText3DDefaultCharacter* CDO = GetDefault<UText3DDefaultCharacter>();
@@ -889,13 +892,13 @@ FVector FAvaTextVisualizer::GetHeightHandleLocation(const UText3DComponent* InTe
 	return IconLocation;
 }
 
-void FAvaTextVisualizer::GetTextActorGradientControlsLocations(const UAvaText3DComponent* InAvaTextComponent, FVector& OutGradientCenterLocation, FVector& OutGradientStartLocation, FVector& OutGradientEndLocation) const
+void FAvaTextVisualizer::GetTextActorGradientControlsLocations(const UText3DComponent* InTextComponent, FVector& OutGradientCenterLocation, FVector& OutGradientStartLocation, FVector& OutGradientEndLocation) const
 {
 	using namespace UE::Ava::Private;
-	if (IsValid(InAvaTextComponent))
+	if (IsValid(InTextComponent))
 	{
-		const AActor* const Text3DActor = InAvaTextComponent->GetOwner();
-		const UText3DDefaultMaterialExtension* MaterialExtension = InAvaTextComponent->GetCastedMaterialExtension<UText3DDefaultMaterialExtension>();
+		const AActor* const Text3DActor = InTextComponent->GetOwner();
+		const UText3DDefaultMaterialExtension* MaterialExtension = InTextComponent->GetCastedMaterialExtension<UText3DDefaultMaterialExtension>();
 
 		if (Text3DActor && MaterialExtension)
 		{
@@ -919,9 +922,9 @@ void FAvaTextVisualizer::GetTextActorGradientControlsLocations(const UAvaText3DC
 	}
 }
 
-FVector FAvaTextVisualizer::GetGradientEndHandleLocation(const UAvaText3DComponent* InAvaTextComponent) const
+FVector FAvaTextVisualizer::GetGradientEndHandleLocation(const UText3DComponent* InTextComponent) const
 {
-	if (IsValid(InAvaTextComponent))
+	if (IsValid(InTextComponent))
 	{
 		FVector GradientCenter;
 		FVector GradientStart;
@@ -935,9 +938,9 @@ FVector FAvaTextVisualizer::GetGradientEndHandleLocation(const UAvaText3DCompone
 	return FVector::Zero();
 }
 
-FVector FAvaTextVisualizer::GetGradientCenterHandleLocation(const UAvaText3DComponent* InAvaTextComponent) const
+FVector FAvaTextVisualizer::GetGradientCenterHandleLocation(const UText3DComponent* InTextComponent) const
 {
-	if (IsValid(InAvaTextComponent))
+	if (IsValid(InTextComponent))
 	{
 		FVector GradientCenter;
 		FVector GradientStart;
@@ -951,9 +954,9 @@ FVector FAvaTextVisualizer::GetGradientCenterHandleLocation(const UAvaText3DComp
 	return FVector::Zero();
 }
 
-FVector FAvaTextVisualizer::GetGradientStartHandleLocation(const UAvaText3DComponent* InAvaTextComponent) const
+FVector FAvaTextVisualizer::GetGradientStartHandleLocation(const UText3DComponent* InTextComponent) const
 {
-	if (IsValid(InAvaTextComponent))
+	if (IsValid(InTextComponent))
 	{
 		FVector GradientCenter;
 		FVector GradientStart;
@@ -967,7 +970,7 @@ FVector FAvaTextVisualizer::GetGradientStartHandleLocation(const UAvaText3DCompo
 	return FVector::Zero();
 }
 
-FVector FAvaTextVisualizer::GetKerningHandleLocation(const UAvaText3DComponent* InTextComponent, int32 InIndex) const
+FVector FAvaTextVisualizer::GetKerningHandleLocation(const UText3DComponent* InTextComponent, int32 InIndex) const
 {
 	if (IsValid(InTextComponent))
 	{
@@ -980,7 +983,7 @@ FVector FAvaTextVisualizer::GetKerningHandleLocation(const UAvaText3DComponent*
 				CharacterLocation *= LayoutExtension->GetTextScale();
 			}
 
-			return InTextComponent->GetComponentTransform().TransformPositionNoScale(CharacterLocation);
+			return InTextComponent->GetComponentTransform().TransformPosition(CharacterLocation);
 		}
 	}
 
@@ -1001,53 +1004,53 @@ void FAvaTextVisualizer::DrawVisualizationNotEditing(const UActorComponent* InCo
 {
 	Super::DrawVisualizationNotEditing(InComponent, InView, InPDI, InOutIconIndex);
 
-	const UAvaText3DComponent* const AvaTextComponent = Cast<UAvaText3DComponent>(InComponent);
+	const UText3DComponent* const Text3DComponent = Cast<UText3DComponent>(InComponent);
 
-	if (!AvaTextComponent)
+	if (!Text3DComponent)
 	{
 		return;
 	}
 	
-	StoreTextMetrics(AvaTextComponent);
+	StoreTextMetrics(Text3DComponent);
 
-	DrawMaxTextSizeVisualization(AvaTextComponent, InView, InPDI);
+	DrawMaxTextSizeVisualization(Text3DComponent, InView, InPDI);
 
-	DrawMaxTextWidthButton(AvaTextComponent, InView, InPDI, InOutIconIndex, AvaTextComponent->HasMaxWidth() ? Enabled : Disabled);
+	DrawMaxTextWidthButton(Text3DComponent, InView, InPDI, InOutIconIndex, Text3DComponent->HasMaxWidth() ? Enabled : Disabled);
 	++InOutIconIndex;
 
-	DrawMaxTextHeightButton(AvaTextComponent, InView, InPDI, InOutIconIndex, AvaTextComponent->HasMaxHeight() ? Enabled : Disabled);
+	DrawMaxTextHeightButton(Text3DComponent, InView, InPDI, InOutIconIndex, Text3DComponent->HasMaxHeight() ? Enabled : Disabled);
 	++InOutIconIndex;
 
-	DrawScaleProportionallyButton(AvaTextComponent, InView, InPDI, InOutIconIndex, AvaTextComponent->ScalesProportionally() ? Enabled : Disabled);
+	DrawScaleProportionallyButton(Text3DComponent, InView, InPDI, InOutIconIndex, Text3DComponent->ScalesProportionally() ? Enabled : Disabled);
 	++InOutIconIndex;
 
-	if (AvaTextComponent->HasMaxWidth())
+	if (Text3DComponent->HasMaxWidth())
 	{
-		DrawMaxTextWidthHandle(AvaTextComponent, InView, InPDI, Inactive);
+		DrawMaxTextWidthHandle(Text3DComponent, InView, InPDI, Inactive);
 	}
 
-	if (AvaTextComponent->HasMaxHeight())
+	if (Text3DComponent->HasMaxHeight())
 	{
-		DrawMaxTextHeightHandle(AvaTextComponent, InView, InPDI, Inactive);
+		DrawMaxTextHeightHandle(Text3DComponent, InView, InPDI, Inactive);
 	}
 	
-	if (const UText3DDefaultMaterialExtension* MaterialExtension = AvaTextComponent->GetCastedMaterialExtension<UText3DDefaultMaterialExtension>())
+	if (const UText3DDefaultMaterialExtension* MaterialExtension = Text3DComponent->GetCastedMaterialExtension<UText3DDefaultMaterialExtension>())
 	{
 		if (MaterialExtension->GetStyle() == EText3DMaterialStyle::Gradient)
 		{
-			DrawEditGradientButton(AvaTextComponent, InView, InPDI, InOutIconIndex, bShowGradientControls ? Enabled : Disabled);
+			DrawEditGradientButton(Text3DComponent, InView, InPDI, InOutIconIndex, bShowGradientControls ? Enabled : Disabled);
 
 			if (bShowGradientControls)
 			{
-				DrawGradientHandles(AvaTextComponent, InView, InPDI);
+				DrawGradientHandles(Text3DComponent, InView, InPDI);
 			}
 		}
 	}
 
-	AvaTextComponent->ForEachCharacter([this, AvaTextComponent, InView, InPDI](UText3DCharacterBase* InCharacter, uint16 InIndex, uint16)
+	Text3DComponent->ForEachCharacter([this, Text3DComponent, InView, InPDI](UText3DCharacterBase* InCharacter, uint16 InIndex, uint16)
 	{
-		const FVector CharacterLocation = GetKerningHandleLocation(AvaTextComponent, InIndex);
-		DrawCharacterKerningHandle(AvaTextComponent, InView, InPDI, CharacterLocation, InIndex, Inactive);
+		const FVector CharacterLocation = GetKerningHandleLocation(Text3DComponent, InIndex);
+		DrawCharacterKerningHandle(Text3DComponent, InView, InPDI, CharacterLocation, InIndex, Inactive);
 	});
 }
 
@@ -1056,34 +1059,34 @@ void FAvaTextVisualizer::DrawVisualizationEditing(const UActorComponent* InCompo
 {
 	Super::DrawVisualizationNotEditing(InComponent, InView, InPDI, InOutIconIndex);
 
-	const UAvaText3DComponent* const AvaTextComponent = Cast<UAvaText3DComponent>(InComponent);
+	const UText3DComponent* const Text3DComponent = Cast<UText3DComponent>(InComponent);
 
-	if (!AvaTextComponent)
+	if (!Text3DComponent)
 	{
 		return;
 	}
 
-	StoreTextMetrics(AvaTextComponent);
+	StoreTextMetrics(Text3DComponent);
 
-	DrawMaxTextSizeVisualization(AvaTextComponent, InView, InPDI);
+	DrawMaxTextSizeVisualization(Text3DComponent, InView, InPDI);
 
 	if (bEditingWidth)
 	{
-		DrawMaxTextWidthButton(AvaTextComponent, InView, InPDI, InOutIconIndex, AvaTextComponent->HasMaxWidth() ? Enabled : Disabled);
-		DrawMaxTextWidthHandle(AvaTextComponent, InView, InPDI, Inactive);
+		DrawMaxTextWidthButton(Text3DComponent, InView, InPDI, InOutIconIndex, Text3DComponent->HasMaxWidth() ? Enabled : Disabled);
+		DrawMaxTextWidthHandle(Text3DComponent, InView, InPDI, Inactive);
 		++InOutIconIndex;
 	}
 	
 	if (bEditingHeight)
 	{
-		DrawMaxTextHeightButton(AvaTextComponent, InView, InPDI, InOutIconIndex, AvaTextComponent->HasMaxHeight() ? Enabled : Disabled);
-		DrawMaxTextHeightHandle(AvaTextComponent, InView, InPDI, Inactive);
+		DrawMaxTextHeightButton(Text3DComponent, InView, InPDI, InOutIconIndex, Text3DComponent->HasMaxHeight() ? Enabled : Disabled);
+		DrawMaxTextHeightHandle(Text3DComponent, InView, InPDI, Inactive);
 		++InOutIconIndex;
 	}
 
 	if (bEditingWidth || bEditingHeight)
 	{
-		DrawScaleProportionallyButton(AvaTextComponent, InView, InPDI, InOutIconIndex, AvaTextComponent->ScalesProportionally() ? Enabled : Disabled);
+		DrawScaleProportionallyButton(Text3DComponent, InView, InPDI, InOutIconIndex, Text3DComponent->ScalesProportionally() ? Enabled : Disabled);
 		++InOutIconIndex;
 	}
 
@@ -1093,12 +1096,12 @@ void FAvaTextVisualizer::DrawVisualizationEditing(const UActorComponent* InCompo
 		{
 			if (MaterialExtension->GetStyle() == EText3DMaterialStyle::Gradient)
 			{
-				DrawEditGradientButton(AvaTextComponent, InView, InPDI, InOutIconIndex, bShowGradientControls ? Enabled : Disabled);
+				DrawEditGradientButton(Text3DComponent, InView, InPDI, InOutIconIndex, bShowGradientControls ? Enabled : Disabled);
 				++InOutIconIndex;
 
 				if (bShowGradientControls)
 				{
-					DrawGradientHandles(AvaTextComponent, InView, InPDI);
+					DrawGradientHandles(Text3DComponent, InView, InPDI);
 				}
 			}
 		}
@@ -1106,10 +1109,10 @@ void FAvaTextVisualizer::DrawVisualizationEditing(const UActorComponent* InCompo
 
 	if (EditingKerningIndex != INDEX_NONE)
 	{
-		AvaTextComponent->ForEachCharacter([this, AvaTextComponent, InView, InPDI](UText3DCharacterBase* InCharacter, uint16 InIndex, uint16)
+		Text3DComponent->ForEachCharacter([this, Text3DComponent, InView, InPDI](UText3DCharacterBase* InCharacter, uint16 InIndex, uint16)
 		{
-			const FVector CharacterLocation = GetKerningHandleLocation(AvaTextComponent, InIndex);
-			DrawCharacterKerningHandle(AvaTextComponent, InView, InPDI, CharacterLocation, InIndex, InIndex == EditingKerningIndex ? Active : Inactive);
+			const FVector CharacterLocation = GetKerningHandleLocation(Text3DComponent, InIndex);
+			DrawCharacterKerningHandle(Text3DComponent, InView, InPDI, CharacterLocation, InIndex, InIndex == EditingKerningIndex ? Active : Inactive);
 		});
 	}
 }
@@ -1208,10 +1211,10 @@ bool FAvaTextVisualizer::DrawGradientSmoothnessHandle(const UText3DComponent* In
 	return false;
 }
 
-void FAvaTextVisualizer::DrawGradientHandles(const UAvaText3DComponent* InAvaTextComponent, const FSceneView* InView, FPrimitiveDrawInterface* InPDI) const
+void FAvaTextVisualizer::DrawGradientHandles(const UText3DComponent* InTextComponent, const FSceneView* InView, FPrimitiveDrawInterface* InPDI) const
 {
 	using namespace UE::Ava::Private;
-	if (!InAvaTextComponent)
+	if (!InTextComponent)
 	{
 		return;
 	}
@@ -1220,18 +1223,18 @@ void FAvaTextVisualizer::DrawGradientHandles(const UAvaText3DComponent* InAvaTex
 	FVector GradientLineStart;
 	FVector GradientLineEnd;
 
-	GetTextActorGradientControlsLocations(InAvaTextComponent, GradientCenter, GradientLineStart, GradientLineEnd);
+	GetTextActorGradientControlsLocations(InTextComponent, GradientCenter, GradientLineStart, GradientLineEnd);
 	
 	InPDI->DrawLine(GradientLineStart, GradientLineEnd, FLinearColor::White,SDPG_Foreground );
 
 	if (UText3DDefaultMaterialExtension* MaterialExtension = TextComponent->GetCastedMaterialExtension<UText3DDefaultMaterialExtension>())
 	{
-		DrawGradientLineStartHandle(InAvaTextComponent, InView, InPDI, GradientLineStart, MaterialExtension->GetGradientColorB());
-		DrawGradientLineEndHandle(InAvaTextComponent, InView, InPDI, GradientLineEnd, MaterialExtension->GetGradientColorA());
+		DrawGradientLineStartHandle(InTextComponent, InView, InPDI, GradientLineStart, MaterialExtension->GetGradientColorB());
+		DrawGradientLineEndHandle(InTextComponent, InView, InPDI, GradientLineEnd, MaterialExtension->GetGradientColorA());
 	}
 
-	DrawGradientCenterHandle(InAvaTextComponent, InView, InPDI, GradientCenter);
-	DrawGradientSmoothnessHandle(InAvaTextComponent, InView, InPDI, GradientCenter + FVector(0, FTextVisualizerStatics::GradientSmoothnessHandleOffset, 0));
+	DrawGradientCenterHandle(InTextComponent, InView, InPDI, GradientCenter);
+	DrawGradientSmoothnessHandle(InTextComponent, InView, InPDI, GradientCenter + FVector(0, FTextVisualizerStatics::GradientSmoothnessHandleOffset, 0));
 	
 }
 
@@ -1451,12 +1454,12 @@ void FAvaTextVisualizer::DrawMaxTextHeightHandle(const UText3DComponent* InTextC
 	InPDI->SetHitProxy(nullptr);
 }
 
-void FAvaTextVisualizer::DrawGradientLineStartHandle(const UAvaText3DComponent* InAvaTextComponent, const FSceneView* InView, 
+void FAvaTextVisualizer::DrawGradientLineStartHandle(const UText3DComponent* InTextComponent, const FSceneView* InView, 
 	FPrimitiveDrawInterface* InPDI, const FVector& InLocation, const FLinearColor& InColor) const
 {
 	static constexpr float BaseSize = 1.5f;
 
-	if (!InAvaTextComponent)
+	if (!InTextComponent)
 	{
 		return;
 	}
@@ -1472,9 +1475,9 @@ void FAvaTextVisualizer::DrawGradientLineStartHandle(const UAvaText3DComponent*
 	const float IconSize = BaseSize * GetIconSizeScale(InView, IconLocation);
 
 	// in case gradient center handle is on top of this handle, or almost on top of it, don't enable interaction
-	if (FVector::Distance(InLocation, GetGradientCenterHandleLocation(InAvaTextComponent)) > BaseSize/2.0)
+	if (FVector::Distance(InLocation, GetGradientCenterHandleLocation(InTextComponent)) > BaseSize/2.0)
 	{
-		InPDI->SetHitProxy(new HAvaTextGradientLineStartHandleProxy(InAvaTextComponent));
+		InPDI->SetHitProxy(new HAvaTextGradientLineStartHandleProxy(InTextComponent));
 	}
 	
 	InPDI->DrawSprite(IconLocation, IconSize, IconSize, UVSprite->GetResource(), InColor,
@@ -1482,12 +1485,12 @@ void FAvaTextVisualizer::DrawGradientLineStartHandle(const UAvaText3DComponent*
 	InPDI->SetHitProxy(nullptr);
 }
 
-void FAvaTextVisualizer::DrawGradientLineEndHandle(const UAvaText3DComponent* InAvaTextComponent, const FSceneView* InView, 
+void FAvaTextVisualizer::DrawGradientLineEndHandle(const UText3DComponent* InTextComponent, const FSceneView* InView, 
 	FPrimitiveDrawInterface* InPDI, const FVector& InLocation, const FLinearColor& InColor) const
 {
 	static constexpr float BaseSize = 1.5f;
 
-	if (!InAvaTextComponent)
+	if (!InTextComponent)
 	{
 		return;
 	}
@@ -1503,9 +1506,9 @@ void FAvaTextVisualizer::DrawGradientLineEndHandle(const UAvaText3DComponent* In
 	const float IconSize = BaseSize * GetIconSizeScale(InView, IconLocation);
 
 	// in case gradient center handle is on top of this handle, or almost on top of it, don't enable interaction
-	if (FVector::Distance(InLocation, GetGradientCenterHandleLocation(InAvaTextComponent)) > BaseSize/2.0)
+	if (FVector::Distance(InLocation, GetGradientCenterHandleLocation(InTextComponent)) > BaseSize/2.0)
 	{
-		InPDI->SetHitProxy(new HAvaTextGradientLineEndHandleProxy(InAvaTextComponent));
+		InPDI->SetHitProxy(new HAvaTextGradientLineEndHandleProxy(InTextComponent));
 	}
 	
 	InPDI->DrawSprite(IconLocation, IconSize, IconSize, UVSprite->GetResource(), InColor,
@@ -1513,7 +1516,7 @@ void FAvaTextVisualizer::DrawGradientLineEndHandle(const UAvaText3DComponent* In
 	InPDI->SetHitProxy(nullptr);
 }
 
-void FAvaTextVisualizer::DrawCharacterKerningHandle(const UAvaText3DComponent* InTextComponent, const FSceneView* InView, FPrimitiveDrawInterface* InPDI,
+void FAvaTextVisualizer::DrawCharacterKerningHandle(const UText3DComponent* InTextComponent, const FSceneView* InView, FPrimitiveDrawInterface* InPDI,
 	const FVector& InLocation, uint16 InCharacterIndex, const FLinearColor& InColor) const
 {
 	static constexpr float BaseSize = 1.f;
diff --git a/Engine/Plugins/Experimental/Avalanche/Source/AvalancheTextEditor/Private/Visualizer/AvaTextVisualizer.h b/Engine/Plugins/Experimental/Avalanche/Source/AvalancheTextEditor/Private/Visualizer/AvaTextVisualizer.h
--- a/Engine/Plugins/Experimental/Avalanche/Source/AvalancheTextEditor/Private/Visualizer/AvaTextVisualizer.h
+++ b/Engine/Plugins/Experimental/Avalanche/Source/AvalancheTextEditor/Private/Visualizer/AvaTextVisualizer.h
@@ -5,7 +5,7 @@
 #include "AvaVisBase.h"
 
 class FProperty;
-class UAvaText3DComponent;
+class UText3DComponent;
 class UText3DComponent;
 
 struct HAvaTextMaxTextHeightProxy : HAvaHitProxy
@@ -146,7 +146,7 @@ protected:
 	FProperty* ScaleProportionallyProperty;
 	FProperty* CharacterKerningProperty;
 
-	TWeakObjectPtr<UAvaText3DComponent> TextComponent;
+	TWeakObjectPtr<UText3DComponent> TextComponent;
 	mutable FBox Bounds;
 	mutable int32 LineCount;
 	float LineHeight;
@@ -196,7 +196,7 @@ protected:
 	static bool DrawGradientSmoothnessHandle(const UText3DComponent* InTextComponent, const FSceneView* InView,
 		FPrimitiveDrawInterface* InPDI, FVector InHandleLocation);
 
-	void DrawGradientHandles(const UAvaText3DComponent* InAvaText3DComponent, const FSceneView* InView, FPrimitiveDrawInterface* InPDI) const;
+	void DrawGradientHandles(const UText3DComponent* InText3DComponent, const FSceneView* InView, FPrimitiveDrawInterface* InPDI) const;
 
 	FBox GetBoundsMax(const UText3DComponent* InTextComp) const;
 	
@@ -204,14 +204,14 @@ protected:
 	
 	FVector GetHeightHandleLocation(const UText3DComponent* InTextComp) const;
 	
-	void GetTextActorGradientControlsLocations(const UAvaText3DComponent* InAvaText3DComponent,
+	void GetTextActorGradientControlsLocations(const UText3DComponent* InText3DComponent,
 		FVector& OutGradientCenterLocation, FVector& OutGradientStartLocation, FVector& OutGradientEndLocation) const;
 	
-	FVector GetGradientEndHandleLocation(const UAvaText3DComponent* InTextComponent) const;
-	FVector GetGradientCenterHandleLocation(const UAvaText3DComponent* InTextComponent) const;
-	FVector GetGradientStartHandleLocation(const UAvaText3DComponent* InTextComponent) const;
-	FVector GetGradientSmoothnessHandleLocation(const UAvaText3DComponent* InAvaText3DComponent) const;
-	FVector GetKerningHandleLocation(const UAvaText3DComponent* InTextComponent, int32 InIndex) const;
+	FVector GetGradientEndHandleLocation(const UText3DComponent* InTextComponent) const;
+	FVector GetGradientCenterHandleLocation(const UText3DComponent* InTextComponent) const;
+	FVector GetGradientStartHandleLocation(const UText3DComponent* InTextComponent) const;
+	FVector GetGradientSmoothnessHandleLocation(const UText3DComponent* InText3DComponent) const;
+	FVector GetKerningHandleLocation(const UText3DComponent* InTextComponent, int32 InIndex) const;
 
 	void DrawMaxTextWidthButton(const UText3DComponent* InTextComponent, const FSceneView* InView, FPrimitiveDrawInterface* InPDI,
 		int32 InIconIndex, const FLinearColor& InColor) const;
@@ -231,12 +231,12 @@ protected:
 	void DrawMaxTextHeightHandle(const UText3DComponent* InTextComponent, const FSceneView* InView, FPrimitiveDrawInterface* InPDI,
 		const FLinearColor& InColor) const;
 
-	void DrawGradientLineStartHandle(const UAvaText3DComponent* InTextComponent, const FSceneView* InView,
+	void DrawGradientLineStartHandle(const UText3DComponent* InTextComponent, const FSceneView* InView,
 		FPrimitiveDrawInterface* InPDI, const FVector& InLocation, const FLinearColor& InColor) const;
 
-	void DrawGradientLineEndHandle(const UAvaText3DComponent* InTextComponent, const FSceneView* InView,
+	void DrawGradientLineEndHandle(const UText3DComponent* InTextComponent, const FSceneView* InView,
 		FPrimitiveDrawInterface* InPDI, const FVector& InLocation, const FLinearColor& InColor) const;
 	
-	void DrawCharacterKerningHandle(const UAvaText3DComponent* InTextComponent, const FSceneView* InView,
+	void DrawCharacterKerningHandle(const UText3DComponent* InTextComponent, const FSceneView* InView,
 		FPrimitiveDrawInterface* InPDI, const FVector& InLocation, uint16 InCharacterIndex, const FLinearColor& InColor) const;
 };
diff --git a/Engine/Plugins/Experimental/Text3D/Source/Text3D/Private/Extensions/Text3DDefaultCharacterExtension.cpp b/Engine/Plugins/Experimental/Text3D/Source/Text3D/Private/Extensions/Text3DDefaultCharacterExtension.cpp
--- a/Engine/Plugins/Experimental/Text3D/Source/Text3D/Private/Extensions/Text3DDefaultCharacterExtension.cpp
+++ b/Engine/Plugins/Experimental/Text3D/Source/Text3D/Private/Extensions/Text3DDefaultCharacterExtension.cpp
@@ -45,17 +45,18 @@ EText3DExtensionResult UText3DDefaultCharacterExtension::PreRendererUpdate(EText
 	}
 
 	const FString FormattedTextStr = Text3DComponent->GetFormattedText().ToString();
+
 	const TArray<FText3DCharacterMetadata> TextCharacterMetadata = UE::Text3D::Utilities::GetCharacterMetadata(*TypeFaceEntry, FormattedTextStr);
 
 	AllocateTextCharacters(TextCharacterMetadata.Num());
 
-	for (TConstEnumerateRef<TObjectPtr<UText3DCharacterBase>> TextCharacter : EnumerateRange(TextCharacters))
+	for (TConstEnumerateRef<TObjectPtr<UText3DCharacterBase>> Character : EnumerateRange(TextCharacters))
 	{
-		const FText3DCharacterMetadata& Metadata = TextCharacterMetadata[TextCharacter.GetIndex()];
+		const FText3DCharacterMetadata& Metadata = TextCharacterMetadata[Character.GetIndex()];
 #if WITH_EDITORONLY_DATA
-		(*TextCharacter)->SetCharacter(FString::Printf(TEXT("%i - %c"), Metadata.TextIndex, FormattedTextStr[Metadata.TextIndex]));
+		(*Character)->SetCharacter(FString::Printf(TEXT("%c"), FormattedTextStr[Metadata.TextIndex]));
 #endif
-		(*TextCharacter)->SetCharacterMetadata(Metadata);
+		(*Character)->SetCharacterMetadata(Metadata);
 	}
 
 	return EText3DExtensionResult::Finished;
diff --git a/Engine/Plugins/Experimental/Text3D/Source/Text3D/Public/Characters/Text3DDefaultCharacter.h b/Engine/Plugins/Experimental/Text3D/Source/Text3D/Public/Characters/Text3DDefaultCharacter.h
--- a/Engine/Plugins/Experimental/Text3D/Source/Text3D/Public/Characters/Text3DDefaultCharacter.h
+++ b/Engine/Plugins/Experimental/Text3D/Source/Text3D/Public/Characters/Text3DDefaultCharacter.h
@@ -38,10 +38,10 @@ protected:
 	virtual TOptional<float> GetCharacterKerning() const override;
 	//~ End UText3DCharacterBase
 
-	UPROPERTY(EditInstanceOnly, Getter="GetKerningOverride", Setter="SetKerningOverride", Category="Character", meta=(InlineEditConditionToggle, AllowPrivateAccess = "true"))
+	UPROPERTY(EditInstanceOnly, Getter="GetKerningOverride", Setter="SetKerningOverride", Category="Character", meta=(AllowPrivateAccess = "true"))
 	bool bKerningOverride = false;
 
 	/** Kerning adjusts the space between this character */
 	UPROPERTY(EditInstanceOnly, Getter, Setter, Category="Character", meta=(EditCondition="bKerningOverride", AllowPrivateAccess = "true"))
-	float Kerning;
+	float Kerning = 0.f;
 };
diff --git a/Engine/Plugins/Experimental/Text3D/Source/Text3D/Public/Extensions/Text3DCharacterExtensionBase.h b/Engine/Plugins/Experimental/Text3D/Source/Text3D/Public/Extensions/Text3DCharacterExtensionBase.h
--- a/Engine/Plugins/Experimental/Text3D/Source/Text3D/Public/Extensions/Text3DCharacterExtensionBase.h
+++ b/Engine/Plugins/Experimental/Text3D/Source/Text3D/Public/Extensions/Text3DCharacterExtensionBase.h
@@ -36,4 +36,10 @@ public:
 	{
 		return {};
 	}
+
+#if WITH_EDITORONLY_DATA
+	/** Used for detail customization within editor */
+	UPROPERTY(EditInstanceOnly, Transient, Category="Character", meta=(TextCharacterSelector, AllowPrivateAccess=true))
+	uint16 TextCharacterIndex = 0;
+#endif
 };
\ No newline at end of file
diff --git a/Engine/Plugins/Experimental/Text3D/Source/Text3D/Public/Extensions/Text3DDefaultCharacterExtension.h b/Engine/Plugins/Experimental/Text3D/Source/Text3D/Public/Extensions/Text3DDefaultCharacterExtension.h
--- a/Engine/Plugins/Experimental/Text3D/Source/Text3D/Public/Extensions/Text3DDefaultCharacterExtension.h
+++ b/Engine/Plugins/Experimental/Text3D/Source/Text3D/Public/Extensions/Text3DDefaultCharacterExtension.h
@@ -3,7 +3,6 @@
 #pragma once
 
 #include "Text3DCharacterExtensionBase.h"
-#include "Subsystems/Text3DEngineSubsystem.h"
 #include "Text3DDefaultCharacterExtension.generated.h"
 
 class UText3DDefaultCharacter;
@@ -29,7 +28,7 @@ protected:
 	void AllocateTextCharacters(uint16 InCharacterCount);
 
 	/** Characters composing the active text */
-	UPROPERTY(VisibleAnywhere, EditFixedSize, Instanced, Category = "Character", meta = (TitleProperty="Character", EditFixedOrder, AllowPrivateAccess = "true"))
+	UPROPERTY()
 	TArray<TObjectPtr<UText3DCharacterBase>> TextCharacters;
 
 	/** Pool of characters to reuse when text changes to avoid creating new objects */
diff --git a/Engine/Plugins/Experimental/Text3D/Source/Text3D/Public/Extensions/Text3DExtensionBase.h b/Engine/Plugins/Experimental/Text3D/Source/Text3D/Public/Extensions/Text3DExtensionBase.h
--- a/Engine/Plugins/Experimental/Text3D/Source/Text3D/Public/Extensions/Text3DExtensionBase.h
+++ b/Engine/Plugins/Experimental/Text3D/Source/Text3D/Public/Extensions/Text3DExtensionBase.h
@@ -53,8 +53,9 @@ public:
 		return UpdatePriority;
 	}
 
+	TEXT3D_API UText3DComponent* GetText3DComponent() const;
+
 protected:
-	UText3DComponent* GetText3DComponent() const;
 	void RequestUpdate(EText3DRendererFlags InFlags, bool bInImmediate = false) const;
 
 private:
diff --git a/Engine/Plugins/Experimental/Text3D/Source/Text3DEditor/Private/Customizations/Text3DEditorCharacterPropertyTypeCustomization.cpp b/Engine/Plugins/Experimental/Text3D/Source/Text3DEditor/Private/Customizations/Text3DEditorCharacterPropertyTypeCustomization.cpp
new file mode 100644
--- /dev/null
+++ b/Engine/Plugins/Experimental/Text3D/Source/Text3DEditor/Private/Customizations/Text3DEditorCharacterPropertyTypeCustomization.cpp
@@ -0,0 +1,180 @@
+// Copyright Epic Games, Inc. All Rights Reserved.
+
+#include "Customizations/Text3DEditorCharacterPropertyTypeCustomization.h"
+
+#include "Characters/Text3DCharacterBase.h"
+#include "DetailLayoutBuilder.h"
+#include "DetailWidgetRow.h"
+#include "Extensions/Text3DDefaultCharacterExtension.h"
+#include "IDetailChildrenBuilder.h"
+#include "IPropertyUtilities.h"
+#include "Text3DComponent.h"
+#include "UObject/Object.h"
+#include "Widgets/Input/SSpinBox.h"
+#include "Widgets/Text//STextBlock.h"
+
+#define LOCTEXT_NAMESPACE "Text3DEditorCharacterPropertyTypeCustomization"
+
+namespace UE::Text3DEditor::Customization
+{
+	TSharedRef<IPropertyTypeCustomization> FText3DEditorCharacterPropertyTypeCustomization::MakeInstance()
+	{
+		return MakeShared<FText3DEditorCharacterPropertyTypeCustomization>();
+	}
+
+	void FText3DEditorCharacterPropertyTypeCustomization::CustomizeHeader(TSharedRef<IPropertyHandle> InPropertyHandle, FDetailWidgetRow& InHeaderRow, IPropertyTypeCustomizationUtils& InUtils)
+	{
+		TArray<UObject*> Objects;
+		InPropertyHandle->GetOuterObjects(Objects);
+		if (Objects.Num() != 1)
+		{
+			return;
+		}
+
+		UText3DCharacterExtensionBase* CharacterExtension = Cast<UText3DCharacterExtensionBase>(Objects[0]);
+		if (!CharacterExtension)
+		{
+			return;
+		}
+
+		PropertyUtilitiesWeak = InUtils.GetPropertyUtilities().ToWeakPtr();
+		if (UText3DComponent* Text3DComponent = CharacterExtension->GetText3DComponent())
+		{
+			Text3DComponent->OnTextPostUpdate().AddSPLambda(this, [this](UText3DComponent*, EText3DRendererFlags InFlags)
+			{
+				if (EnumHasAnyFlags(InFlags, EText3DRendererFlags::Geometry))
+				{
+					if (const TSharedPtr<IPropertyUtilities> Utilities = PropertyUtilitiesWeak.Pin())
+					{
+						Utilities->RequestForceRefresh();
+					}
+				}
+			});
+		}
+
+		CharacterExtensionWeak = CharacterExtension;
+		CharacterPropertyHandleWeak = InPropertyHandle;
+		ActiveIndex = CharacterExtension->TextCharacterIndex;
+
+		const int32 LastCharacterIndex = FMath::Max(0, CharacterExtension->GetCharacterCount() - 1);
+
+		InHeaderRow
+		.NameContent()
+		[
+			SNew(SHorizontalBox)
+			+ SHorizontalBox::Slot()
+			.HAlign(HAlign_Left)
+			.VAlign(VAlign_Center)
+			.AutoWidth()
+			[
+				InPropertyHandle->CreatePropertyNameWidget()
+			]
+			+ SHorizontalBox::Slot()
+			.AutoWidth()
+			.HAlign(HAlign_Left)
+			.VAlign(VAlign_Center)
+			.Padding(5.f, 0.f)
+			[
+				SNew(STextBlock)
+				.Font(IDetailLayoutBuilder::GetDetailFont())
+				.Text(this, &FText3DEditorCharacterPropertyTypeCustomization::GetCharacterText)
+			]
+		]
+		.ValueContent()
+		[
+			SNew(SSpinBox<int32>)
+			.MinValue(0)
+			.MaxValue(LastCharacterIndex)
+			.Delta(1)
+			.EnableSlider(true)
+			.PreventThrottling(true)
+			.Font(IDetailLayoutBuilder::GetDetailFont())
+			.Value_Lambda([this]()
+			{
+				return ActiveIndex;
+			})
+			.OnValueChanged_Lambda([this](int32 InNewValue)
+			{
+				ActiveIndex = InNewValue;
+			})
+			.OnValueCommitted_Lambda([this](int32 InNewValue, ETextCommit::Type InType)
+			{
+				if (UText3DCharacterExtensionBase* Extension = CharacterExtensionWeak.Get())
+				{
+					if (InType != ETextCommit::Type::OnCleared)
+					{
+						ActiveIndex = InNewValue;
+
+						Extension->TextCharacterIndex = ActiveIndex;
+
+						if (const TSharedPtr<IPropertyUtilities> Utilities = PropertyUtilitiesWeak.Pin())
+						{
+							Utilities->RequestForceRefresh();
+						}
+					}
+					else
+					{
+						ActiveIndex = Extension->TextCharacterIndex;
+					}
+				}
+			})
+		];
+	}
+
+	void FText3DEditorCharacterPropertyTypeCustomization::CustomizeChildren(TSharedRef<IPropertyHandle> InPropertyHandle, IDetailChildrenBuilder& InChildBuilder, IPropertyTypeCustomizationUtils& InUtils)
+	{
+		if (const UText3DCharacterExtensionBase* Extension = CharacterExtensionWeak.Get())
+		{
+			if (UText3DCharacterBase* Character = Extension->GetCharacter(Extension->TextCharacterIndex))
+			{
+				for (const FProperty* Property : TFieldRange<FProperty>(Character->GetClass()))
+				{
+					if (!Property || !Property->HasAnyPropertyFlags(CPF_Edit))
+					{
+						continue;
+					}
+
+					FAddPropertyParams AddParams;
+					AddParams.CreateCategoryNodes(false);
+					AddParams.HideRootObjectNode(true);
+
+					if (const IDetailPropertyRow* Row = InChildBuilder.AddExternalObjectProperty({Character}, Property->GetFName(), AddParams))
+					{
+						if (const TSharedPtr<IPropertyHandle> RowHandle = Row->GetPropertyHandle())
+						{
+							// Due to AddExternalObjectProperty, EditCondition are not updating immediately so we give them a hand
+							RowHandle->SetOnPropertyValueChangedWithData(TDelegate<void(const FPropertyChangedEvent&)>::CreateSPLambda(this, [this](const FPropertyChangedEvent& InEvent)
+							{
+								if (InEvent.ChangeType != EPropertyChangeType::Interactive)
+								{
+									if (const TSharedPtr<IPropertyHandle> ParentProperty = CharacterPropertyHandleWeak.Pin())
+									{
+										ParentProperty->RequestRebuildChildren();
+									}
+								}
+							}));
+						}
+					}
+				}
+			}
+		}
+	}
+
+	FText FText3DEditorCharacterPropertyTypeCustomization::GetCharacterText() const
+	{
+		if (const UText3DCharacterExtensionBase* Extension = CharacterExtensionWeak.Get())
+		{
+			const UText3DCharacterBase* Character = Extension->GetCharacter(ActiveIndex);
+
+			return FText::Format(
+				LOCTEXT("CharacterCount", "Total: {0} / Current: {1}")
+				, FText::AsNumber(Extension->GetCharacterCount())
+				, IsValid(Character) ? FText::FromString(Character->GetCharacter()) : FText::GetEmpty()
+			);
+		}
+
+		return FText::GetEmpty();
+	}
+}
+
+#undef LOCTEXT_NAMESPACE
diff --git a/Engine/Plugins/Experimental/Text3D/Source/Text3DEditor/Private/Customizations/Text3DEditorCharacterPropertyTypeCustomization.h b/Engine/Plugins/Experimental/Text3D/Source/Text3DEditor/Private/Customizations/Text3DEditorCharacterPropertyTypeCustomization.h
new file mode 100644
--- /dev/null
+++ b/Engine/Plugins/Experimental/Text3D/Source/Text3DEditor/Private/Customizations/Text3DEditorCharacterPropertyTypeCustomization.h
@@ -0,0 +1,54 @@
+// Copyright Epic Games, Inc. All Rights Reserved.
+
+#pragma once
+
+#include "IPropertyTypeCustomization.h"
+#include "PropertyEditorModule.h"
+#include "PropertyHandle.h"
+#include "Templates/SharedPointer.h"
+
+class FDetailWidgetRow;
+class IDetailChildrenBuilder;
+class IPropertyHandle;
+class IPropertyTypeCustomizationUtils;
+class UText3DCharacterExtensionBase;
+
+namespace UE::Text3DEditor::Customization
+{
+	/** Only allow property customization with "TextCharacterSelector" metadata */
+	class FText3DEditorCharacterPropertyTypeIdentifier : public IPropertyTypeIdentifier
+	{
+	public:
+		//~ Begin IPropertyTypeIdentifier
+		virtual bool IsPropertyTypeCustomized(const IPropertyHandle& InPropertyHandle) const override
+		{
+			return InPropertyHandle.HasMetaData(TEXT("TextCharacterSelector"));
+		}
+		//~ End IPropertyTypeIdentifier
+	};
+
+	/** Customization for UText3DCharacter object */
+	class FText3DEditorCharacterPropertyTypeCustomization : public IPropertyTypeCustomization
+	{
+	public:
+		static TSharedRef<IPropertyTypeCustomization> MakeInstance();
+
+		//~ Begin IPropertyTypeCustomization
+		virtual void CustomizeHeader(TSharedRef<IPropertyHandle> InPropertyHandle
+			, FDetailWidgetRow& InHeaderRow
+			, IPropertyTypeCustomizationUtils& InUtils) override;
+
+		virtual void CustomizeChildren(TSharedRef<IPropertyHandle> InPropertyHandle
+			, IDetailChildrenBuilder& InChildBuilder
+			, IPropertyTypeCustomizationUtils& InUtils) override;
+		//~ End IPropertyTypeCustomization
+
+	private:
+		FText GetCharacterText() const;
+
+		int32 ActiveIndex = 0;
+		TWeakPtr<IPropertyHandle> CharacterPropertyHandleWeak;
+		TWeakObjectPtr<UText3DCharacterExtensionBase> CharacterExtensionWeak;
+		TWeakPtr<IPropertyUtilities> PropertyUtilitiesWeak;
+	};
+}
diff --git a/Engine/Plugins/Experimental/Text3D/Source/Text3DEditor/Private/Customizations/Text3DEditorTextComponentDetailCustomization.cpp b/Engine/Plugins/Experimental/Text3D/Source/Text3DEditor/Private/Customizations/Text3DEditorTextComponentDetailCustomization.cpp
new file mode 100644
--- /dev/null
+++ b/Engine/Plugins/Experimental/Text3D/Source/Text3DEditor/Private/Customizations/Text3DEditorTextComponentDetailCustomization.cpp
@@ -0,0 +1,44 @@
+// Copyright Epic Games, Inc. All Rights Reserved.
+
+#include "Customizations/Text3DEditorTextComponentDetailCustomization.h"
+
+#include "Modules/ModuleManager.h"
+#include "PropertyEditorModule.h"
+#include "Text3DComponent.h"
+
+namespace UE::Text3DEditor::Customization
+{
+	void FText3DEditorTextComponentDetailCustomization::CustomizeDetails(IDetailLayoutBuilder& InDetailBuilder)
+	{
+		static bool bSectionInitialized = false;
+		if (!bSectionInitialized)
+		{
+			bSectionInitialized = true;
+
+			FPropertyEditorModule& PropertyModule = FModuleManager::GetModuleChecked<FPropertyEditorModule>(TEXT("PropertyEditor"));
+			const FName ComponentClassName = UText3DComponent::StaticClass()->GetFName();
+
+			const FName GeometrySectionName = TEXT("Geometry");
+			const TSharedRef<FPropertySection> GeometrySection = PropertyModule.FindOrCreateSection(ComponentClassName, GeometrySectionName, FText::FromName(GeometrySectionName));
+			GeometrySection->AddCategory(TEXT("Geometry"));
+
+			const FName LayoutSectionName = TEXT("Layout");
+			const TSharedRef<FPropertySection> LayoutSection = PropertyModule.FindOrCreateSection(ComponentClassName, LayoutSectionName, FText::FromName(LayoutSectionName));
+			LayoutSection->AddCategory(TEXT("Layout"));
+			LayoutSection->AddCategory(TEXT("Character"));
+			LayoutSection->AddCategory(TEXT("Transform"));
+
+			const FName RenderingSectionName = TEXT("Rendering");
+			const TSharedRef<FPropertySection> RenderingSection = PropertyModule.FindOrCreateSection(ComponentClassName, RenderingSectionName, FText::FromName(RenderingSectionName));
+			RenderingSection->AddCategory(TEXT("Rendering"));
+
+			const FName TextSectionName = TEXT("Text");
+			const TSharedRef<FPropertySection> TextSection = PropertyModule.FindOrCreateSection(ComponentClassName, TextSectionName, FText::FromName(TextSectionName));
+			TextSection->AddCategory(TEXT("Text"));
+
+			const FName StyleSectionName = TEXT("Style");
+			const TSharedRef<FPropertySection> MaterialSection = PropertyModule.FindOrCreateSection(ComponentClassName, StyleSectionName, FText::FromName(StyleSectionName));
+			MaterialSection->AddCategory(TEXT("Material"));
+		}
+	}
+}
diff --git a/Engine/Plugins/Experimental/Text3D/Source/Text3DEditor/Private/Customizations/Text3DEditorTextComponentDetailCustomization.h b/Engine/Plugins/Experimental/Text3D/Source/Text3DEditor/Private/Customizations/Text3DEditorTextComponentDetailCustomization.h
new file mode 100644
--- /dev/null
+++ b/Engine/Plugins/Experimental/Text3D/Source/Text3DEditor/Private/Customizations/Text3DEditorTextComponentDetailCustomization.h
@@ -0,0 +1,26 @@
+// Copyright Epic Games, Inc. All Rights Reserved.
+
+#pragma once
+
+#include "IDetailCustomization.h"
+#include "Templates/SharedPointerFwd.h"
+
+class IPropertyHandle;
+class IPropertyUtilities;
+
+namespace UE::Text3DEditor::Customization
+{
+	/** Used to customize Text3D component properties in details panel */
+	class FText3DEditorTextComponentDetailCustomization : public IDetailCustomization
+	{
+	public:
+		static TSharedRef<IDetailCustomization> MakeInstance()
+		{
+			return MakeShared<FText3DEditorTextComponentDetailCustomization>();
+		}
+
+		//~ Begin IDetailCustomization
+		virtual void CustomizeDetails(IDetailLayoutBuilder& InDetailBuilder) override;
+		//~ End IDetailCustomization
+	};
+}
\ No newline at end of file
diff --git a/Engine/Plugins/Experimental/Text3D/Source/Text3DEditor/Private/Text3DEditorModule.cpp b/Engine/Plugins/Experimental/Text3D/Source/Text3DEditor/Private/Text3DEditorModule.cpp
--- a/Engine/Plugins/Experimental/Text3D/Source/Text3DEditor/Private/Text3DEditorModule.cpp
+++ b/Engine/Plugins/Experimental/Text3D/Source/Text3DEditor/Private/Text3DEditorModule.cpp
@@ -3,14 +3,18 @@
 #include "Text3DEditorModule.h"
 
 #include "Commands/Text3DEditorFontSelectorCommands.h"
+#include "Customizations/Text3DEditorCharacterPropertyTypeCustomization.h"
 #include "Customizations/Text3DEditorFontPropertyTypeCustomization.h"
 #include "Customizations/Text3DEditorHorizontalPropertyTypeCustomization.h"
 #include "Customizations/Text3DEditorVerticalPropertyTypeCustomization.h"
+#include "Customizations/Text3DEditorTextComponentDetailCustomization.h"
 #include "Engine/Font.h"
+#include "Extensions/Text3DDefaultCharacterExtension.h"
 #include "Logs/Text3DEditorLogs.h"
 #include "Modules/ModuleManager.h"
 #include "PropertyEditorModule.h"
 #include "Styles/Text3DEditorStyle.h"
+#include "Text3DComponent.h"
 #include "Text3DTypes.h"
 #include "UObject/ReflectedTypeAccessors.h"
 
@@ -20,6 +24,8 @@ IMPLEMENT_MODULE(FText3DEditorModule, Text3DEditor)
 
 void FText3DEditorModule::StartupModule()
 {
+	using namespace UE::Text3DEditor::Customization;
+
 	// Init style set
 	FText3DEditorStyle::Get();
 
@@ -38,6 +44,13 @@ void FText3DEditorModule::StartupModule()
 	PropertyModule.RegisterCustomPropertyTypeLayout(RegisteredTypeNames.Add_GetRef(UFont::StaticClass()->GetFName())
 		, FOnGetPropertyTypeCustomizationInstance::CreateStatic(&FText3DEditorFontPropertyTypeCustomization::MakeInstance)
 		, MakeShared<FText3DEditorFontPropertyTypeIdentifier>());
+
+	PropertyModule.RegisterCustomPropertyTypeLayout(RegisteredTypeNames.Add_GetRef(FUInt16Property::StaticClass()->GetFName())
+	 	, FOnGetPropertyTypeCustomizationInstance::CreateStatic(&FText3DEditorCharacterPropertyTypeCustomization::MakeInstance)
+	 	, MakeShared<FText3DEditorCharacterPropertyTypeIdentifier>());
+
+	PropertyModule.RegisterCustomClassLayout(RegisteredTypeNames.Add_GetRef(UText3DComponent::StaticClass()->GetFName())
+		, FOnGetDetailCustomizationInstance::CreateStatic(&FText3DEditorTextComponentDetailCustomization::MakeInstance));
 }
 
 void FText3DEditorModule::ShutdownModule()
@@ -49,6 +62,7 @@ void FText3DEditorModule::ShutdownModule()
 		for (const FName& RegisteredTypeName : RegisteredTypeNames)
 		{
 			PropertyModule->UnregisterCustomPropertyTypeLayout(RegisteredTypeName);
+			PropertyModule->UnregisterCustomClassLayout(RegisteredTypeName);
 		}
 	}
 
diff --git a/Engine/Plugins/Experimental/Text3D/Source/Text3DEditor/Text3DEditor.Build.cs b/Engine/Plugins/Experimental/Text3D/Source/Text3DEditor/Text3DEditor.Build.cs
--- a/Engine/Plugins/Experimental/Text3D/Source/Text3DEditor/Text3DEditor.Build.cs
+++ b/Engine/Plugins/Experimental/Text3D/Source/Text3DEditor/Text3DEditor.Build.cs
@@ -17,6 +17,7 @@ public class Text3DEditor : ModuleRules
 			"FreeType2",
 			"InputCore",
 			"Projects",
+			"PropertyEditor",
 			"SlateCore",
 			"Slate",
 			"Text3D",
