Motion Design: refactored outliner add item action to support multiple items. This fixes issues with ordering issues in outliner when duplicating and moving items. As well as limitations with single item operations like when dragging multiple items under an item that is already a parent to at least one of the dragged items (but not all).

diff --git a/Engine/Plugins/Experimental/Avalanche/Source/AvalancheEditor/Private/Outliner/AvaOutlinerExtension.cpp b/Engine/Plugins/Experimental/Avalanche/Source/AvalancheEditor/Private/Outliner/AvaOutlinerExtension.cpp
--- a/Engine/Plugins/Experimental/Avalanche/Source/AvalancheEditor/Private/Outliner/AvaOutlinerExtension.cpp
+++ b/Engine/Plugins/Experimental/Avalanche/Source/AvalancheEditor/Private/Outliner/AvaOutlinerExtension.cpp
@@ -205,7 +205,11 @@ void FAvaOutlinerExtension::OnPasteActors(FStringView InPastedData, TConstArrayV
 	TMap<AActor*, AActor*> DuplicatedActors = FAvaEditorPastedActor::BuildDuplicatedActorMap(InPastedActors);
 
 	AvaOutliner->SetIgnoreNotify(FAvaOutlinerExtension::NotifyFlags, /*bIgnore*/false);
+
+	// Process pasted actors that had a source (i.e. deemed a duplicate operation)
 	AvaOutliner->OnActorsDuplicated(DuplicatedActors);
+
+	// Process the rest of the pasted actors that did not have a source
 	AvaOutliner->OnActorsPasted(InPastedData, PastedActors);
 }
 
diff --git a/Engine/Plugins/Experimental/Avalanche/Source/AvalancheEditor/Private/Outliner/AvaOutlinerMaterialDesignerProxy.cpp b/Engine/Plugins/Experimental/Avalanche/Source/AvalancheEditor/Private/Outliner/AvaOutlinerMaterialDesignerProxy.cpp
--- a/Engine/Plugins/Experimental/Avalanche/Source/AvalancheEditor/Private/Outliner/AvaOutlinerMaterialDesignerProxy.cpp
+++ b/Engine/Plugins/Experimental/Avalanche/Source/AvalancheEditor/Private/Outliner/AvaOutlinerMaterialDesignerProxy.cpp
@@ -55,9 +55,9 @@ void FAvaOutlinerMaterialDesignerProxy::GetProxiedItems(const TSharedRef<IAvaOut
 					MaterialItem->SetParent(SharedThis(this));
 
 					FAvaOutlinerAddItemParams AddParams;
-					AddParams.Item = MaterialItem;
+					AddParams.Items.Add(MaterialItem);
 					AddParams.Flags = EAvaOutlinerAddItemFlags::None;
-					AddChild(AddParams);
+					AddChildren(AddParams);
 				}
 
 				// Check if it's already in the array
diff --git a/Engine/Plugins/Experimental/Avalanche/Source/AvalancheOutliner/Private/AvaOutliner.cpp b/Engine/Plugins/Experimental/Avalanche/Source/AvalancheOutliner/Private/AvaOutliner.cpp
--- a/Engine/Plugins/Experimental/Avalanche/Source/AvalancheOutliner/Private/AvaOutliner.cpp
+++ b/Engine/Plugins/Experimental/Avalanche/Source/AvalancheOutliner/Private/AvaOutliner.cpp
@@ -596,7 +596,7 @@ void FAvaOutliner::OnActorsDuplicated(const TMap<AActor*, AActor*>& InDuplicateA
 		}
 
 		FAvaOutlinerItemPtr Item = FindOrAdd<FAvaOutlinerActor>(DuplicateActor);
-		Params.Item = Item;
+		Params.Items = { Item };
 		
 		if (FAvaOutlinerItemPtr* const FoundTemplateItem = DuplicateItemMap.Find(DuplicateActor))
 		{
@@ -659,7 +659,7 @@ void FAvaOutliner::GroupSelection(AActor* InGroupingActor, const TOptional<FAtta
 	if (bIsGroupingItemNew)
 	{
 		FAvaOutlinerAddItemParams AddGroupingItemParams;
-		AddGroupingItemParams.Item  = GroupingItem;
+		AddGroupingItemParams.Items = { GroupingItem };
 		AddGroupingItemParams.Flags = EAvaOutlinerAddItemFlags::AddChildren;
 
 		AddGroupingItemParams.AttachmentTransformRules = InTransformRules;
@@ -697,15 +697,9 @@ void FAvaOutliner::GroupSelection(AActor* InGroupingActor, const TOptional<FAtta
 	AddChildItemParams.SelectionFlags   = EAvaOutlinerItemSelectionFlags::AppendToCurrentSelection | EAvaOutlinerItemSelectionFlags::ScrollIntoView;
 
 	AddChildItemParams.AttachmentTransformRules = InTransformRules;
+	AddChildItemParams.Items = SelectedItems;
 
-	for (const FAvaOutlinerItemPtr& Item : SelectedItems)
-	{
-		if (Item.IsValid())
-		{
-			AddChildItemParams.Item = Item;
-			EnqueueItemAction<FAvaOutlinerAddItem>(AddChildItemParams);
-		}
-	}
+	EnqueueItemAction<FAvaOutlinerAddItem>(AddChildItemParams);
 }
 
 void FAvaOutliner::PostUndo(bool bSuccess)
@@ -1256,7 +1250,7 @@ void FAvaOutliner::OnActorSpawned(AActor* InActor)
 		if (CanProcessActorSpawn(InActor) && !FindItem(FAvaOutlinerItemId(InActor)))
 		{
 			FAvaOutlinerAddItemParams Params;
-			Params.Item             = FindOrAdd<FAvaOutlinerActor>(InActor);
+			Params.Items            = { FindOrAdd<FAvaOutlinerActor>(InActor) };
 			Params.RelativeItem     = nullptr;
 			Params.RelativeDropZone = TOptional<EItemDropZone>();
 			Params.Flags            = EAvaOutlinerAddItemFlags::AddChildren;
diff --git a/Engine/Plugins/Experimental/Avalanche/Source/AvalancheOutliner/Private/DragDropOps/Handlers/AvaOutlinerActorDropHandler.cpp b/Engine/Plugins/Experimental/Avalanche/Source/AvalancheOutliner/Private/DragDropOps/Handlers/AvaOutlinerActorDropHandler.cpp
--- a/Engine/Plugins/Experimental/Avalanche/Source/AvalancheOutliner/Private/DragDropOps/Handlers/AvaOutlinerActorDropHandler.cpp
+++ b/Engine/Plugins/Experimental/Avalanche/Source/AvalancheOutliner/Private/DragDropOps/Handlers/AvaOutlinerActorDropHandler.cpp
@@ -73,10 +73,6 @@ void FAvaOutlinerActorDropHandler::MoveItems(EItemDropZone InDropZone, FAvaOutli
 		});
 
 	TSharedRef<FAvaOutliner> Outliner = StaticCastSharedRef<FAvaOutliner>(InTargetItem->GetOwnerOutliner());
-
-	for (const FAvaOutlinerItemPtr& Item : Items)
-	{
-		Params.Item = Item;
-		Outliner->EnqueueItemAction<FAvaOutlinerAddItem>(Params);
-	}
+	Params.Items = Items;
+	Outliner->EnqueueItemAction<FAvaOutlinerAddItem>(Params);
 }
diff --git a/Engine/Plugins/Experimental/Avalanche/Source/AvalancheOutliner/Private/Item/AvaOutlinerActor.cpp b/Engine/Plugins/Experimental/Avalanche/Source/AvalancheOutliner/Private/Item/AvaOutlinerActor.cpp
--- a/Engine/Plugins/Experimental/Avalanche/Source/AvalancheOutliner/Private/Item/AvaOutlinerActor.cpp
+++ b/Engine/Plugins/Experimental/Avalanche/Source/AvalancheOutliner/Private/Item/AvaOutlinerActor.cpp
@@ -64,61 +64,67 @@ void FAvaOutlinerActor::GetItemProxies(TArray<TSharedPtr<FAvaOutlinerItemProxy>>
 	}
 }
 
-bool FAvaOutlinerActor::AddChild(const FAvaOutlinerAddItemParams& InAddItemParams)
+bool FAvaOutlinerActor::CanAddChild(const FAvaOutlinerItemPtr& InChild) const
 {
-	if (CanAddChild(InAddItemParams.Item))
+	// Super checks InChild.IsValid
+	return Super::CanAddChild(InChild) && InChild->IsA<FAvaOutlinerActor>();
+}
+
+void FAvaOutlinerActor::OnChildrenAdded(TConstArrayView<FAvaOutlinerItemPtr> InAddedChildren, const FAvaOutlinerAddItemParams& InAddItemParams)
+{
+	Super::OnChildrenAdded(InAddedChildren, InAddItemParams);
+
+	for (const FAvaOutlinerItemPtr& Item : InAddedChildren)
 	{
-		// Only accept Items that are Actor Items
-		if (const FAvaOutlinerActor* const ActorItem = InAddItemParams.Item->CastTo<FAvaOutlinerActor>())
+		FAvaOutlinerActor* ActorItem = Item->CastTo<FAvaOutlinerActor>();
+		if (!ActorItem)
 		{
-			AddChildChecked(InAddItemParams);
-			
-			if (AActor* const ChildActor = ActorItem->GetActor())
+			continue;
+		}
+
+		if (AActor* const ChildActor = ActorItem->GetActor())
+		{
+			AActor* const ParentActor = GetActor();
+
+			USceneComponent* const ChildRoot  = ChildActor->GetRootComponent();
+			USceneComponent* const ParentRoot = ParentActor
+				? ParentActor->GetDefaultAttachComponent()
+				: nullptr;
+
+			// Set Component Mobility to Moveable if ParentRoot isn't Static
+			if (ChildRoot && ChildRoot->Mobility == EComponentMobility::Static
+				&& ParentRoot && ParentRoot->Mobility != EComponentMobility::Static)
 			{
-				AActor* const ParentActor = GetActor();
-				
-				USceneComponent* const ChildRoot  = ChildActor->GetRootComponent();
-				USceneComponent* const ParentRoot = ParentActor
-					? ParentActor->GetDefaultAttachComponent()
-					: nullptr;
-				
-				// Set Component Mobility to Moveable if ParentRoot isn't Static
-				if (ChildRoot && ChildRoot->Mobility == EComponentMobility::Static
-					&& ParentRoot && ParentRoot->Mobility != EComponentMobility::Static)
-				{
-					ChildRoot->SetMobility(ParentRoot->Mobility);
-				}
-				
-				EAvaOutlinerHierarchyChangeType ChangeType;
-				
-				if (ParentActor != ChildActor->GetAttachParentActor() && GUnrealEd->CanParentActors(ParentActor, ChildActor))
-				{
-					// Attachment is persisted on the child so modify both actors for Undo/Redo but do not mark the Parent package dirty
-					ChildActor->Modify();
-					ParentActor->Modify(false);
-					
-					// Snap to socket if a valid socket name was provided, otherwise attach without changing the relative transform
-					ChildRoot->AttachToComponent(ParentRoot
-						, InAddItemParams.AttachmentTransformRules.Get(FAttachmentTransformRules::KeepWorldTransform)
-						, NAME_None);
-					
-					ChangeType = EAvaOutlinerHierarchyChangeType::Attached;
-				}
-				else
-				{
-					ChangeType = EAvaOutlinerHierarchyChangeType::Rearranged;
-				}
+				ChildRoot->SetMobility(ParentRoot->Mobility);
+			}
 
-				FAvaOutliner& OutlinerPrivate = static_cast<FAvaOutliner&>(Outliner);
-				if (UAvaOutlinerSubsystem* const OutlinerSubsystem = OutlinerPrivate.GetOutlinerSubsystem())
-				{
-					OutlinerSubsystem->BroadcastActorHierarchyChanged(ChildActor, ParentActor, ChangeType);
-				}
+			EAvaOutlinerHierarchyChangeType ChangeType;
+
+			if (ParentActor != ChildActor->GetAttachParentActor() && GUnrealEd->CanParentActors(ParentActor, ChildActor))
+			{
+				// Attachment is persisted on the child so modify both actors for Undo/Redo but do not mark the Parent package dirty
+				ChildActor->Modify();
+				ParentActor->Modify(false);
+
+				// Snap to socket if a valid socket name was provided, otherwise attach without changing the relative transform
+				ChildRoot->AttachToComponent(ParentRoot
+					, InAddItemParams.AttachmentTransformRules.Get(FAttachmentTransformRules::KeepWorldTransform)
+					, NAME_None);
+
+				ChangeType = EAvaOutlinerHierarchyChangeType::Attached;
+			}
+			else
+			{
+				ChangeType = EAvaOutlinerHierarchyChangeType::Rearranged;
+			}
+
+			FAvaOutliner& OutlinerPrivate = static_cast<FAvaOutliner&>(Outliner);
+			if (UAvaOutlinerSubsystem* const OutlinerSubsystem = OutlinerPrivate.GetOutlinerSubsystem())
+			{
+				OutlinerSubsystem->BroadcastActorHierarchyChanged(ChildActor, ParentActor, ChangeType);
 			}
-			return true;
 		}
 	}
-	return false;
 }
 
 bool FAvaOutlinerActor::RemoveChild(const FAvaOutlinerRemoveItemParams& InRemoveItemParams)
diff --git a/Engine/Plugins/Experimental/Avalanche/Source/AvalancheOutliner/Private/Item/AvaOutlinerComponent.cpp b/Engine/Plugins/Experimental/Avalanche/Source/AvalancheOutliner/Private/Item/AvaOutlinerComponent.cpp
--- a/Engine/Plugins/Experimental/Avalanche/Source/AvalancheOutliner/Private/Item/AvaOutlinerComponent.cpp
+++ b/Engine/Plugins/Experimental/Avalanche/Source/AvalancheOutliner/Private/Item/AvaOutlinerComponent.cpp
@@ -57,38 +57,46 @@ void FAvaOutlinerComponent::GetItemProxies(TArray<TSharedPtr<FAvaOutlinerItemPro
 	}
 }
 
-bool FAvaOutlinerComponent::AddChild(const FAvaOutlinerAddItemParams& InAddItemParams)
+bool FAvaOutlinerComponent::CanAddChild(const FAvaOutlinerItemPtr& InChild) const
 {
-	if (CanAddChild(InAddItemParams.Item))
+	// Super checks InChild.IsValid
+	return Super::CanAddChild(InChild) && InChild->IsA<FAvaOutlinerComponent>();
+}
+
+void FAvaOutlinerComponent::OnChildrenAdded(TConstArrayView<FAvaOutlinerItemPtr> InAddedChildren, const FAvaOutlinerAddItemParams& InAddItemParams)
+{
+	Super::OnChildrenAdded(InAddedChildren, InAddItemParams);
+
+	for (const FAvaOutlinerItemPtr& Item : InAddedChildren)
 	{
-		if (const FAvaOutlinerComponent* const ComponentItem = InAddItemParams.Item->CastTo<FAvaOutlinerComponent>())
+		FAvaOutlinerComponent* ComponentItem = Item->CastTo<FAvaOutlinerComponent>();
+		if (!ComponentItem)
 		{
-			AddChildChecked(InAddItemParams);
-			if (USceneComponent* const ChildComp = ComponentItem->GetComponent())
-			{
-				USceneComponent* const UnderlyingComponent = GetComponent();
-				USceneComponent* const AttachParent = ChildComp->GetAttachParent();
+			continue;
+		}
 
-				//Set Component Mobility to Moveable.
-				if (ChildComp && ChildComp->Mobility == EComponentMobility::Static
-					&& AttachParent && AttachParent->Mobility != EComponentMobility::Static)
-				{
-					ChildComp->SetMobility(AttachParent->Mobility);
-				}
+		if (USceneComponent* const ChildComp = ComponentItem->GetComponent())
+		{
+			USceneComponent* const UnderlyingComponent = GetComponent();
+			USceneComponent* const AttachParent = ChildComp->GetAttachParent();
+
+			//Set Component Mobility to Moveable.
+			if (ChildComp && ChildComp->Mobility == EComponentMobility::Static
+				&& AttachParent && AttachParent->Mobility != EComponentMobility::Static)
+			{
+				ChildComp->SetMobility(AttachParent->Mobility);
+			}
 				
-				if (UnderlyingComponent && AttachParent != UnderlyingComponent)
+			if (UnderlyingComponent && AttachParent != UnderlyingComponent)
+			{
+				if (AttachParent)
 				{
-					if (AttachParent)
-					{
-						AttachParent->Modify();
-					}
-					ChildComp->AttachToComponent(UnderlyingComponent, FAttachmentTransformRules::KeepRelativeTransform, NAME_None);
+					AttachParent->Modify();
 				}
+				ChildComp->AttachToComponent(UnderlyingComponent, FAttachmentTransformRules::KeepRelativeTransform, NAME_None);
 			}
-			return true;
 		}
 	}
-	return false;
 }
 
 bool FAvaOutlinerComponent::RemoveChild(const FAvaOutlinerRemoveItemParams& InRemoveItemParams)
diff --git a/Engine/Plugins/Experimental/Avalanche/Source/AvalancheOutliner/Private/Item/AvaOutlinerItem.cpp b/Engine/Plugins/Experimental/Avalanche/Source/AvalancheOutliner/Private/Item/AvaOutlinerItem.cpp
--- a/Engine/Plugins/Experimental/Avalanche/Source/AvalancheOutliner/Private/Item/AvaOutlinerItem.cpp
+++ b/Engine/Plugins/Experimental/Avalanche/Source/AvalancheOutliner/Private/Item/AvaOutlinerItem.cpp
@@ -163,14 +163,9 @@ bool FAvaOutlinerItem::CanAddChild(const FAvaOutlinerItemPtr& InChild) const
 	return InChild.IsValid();
 }
 
-bool FAvaOutlinerItem::AddChild(const FAvaOutlinerAddItemParams& InAddItemParams)
+TArray<FAvaOutlinerItemPtr> FAvaOutlinerItem::AddChildren(const FAvaOutlinerAddItemParams& InAddItemParams)
 {
-	if (CanAddChild(InAddItemParams.Item))
-	{
-		AddChildChecked(InAddItemParams);
-		return true;
-	}
-	return false;
+	return AddChildrenInternal(InAddItemParams);
 }
 
 bool FAvaOutlinerItem::RemoveChild(const FAvaOutlinerRemoveItemParams& InRemoveItemParams)
@@ -302,22 +297,32 @@ void FAvaOutlinerItem::RecalculateItemId()
 	OutlinerPrivate.NotifyItemIdChanged(OldItemId, SharedThis(this));
 }
 
-void FAvaOutlinerItem::AddChildChecked(const FAvaOutlinerAddItemParams& InAddItemParams)
+TArray<FAvaOutlinerItemPtr> FAvaOutlinerItem::AddChildrenInternal(const FAvaOutlinerAddItemParams& InAddItemParams)
 {
-	if (const FAvaOutlinerItemPtr OldParent = InAddItemParams.Item->GetParent())
-	{
-		// if we are adding the child and the old parent is this,
-		// then it means we're just rearranging, only remove from array
-		if (OldParent.Get() == this)
+	// Remove all items that can't be added as children
+	TArray<FAvaOutlinerItemPtr> ChildrenToAdd(InAddItemParams.Items);
+	ChildrenToAdd.RemoveAll(
+		[This=this](const FAvaOutlinerItemPtr& Item)
 		{
-			Children.Remove(InAddItemParams.Item);
-		}
-		else
-		{
-			FAvaOutlinerRemoveItemParams RemoveParams(InAddItemParams.Item);
-			RemoveParams.DetachmentTransformRules = UE::AvaOutliner::Private::GetDetachmentRulesFromAttachmentRules(InAddItemParams.AttachmentTransformRules);
+			return !This->CanAddChild(Item);
+		});
 
-			OldParent->RemoveChild(RemoveParams);
+	for (const FAvaOutlinerItemPtr& Item : ChildrenToAdd)
+	{
+		if (const FAvaOutlinerItemPtr OldParent = Item->GetParent())
+		{
+			// if we are adding the child and the old parent is this,
+			// then it means we're just rearranging, only remove from array
+			if (OldParent.Get() == this)
+			{
+				Children.Remove(Item);
+			}
+			else
+			{
+				FAvaOutlinerRemoveItemParams RemoveParams(Item);
+				RemoveParams.DetachmentTransformRules = UE::AvaOutliner::Private::GetDetachmentRulesFromAttachmentRules(InAddItemParams.AttachmentTransformRules);
+				OldParent->RemoveChild(RemoveParams);
+			}
 		}
 	}
 
@@ -330,19 +335,26 @@ void FAvaOutlinerItem::AddChildChecked(const FAvaOutlinerAddItemParams& InAddIte
 				? RelativeItemIndex + 1
 				: RelativeItemIndex;
 
-			Children.EmplaceAt(TargetIndex, InAddItemParams.Item);
+			Children.Insert(ChildrenToAdd, TargetIndex);
 		}
 		else
 		{
-			Children.Emplace(InAddItemParams.Item);
+			Children.Append(ChildrenToAdd);
 		}
 	}
 	else
 	{
-		Children.Emplace(InAddItemParams.Item);
+		Children.Append(ChildrenToAdd);
+	}
+
+	// Update Parents
+	for (const FAvaOutlinerItemPtr& Item : ChildrenToAdd)
+	{
+		Item->SetParent(SharedThis(this));
 	}
 
-	InAddItemParams.Item->SetParent(SharedThis(this));
+	OnChildrenAdded(ChildrenToAdd, InAddItemParams);
+	return ChildrenToAdd;
 }
 
 bool FAvaOutlinerItem::RemoveChildChecked(const FAvaOutlinerRemoveItemParams& InRemoveItemParams)
@@ -368,7 +380,7 @@ void FAvaOutlinerItem::HandleNewSortableChildren(TArray<FAvaOutlinerItemPtr> InS
 	FAvaOutlinerAddItemParams AddItemParams;
 	for (const FAvaOutlinerItemPtr& NewChild : InSortableChildren)
 	{
-		AddItemParams.Item = NewChild;
+		AddItemParams.Items = { NewChild };
 
 		const FAvaSceneTreeNode* const TreeNode = SceneTree
 			? SceneTree->FindTreeNode(FAvaOutliner::MakeSceneItemFromOutlinerItem(NewChild))
@@ -387,7 +399,7 @@ void FAvaOutlinerItem::HandleNewSortableChildren(TArray<FAvaOutlinerItemPtr> InS
 			AddItemParams.RelativeDropZone = EItemDropZone::BelowItem;
 		}
 
-		AddChild(AddItemParams);
+		AddChildren(AddItemParams);
 	}
 }
 
@@ -405,7 +417,6 @@ FReply FAvaOutlinerItem::CreateItemsFromAssetDrop(const TSharedPtr<FAssetDragDro
 
 		bool bTransacting = false;
 		int32 TransactionIndex = INDEX_NONE;
-		int32 ItemsAdded  = 0;
 
 		const TArray<FAssetData>& Assets = AssetDragDropOp->GetAssets();
 
@@ -440,33 +451,34 @@ FReply FAvaOutlinerItem::CreateItemsFromAssetDrop(const TSharedPtr<FAssetDragDro
 					, Outliner.GetProvider().GetOutlinerDefaultActorSpawnTransform()
 					, SpawnParams))
 				{
-					++ItemsAdded;
-
 					SpawnedActor->InvalidateLightingCache();
 					SpawnedActor->PostEditChange();
 					SpawnedActor->MarkPackageDirty();
 
 					ULevel::LevelDirtiedEvent.Broadcast();
 
-					AddChildParams.Item = Outliner.FindOrAdd<FAvaOutlinerActor>(SpawnedActor);
-					if (ItemsAdded > 1)
+					AddChildParams.Items.Add(Outliner.FindOrAdd<FAvaOutlinerActor>(SpawnedActor));
+					if (AddChildParams.Items.Num() > 1)
 					{
 						// When there is more than 1 item being added, start appending to Selection Instead
 						AddChildParams.SelectionFlags |= EAvaOutlinerItemSelectionFlags::AppendToCurrentSelection;
 					}
-
-					FAvaOutliner& OutlinerPrivate = static_cast<FAvaOutliner&>(Outliner);
-					OutlinerPrivate.EnqueueItemAction<FAvaOutlinerAddItem>(AddChildParams);
 				}
 
 				Outliner.SetIgnoreNotify(EAvaOutlinerIgnoreNotifyFlags::Spawn, false);
 			}
 		}
 
+		if (!AddChildParams.Items.IsEmpty())
+		{
+			FAvaOutliner& OutlinerPrivate = static_cast<FAvaOutliner&>(Outliner);
+			OutlinerPrivate.EnqueueItemAction<FAvaOutlinerAddItem>(AddChildParams);
+		}
+
 		if (bTransacting)
 		{
 			// Likely: Items have been added if Transaction Started is true
-			if (ItemsAdded > 0)
+			if (!AddChildParams.Items.IsEmpty())
 			{
 				// Refresh here so that there's no delay between Spawn and Acknowledgement
 				Outliner.Refresh();
diff --git a/Engine/Plugins/Experimental/Avalanche/Source/AvalancheOutliner/Private/Item/AvaOutlinerTreeRoot.cpp b/Engine/Plugins/Experimental/Avalanche/Source/AvalancheOutliner/Private/Item/AvaOutlinerTreeRoot.cpp
--- a/Engine/Plugins/Experimental/Avalanche/Source/AvalancheOutliner/Private/Item/AvaOutlinerTreeRoot.cpp
+++ b/Engine/Plugins/Experimental/Avalanche/Source/AvalancheOutliner/Private/Item/AvaOutlinerTreeRoot.cpp
@@ -49,37 +49,57 @@ void FAvaOutlinerTreeRoot::FindChildren(TArray<FAvaOutlinerItemPtr>& OutChildren
 
 bool FAvaOutlinerTreeRoot::CanAddChild(const FAvaOutlinerItemPtr& InChild) const
 {
-	return Super::CanAddChild(InChild)
-		&& InChild->CanBeTopLevel();
+	return Super::CanAddChild(InChild) && InChild->CanBeTopLevel();
 }
 
-bool FAvaOutlinerTreeRoot::AddChild(const FAvaOutlinerAddItemParams& InAddItemParams)
+TArray<FAvaOutlinerItemPtr> FAvaOutlinerTreeRoot::AddChildren(const FAvaOutlinerAddItemParams& InAddItemParams)
 {
-	// if current parent is the root, it means it's just rearranging
-	const bool bRearranging = InAddItemParams.Item.IsValid() && InAddItemParams.Item->GetParent().Get() == this;
+	/** Holds whether an item is spawned (true) or rearranged (false) */
+	TMap<FAvaOutlinerItemPtr, EAvaOutlinerHierarchyChangeType> HierarchyChangeTypeMap;
+	HierarchyChangeTypeMap.Reserve(InAddItemParams.Items.Num());
 
-	// Is it a new actor that we just spawned
-	const bool bSpawning = InAddItemParams.Item.IsValid() && !Children.Contains(InAddItemParams.Item);
+	// Store whether the item is being spawned or rearranged
+	for (const FAvaOutlinerItemPtr& Item : InAddItemParams.Items)
+	{
+		if (!Item.IsValid())
+		{
+			continue;
+		}
 
-	const bool bResult = FAvaOutlinerItem::AddChild(InAddItemParams);
+		// if current parent is the root, it means it's just rearranging
+		const bool bRearranging = Item->GetParent().Get() == this;
 
-	if ((bSpawning || bRearranging) && InAddItemParams.Item->IsA<FAvaOutlinerActor>())
-	{
-		FAvaOutliner& OutlinerPrivate = static_cast<FAvaOutliner&>(Outliner);
+		// Is it a new actor that we just spawned
+		const bool bSpawning = !Children.Contains(Item);
 
-		UAvaOutlinerSubsystem* const OutlinerSubsystem = OutlinerPrivate.GetOutlinerSubsystem();
-		
-		AActor* const Actor = StaticCastSharedPtr<FAvaOutlinerActor>(InAddItemParams.Item)->GetActor();
-		
-		if (OutlinerSubsystem && Actor)
+		if (bRearranging || bSpawning)
 		{
-			const EAvaOutlinerHierarchyChangeType Type = bSpawning ? EAvaOutlinerHierarchyChangeType::Attached : EAvaOutlinerHierarchyChangeType::Rearranged;
-			
-			OutlinerSubsystem->BroadcastActorHierarchyChanged(Actor, nullptr, Type);
+			HierarchyChangeTypeMap.Add(Item, bSpawning ? EAvaOutlinerHierarchyChangeType::Attached : EAvaOutlinerHierarchyChangeType::Rearranged);
 		}
 	}
-	
-	return bResult;
+
+	const TArray<FAvaOutlinerItemPtr> AddedChildren = FAvaOutlinerItem::AddChildren(InAddItemParams);
+
+	FAvaOutliner& OutlinerPrivate = static_cast<FAvaOutliner&>(Outliner);
+
+	if (UAvaOutlinerSubsystem* const OutlinerSubsystem = OutlinerPrivate.GetOutlinerSubsystem())
+	{
+		for (const FAvaOutlinerItemPtr& AddedChild : AddedChildren)
+		{
+			if (const EAvaOutlinerHierarchyChangeType* HierarchyChangeType = HierarchyChangeTypeMap.Find(AddedChild))
+			{
+				if (const FAvaOutlinerActor* ActorItem = AddedChild->CastTo<FAvaOutlinerActor>())
+				{
+					if (AActor* const Actor = ActorItem->GetActor())
+					{
+						OutlinerSubsystem->BroadcastActorHierarchyChanged(Actor, /*ParentActor*/nullptr, *HierarchyChangeType);
+					}
+				}
+			}
+		}
+	}
+
+	return AddedChildren;
 }
 
 bool FAvaOutlinerTreeRoot::IsAllowedInOutliner() const
diff --git a/Engine/Plugins/Experimental/Avalanche/Source/AvalancheOutliner/Private/Item/AvaOutlinerTreeRoot.h b/Engine/Plugins/Experimental/Avalanche/Source/AvalancheOutliner/Private/Item/AvaOutlinerTreeRoot.h
--- a/Engine/Plugins/Experimental/Avalanche/Source/AvalancheOutliner/Private/Item/AvaOutlinerTreeRoot.h
+++ b/Engine/Plugins/Experimental/Avalanche/Source/AvalancheOutliner/Private/Item/AvaOutlinerTreeRoot.h
@@ -23,7 +23,7 @@ public:
 	//~ Begin IAvaOutlinerItem
 	virtual void FindChildren(TArray<FAvaOutlinerItemPtr>& OutChildren, bool bRecursive) override;
 	virtual bool CanAddChild(const FAvaOutlinerItemPtr& InChild) const override;
-	virtual bool AddChild(const FAvaOutlinerAddItemParams& InAddItemParams) override;
+	virtual TArray<FAvaOutlinerItemPtr> AddChildren(const FAvaOutlinerAddItemParams& InAddItemParams) override;
 	virtual bool IsAllowedInOutliner() const override;
 	virtual FText GetDisplayName() const override;
 	virtual FText GetClassName() const override;
diff --git a/Engine/Plugins/Experimental/Avalanche/Source/AvalancheOutliner/Private/ItemActions/AvaOutlinerAddItem.cpp b/Engine/Plugins/Experimental/Avalanche/Source/AvalancheOutliner/Private/ItemActions/AvaOutlinerAddItem.cpp
--- a/Engine/Plugins/Experimental/Avalanche/Source/AvalancheOutliner/Private/ItemActions/AvaOutlinerAddItem.cpp
+++ b/Engine/Plugins/Experimental/Avalanche/Source/AvalancheOutliner/Private/ItemActions/AvaOutlinerAddItem.cpp
@@ -1,6 +1,7 @@
 // Copyright Epic Games, Inc. All Rights Reserved.
 
 #include "ItemActions/AvaOutlinerAddItem.h"
+#include "Algo/AllOf.h"
 #include "AvaOutliner.h"
 #include "Item/AvaOutlinerObject.h"
 #include "Item/AvaOutlinerTreeRoot.h"
@@ -20,76 +21,93 @@ bool FAvaOutlinerAddItem::ShouldTransact() const
 
 void FAvaOutlinerAddItem::Execute(FAvaOutliner& InOutliner)
 {
-	if (AddParams.Item.IsValid())
+	// Pre-pass on all items
+	for (const FAvaOutlinerItemPtr& ItemToAdd : AddParams.Items)
 	{
+		if (!ItemToAdd.IsValid())
+		{
+			continue;
+		}
+
 		//Try to Create Children on Find
 		if (EnumHasAnyFlags(AddParams.Flags, EAvaOutlinerAddItemFlags::AddChildren))
 		{
 			constexpr bool bRecursiveFind = true;
 			TArray<FAvaOutlinerItemPtr> Children;
-			AddParams.Item->FindValidChildren(Children, bRecursiveFind);
+			ItemToAdd->FindValidChildren(Children, bRecursiveFind);
 		}
 
-		const FAvaOutlinerItemPtr ParentItem = AddParams.Item->GetParent();
+		const FAvaOutlinerItemPtr ParentItem = ItemToAdd->GetParent();
 
 		//If this Array has elements in it, then we need to stop a Circular Dependency from forming
+		TArray<FAvaOutlinerItemPtr> PathToRelativeItem = ItemToAdd->FindPath({AddParams.RelativeItem});
+		if (!PathToRelativeItem.IsEmpty() && ParentItem.IsValid())
 		{
-			TArray<FAvaOutlinerItemPtr> PathToRelativeItem = AddParams.Item->FindPath({AddParams.RelativeItem});
-			if (PathToRelativeItem.Num() > 0 && ParentItem.IsValid())
-			{
-				FAvaOutlinerAddItemParams CircularSolverParams;
-				CircularSolverParams.Item                     = PathToRelativeItem[0];
-				CircularSolverParams.RelativeItem             = AddParams.Item;
-				CircularSolverParams.RelativeDropZone         = EItemDropZone::AboveItem;
-				CircularSolverParams.Flags                    = AddParams.Flags;
-				CircularSolverParams.AttachmentTransformRules = AddParams.AttachmentTransformRules;
-
-				ParentItem->AddChild(CircularSolverParams);
-			}
+			FAvaOutlinerAddItemParams CircularSolverParams;
+			CircularSolverParams.Items                    = { PathToRelativeItem[0] };
+			CircularSolverParams.RelativeItem             = ItemToAdd;
+			CircularSolverParams.RelativeDropZone         = EItemDropZone::AboveItem;
+			CircularSolverParams.Flags                    = AddParams.Flags;
+			CircularSolverParams.AttachmentTransformRules = AddParams.AttachmentTransformRules;
+
+			ParentItem->AddChildren(CircularSolverParams);
 		}
+	}
 
-		if (AddParams.RelativeItem.IsValid())
-		{
-			const TSharedPtr<IAvaOutlinerItem> RelativeItemParent = AddParams.RelativeItem->GetParent();
+	// Item parent where the items will be added to. Defaults to tree root
+	TSharedPtr<IAvaOutlinerItem> ItemParent = InOutliner.GetTreeRoot();
 
-			//If it's onto item, the Relative Item is going to be the Parent
-			if (!AddParams.RelativeDropZone.IsSet() || AddParams.RelativeDropZone == EItemDropZone::OntoItem)
-			{
-				//If the Relative Item is Onto and it's the same as the Current Parent, shift Item up in the Hierarchy
-				//(as long as the parent is a valid one)
-				if (AddParams.RelativeItem == ParentItem && RelativeItemParent)
-				{
-					AddParams.RelativeDropZone = EItemDropZone::BelowItem;
-					RelativeItemParent->AddChild(AddParams);
-				}
-				else
+	if (AddParams.RelativeItem.IsValid())
+	{
+		const TSharedPtr<IAvaOutlinerItem> RelativeItemParent = AddParams.RelativeItem->GetParent();
+
+		// If it's onto item, the Relative Item is going to be the Parent
+		if (!AddParams.RelativeDropZone.IsSet() || AddParams.RelativeDropZone == EItemDropZone::OntoItem)
+		{
+			auto IsParentRelativeItem =
+				[&RelativeItem = AddParams.RelativeItem](const FAvaOutlinerItemPtr& InItem)
 				{
-					AddParams.RelativeItem->AddChild(AddParams);
-				}
-			}
-			//else we place it as a Sibling to the Relative Item
-			else if (RelativeItemParent)
+					return InItem.IsValid() && InItem->GetParent() == RelativeItem;
+				};
+
+			// If adding onto the Relative Item is and relative item is the Current Parent of ALL the dragged items,
+			// detach item from item by adding these items below the parent
+			if (RelativeItemParent && Algo::AllOf(AddParams.Items, IsParentRelativeItem))
 			{
-				RelativeItemParent->AddChild(AddParams);
+				ItemParent = RelativeItemParent;
+				AddParams.RelativeDropZone = EItemDropZone::BelowItem;
 			}
-			//if no parent, then add it to the tree root
 			else
 			{
-				InOutliner.GetTreeRoot()->AddChild(AddParams);
+				ItemParent = AddParams.RelativeItem;
 			}
 		}
-		else
+		//else we place it as a Sibling to the Relative Item
+		else if (RelativeItemParent)
 		{
-			//If no Relative Item, add to tree root
-			InOutliner.GetTreeRoot()->AddChild(AddParams);
+			ItemParent = RelativeItemParent;
 		}
+	}
 
-		if (const FEditorModeTools* const ModeTools = InOutliner.GetModeTools())
-		{
-			const FAvaOutlinerScopedSelection ScopedSelection(*ModeTools, EAvaOutlinerScopedSelectionPurpose::Read);
+	// Item parent is defaulted to tree root which is guaranteed to be valid.
+	// And item parent can only be set to a valid item thereafter
+	check(ItemParent.IsValid());
+	const TArray<FAvaOutlinerItemPtr> AddedChildren = ItemParent->AddChildren(AddParams);
 
-			const bool bSelectedInModeTools = AddParams.Item->IsSelected(ScopedSelection);
-			const bool bSelectedInOutliner  = InOutliner.GetSelectedItems().Contains(AddParams.Item);
+	if (AddedChildren.IsEmpty())
+	{
+		return;
+	}
+
+	// Sync selection from mode tools and outliner
+	if (const FEditorModeTools* const ModeTools = InOutliner.GetModeTools())
+	{
+		const FAvaOutlinerScopedSelection ScopedSelection(*ModeTools, EAvaOutlinerScopedSelectionPurpose::Read);
+
+		for (const FAvaOutlinerItemPtr& AddedChild : AddedChildren)
+		{
+			const bool bSelectedInModeTools = AddedChild->IsSelected(ScopedSelection);
+			const bool bSelectedInOutliner  = InOutliner.GetSelectedItems().Contains(AddedChild);
 
 			// Update Selection if there's a discrepancy between the Mode Tools Selection & the Outliner View Selected Items
 			if (bSelectedInModeTools != bSelectedInOutliner)
@@ -108,21 +126,24 @@ void FAvaOutlinerAddItem::Execute(FAvaOutliner& InOutliner)
 				}
 			}
 		}
+	}
 
-		if (EnumHasAnyFlags(AddParams.Flags, EAvaOutlinerAddItemFlags::Select))
-		{
-			InOutliner.SelectItems({AddParams.Item}, AddParams.SelectionFlags);
-		}
-
-		InOutliner.SetOutlinerModified();
+	if (EnumHasAnyFlags(AddParams.Flags, EAvaOutlinerAddItemFlags::Select))
+	{
+		InOutliner.SelectItems(AddedChildren, AddParams.SelectionFlags);
 	}
+
+	InOutliner.SetOutlinerModified();
 }
 
 void FAvaOutlinerAddItem::OnObjectsReplaced(const TMap<UObject*, UObject*>& InReplacementMap, bool bRecursive)
 {
-	if (AddParams.Item.IsValid())
+	for (const FAvaOutlinerItemPtr& Item : AddParams.Items)
 	{
-		AddParams.Item->OnObjectsReplaced(InReplacementMap, bRecursive);
+		if (Item.IsValid())
+		{
+			Item->OnObjectsReplaced(InReplacementMap, bRecursive);
+		}
 	}
 	if (AddParams.RelativeItem.IsValid())
 	{
diff --git a/Engine/Plugins/Experimental/Avalanche/Source/AvalancheOutliner/Private/ItemActions/AvaOutlinerRemoveItem.cpp b/Engine/Plugins/Experimental/Avalanche/Source/AvalancheOutliner/Private/ItemActions/AvaOutlinerRemoveItem.cpp
--- a/Engine/Plugins/Experimental/Avalanche/Source/AvalancheOutliner/Private/ItemActions/AvaOutlinerRemoveItem.cpp
+++ b/Engine/Plugins/Experimental/Avalanche/Source/AvalancheOutliner/Private/ItemActions/AvaOutlinerRemoveItem.cpp
@@ -48,17 +48,17 @@ void FAvaOutlinerRemoveItem::Execute(FAvaOutliner& InOutliner)
 			ItemsToReparent.Append(Children);
 			while (ItemsToReparent.Num() > 0)
 			{
-				ReparentParams.Item = ItemsToReparent.Pop();
+				ReparentParams.Items = { ItemsToReparent.Pop() };
 
 				//Then Parent it to the Item's Parent.
 				//If we couldn't add the Child, then it means either the Child Item itself is invalid
 				//or the underlying info (e.g. object) is invalid (e.g. actor pending kill)
-				if (Parent->AddChild(ReparentParams) == false)
+				if (Parent->AddChildren(ReparentParams).IsEmpty())
 				{
 					//so try reparenting the children of this invalid child (since this invalid child will be removed)
-					if (ReparentParams.Item.IsValid())
+					if (ReparentParams.Items[0].IsValid())
 					{
-						ItemsToReparent.Append(ReparentParams.Item->GetChildren());
+						ItemsToReparent.Append(ReparentParams.Items[0]->GetChildren());
 					}
 				}
 			}
diff --git a/Engine/Plugins/Experimental/Avalanche/Source/AvalancheOutliner/Public/Item/AvaOutlinerActor.h b/Engine/Plugins/Experimental/Avalanche/Source/AvalancheOutliner/Public/Item/AvaOutlinerActor.h
--- a/Engine/Plugins/Experimental/Avalanche/Source/AvalancheOutliner/Public/Item/AvaOutlinerActor.h
+++ b/Engine/Plugins/Experimental/Avalanche/Source/AvalancheOutliner/Public/Item/AvaOutlinerActor.h
@@ -20,7 +20,8 @@ public:
 	virtual void FindChildren(TArray<FAvaOutlinerItemPtr>& OutChildren, bool bRecursive) override;
 	virtual void GetItemProxies(TArray<TSharedPtr<FAvaOutlinerItemProxy>>& OutItemProxies) override;
 	virtual bool IsSortable() const override { return true; }
-	virtual bool AddChild(const FAvaOutlinerAddItemParams& InAddItemParams) override;
+	virtual bool CanAddChild(const FAvaOutlinerItemPtr& InChild) const override;
+	virtual void OnChildrenAdded(TConstArrayView<FAvaOutlinerItemPtr> InAddedChildren, const FAvaOutlinerAddItemParams& InAddItemParams) override;
 	virtual bool RemoveChild(const FAvaOutlinerRemoveItemParams& InRemoveItemParams) override;
 	virtual bool CanBeTopLevel() const override { return true; }
 	virtual bool IsAllowedInOutliner() const override;
diff --git a/Engine/Plugins/Experimental/Avalanche/Source/AvalancheOutliner/Public/Item/AvaOutlinerComponent.h b/Engine/Plugins/Experimental/Avalanche/Source/AvalancheOutliner/Public/Item/AvaOutlinerComponent.h
--- a/Engine/Plugins/Experimental/Avalanche/Source/AvalancheOutliner/Public/Item/AvaOutlinerComponent.h
+++ b/Engine/Plugins/Experimental/Avalanche/Source/AvalancheOutliner/Public/Item/AvaOutlinerComponent.h
@@ -19,7 +19,8 @@ public:
 	//~ Begin IAvaOutlinerItem
 	virtual void FindChildren(TArray<FAvaOutlinerItemPtr>& OutChildren, bool bRecursive) override;
 	virtual void GetItemProxies(TArray<TSharedPtr<FAvaOutlinerItemProxy>>& OutItemProxies) override;
-	virtual bool AddChild(const FAvaOutlinerAddItemParams& InAddItemParams) override;
+	virtual bool CanAddChild(const FAvaOutlinerItemPtr& InChild) const override;
+	virtual void OnChildrenAdded(TConstArrayView<FAvaOutlinerItemPtr> InAddedChildren, const FAvaOutlinerAddItemParams& InAddItemParams) override;
 	virtual bool RemoveChild(const FAvaOutlinerRemoveItemParams& InRemoveItemParams) override;
 	virtual EAvaOutlinerItemViewMode GetSupportedViewModes(const FAvaOutlinerView& InOutlinerView) const override;
 	virtual bool IsAllowedInOutliner() const override;
diff --git a/Engine/Plugins/Experimental/Avalanche/Source/AvalancheOutliner/Public/Item/AvaOutlinerItem.h b/Engine/Plugins/Experimental/Avalanche/Source/AvalancheOutliner/Public/Item/AvaOutlinerItem.h
--- a/Engine/Plugins/Experimental/Avalanche/Source/AvalancheOutliner/Public/Item/AvaOutlinerItem.h
+++ b/Engine/Plugins/Experimental/Avalanche/Source/AvalancheOutliner/Public/Item/AvaOutlinerItem.h
@@ -30,7 +30,7 @@ public:
 	virtual TArray<FAvaOutlinerItemPtr>& GetChildrenMutable() override { return Children; }
 	virtual bool IsSortable() const override { return false; }
 	virtual bool CanAddChild(const FAvaOutlinerItemPtr& InChild) const override;
-	virtual bool AddChild(const FAvaOutlinerAddItemParams& InAddItemParams) override;
+	virtual TArray<FAvaOutlinerItemPtr> AddChildren(const FAvaOutlinerAddItemParams& InAddItemParams) override;
 	virtual bool RemoveChild(const FAvaOutlinerRemoveItemParams& InRemoveItemParams) override;
 	virtual FAvaOutlinerItemPtr GetParent() const override { return ParentWeak.Pin(); }
 	virtual bool CanBeTopLevel() const override { return false; }
@@ -73,7 +73,14 @@ protected:
 	void RecalculateItemId();
 	
 	/** The actual implementation of putting the given Item under the Children array */
-	void AddChildChecked(const FAvaOutlinerAddItemParams& InAddItemParams);
+	UE_DEPRECATED(5.6, "Use AddChildrenInternal instead")
+	void AddChildChecked(const FAvaOutlinerAddItemParams& InAddItemParams)
+	{
+		AddChildrenInternal(InAddItemParams);
+	}
+
+	/** The actual implementation of putting the given Item under the Children array */
+	TArray<FAvaOutlinerItemPtr> AddChildrenInternal(const FAvaOutlinerAddItemParams& InAddItemParams);
 
 	/** The actual implementation of removing the given item from the Children array */
 	bool RemoveChildChecked(const FAvaOutlinerRemoveItemParams& InRemoveItemParams);
diff --git a/Engine/Plugins/Experimental/Avalanche/Source/AvalancheOutliner/Public/Item/AvaOutlinerItemParameters.h b/Engine/Plugins/Experimental/Avalanche/Source/AvalancheOutliner/Public/Item/AvaOutlinerItemParameters.h
--- a/Engine/Plugins/Experimental/Avalanche/Source/AvalancheOutliner/Public/Item/AvaOutlinerItemParameters.h
+++ b/Engine/Plugins/Experimental/Avalanche/Source/AvalancheOutliner/Public/Item/AvaOutlinerItemParameters.h
@@ -22,21 +22,36 @@ ENUM_CLASS_FLAGS(EAvaOutlinerAddItemFlags);
 
 struct FAvaOutlinerAddItemParams
 {
-	FAvaOutlinerAddItemParams(const FAvaOutlinerItemPtr& InItem = nullptr
+	/** *Note: We need to explicitly disable warnings on these constructors/operators for clang to be happy with deprecated variables */
+	PRAGMA_DISABLE_DEPRECATION_WARNINGS
+	FAvaOutlinerAddItemParams() = default;
+	~FAvaOutlinerAddItemParams() = default;
+	FAvaOutlinerAddItemParams(const FAvaOutlinerAddItemParams&) = default;
+	FAvaOutlinerAddItemParams(FAvaOutlinerAddItemParams&&) = default;
+	FAvaOutlinerAddItemParams& operator=(const FAvaOutlinerAddItemParams&) = default;
+	FAvaOutlinerAddItemParams& operator=(FAvaOutlinerAddItemParams&&) = default;
+	PRAGMA_ENABLE_DEPRECATION_WARNINGS
+
+	explicit FAvaOutlinerAddItemParams(const FAvaOutlinerItemPtr& InItem
 			, EAvaOutlinerAddItemFlags InFlags            = EAvaOutlinerAddItemFlags::None
 			, const FAvaOutlinerItemPtr& InRelativeItem   = nullptr
 			, TOptional<EItemDropZone> InRelativeDropZone = TOptional<EItemDropZone>())
-		: Item(InItem)
-		, RelativeItem(InRelativeItem)
+		: RelativeItem(InRelativeItem)
 		, RelativeDropZone(InRelativeDropZone)
 		, Flags(InFlags)
-		, SelectionFlags(EAvaOutlinerItemSelectionFlags::None)
 	{
+		if (InItem.IsValid())
+		{
+			Items.Add(InItem);
+		}
 	}
 
-	/** The Item to Add */
+	UE_DEPRECATED(5.6, "Use 'Items' instead")
 	FAvaOutlinerItemPtr Item;
 
+	/** The items to add */
+	TArray<FAvaOutlinerItemPtr, TInlineAllocator<1>> Items;
+
 	/** The Item to use as base in where to place the Item */
 	FAvaOutlinerItemPtr RelativeItem;
 
@@ -44,10 +59,10 @@ struct FAvaOutlinerAddItemParams
 	TOptional<EItemDropZone> RelativeDropZone;
 
 	/** Some Extra Flags for what to do when Adding or After Adding the Items */
-	EAvaOutlinerAddItemFlags Flags;
+	EAvaOutlinerAddItemFlags Flags = EAvaOutlinerAddItemFlags::None;
 
 	/** Flags to Indicate how we should Select the Item. This only applies if bool bSelectItem is true */
-	EAvaOutlinerItemSelectionFlags SelectionFlags;
+	EAvaOutlinerItemSelectionFlags SelectionFlags = EAvaOutlinerItemSelectionFlags::None;
 
 	/** Optional Transform override Rule when Attaching Items */
 	TOptional<FAttachmentTransformRules> AttachmentTransformRules;
diff --git a/Engine/Plugins/Experimental/Avalanche/Source/AvalancheOutliner/Public/Item/AvaOutlinerLevel.cpp b/Engine/Plugins/Experimental/Avalanche/Source/AvalancheOutliner/Public/Item/AvaOutlinerLevel.cpp
--- a/Engine/Plugins/Experimental/Avalanche/Source/AvalancheOutliner/Public/Item/AvaOutlinerLevel.cpp
+++ b/Engine/Plugins/Experimental/Avalanche/Source/AvalancheOutliner/Public/Item/AvaOutlinerLevel.cpp
@@ -46,12 +46,6 @@ void FAvaOutlinerLevel::FindChildren(TArray<FAvaOutlinerItemPtr>& OutChildren, b
 	}
 }
 
-bool FAvaOutlinerLevel::AddChild(const FAvaOutlinerAddItemParams& InAddItemParams)
-{
-	// todo: handle actor being moved to this level item
-	return Super::AddChild(InAddItemParams);
-}
-
 bool FAvaOutlinerLevel::RemoveChild(const FAvaOutlinerRemoveItemParams& InRemoveItemParams)
 {
 	// todo: handle actor being moved from this level item
diff --git a/Engine/Plugins/Experimental/Avalanche/Source/AvalancheOutliner/Public/Item/AvaOutlinerLevel.h b/Engine/Plugins/Experimental/Avalanche/Source/AvalancheOutliner/Public/Item/AvaOutlinerLevel.h
--- a/Engine/Plugins/Experimental/Avalanche/Source/AvalancheOutliner/Public/Item/AvaOutlinerLevel.h
+++ b/Engine/Plugins/Experimental/Avalanche/Source/AvalancheOutliner/Public/Item/AvaOutlinerLevel.h
@@ -18,7 +18,6 @@ public:
 
 	//~ Begin IAvaOutlinerItem
 	AVALANCHEOUTLINER_API virtual void FindChildren(TArray<FAvaOutlinerItemPtr>& OutChildren, bool bRecursive) override;
-	AVALANCHEOUTLINER_API virtual bool AddChild(const FAvaOutlinerAddItemParams& InAddItemParams) override;
 	AVALANCHEOUTLINER_API virtual bool RemoveChild(const FAvaOutlinerRemoveItemParams& InRemoveItemParams) override;
 	AVALANCHEOUTLINER_API virtual EAvaOutlinerItemViewMode GetSupportedViewModes(const FAvaOutlinerView& InOutlinerView) const override;
 	virtual bool CanBeTopLevel() const override { return true; }
diff --git a/Engine/Plugins/Experimental/Avalanche/Source/AvalancheOutliner/Public/Item/IAvaOutlinerItem.h b/Engine/Plugins/Experimental/Avalanche/Source/AvalancheOutliner/Public/Item/IAvaOutlinerItem.h
--- a/Engine/Plugins/Experimental/Avalanche/Source/AvalancheOutliner/Public/Item/IAvaOutlinerItem.h
+++ b/Engine/Plugins/Experimental/Avalanche/Source/AvalancheOutliner/Public/Item/IAvaOutlinerItem.h
@@ -90,7 +90,20 @@ public:
 	 * Adds another Child under this Item if such Item is supported.
 	 * Returns true if it did, false if it could not add it (e.g. item not supported).
 	 */
-	virtual bool AddChild(const FAvaOutlinerAddItemParams& InAddItemParams) = 0;
+	UE_DEPRECATED(5.6, "Use AddChildren instead")
+	virtual bool AddChild(const FAvaOutlinerAddItemParams& InAddItemParams) final
+	{
+		return AddChildren(InAddItemParams).Num() > 0;
+	}
+
+	/**
+	 * Adds the items in the add items struct under this item if such items are supported.
+	 * Returns the children that were successfully added
+	 */
+	virtual TArray<FAvaOutlinerItemPtr> AddChildren(const FAvaOutlinerAddItemParams& InAddItemParams) = 0;
+
+	/** Called when children have been added to this item */
+	virtual void OnChildrenAdded(TConstArrayView<FAvaOutlinerItemPtr> InAddedChildren, const FAvaOutlinerAddItemParams& InAddItemParams) {}
 
 	/**
 	 * Removes the given child from this Item if it was ever indeed a child.
