Motion Design: Better Material Designer preset picking
- Now uses translucent preset for textures with alpha and emissive for ones without it.

diff --git a/Engine/Plugins/Experimental/Avalanche/Source/AvalancheEditor/Private/MaterialDesigner/AvaMaterialDesignerTextureAssetFactory.cpp b/Engine/Plugins/Experimental/Avalanche/Source/AvalancheEditor/Private/MaterialDesigner/AvaMaterialDesignerTextureAssetFactory.cpp
--- a/Engine/Plugins/Experimental/Avalanche/Source/AvalancheEditor/Private/MaterialDesigner/AvaMaterialDesignerTextureAssetFactory.cpp
+++ b/Engine/Plugins/Experimental/Avalanche/Source/AvalancheEditor/Private/MaterialDesigner/AvaMaterialDesignerTextureAssetFactory.cpp
@@ -16,6 +16,7 @@
 #include "DynamicMeshes/AvaShapeRectangleDynMesh.h"
 #include "Engine/Level.h"
 #include "Engine/Texture.h"
+#include "Engine/Texture2D.h"
 #include "Framework/Application/SlateApplication.h"
 #include "IDynamicMaterialEditorModule.h"
 #include "Material/DynamicMaterialInstance.h"
@@ -124,7 +125,24 @@ AActor* UAvaMaterialDesignerTextureAssetFactory::SpawnActor(UObject* InAsset, UL
 		return InNewActor;
 	}
 
-	EditorOnlyData->SetChannelListPreset(TEXT("Emissive"));
+	// Assume true if it's not a texture 2d.
+	// It's not reasonable to check the texture format of every single type of texture.
+	bool bHasAlphaChannel = true;
+
+	if (UTexture2D* Texture2D = Cast<UTexture2D>(Texture))
+	{
+		bHasAlphaChannel = Texture2D->HasAlphaChannel();
+	}
+
+	if (bHasAlphaChannel)
+	{
+		EditorOnlyData->SetChannelListPreset(TEXT("Translucent"));
+	}
+	else
+	{
+		EditorOnlyData->SetChannelListPreset(TEXT("Emissive"));
+	}
+
 	EditorOnlyData->OnWizardComplete();
 
 	const UDMMaterialSlot* Slot = EditorOnlyData->GetSlotForMaterialProperty(EDMMaterialPropertyType::EmissiveColor);
