[Motion Design] Implement Playback Transition Start Synchronization on display cluster

The playback transitions now have an additional AsyncAssetLoader (for UE-256696) which adds non-determinism. It is no longer sufficient to synchronize only the level streaming.

We need to synchronize the start of the transitions too, once all assets are loaded and ready.

diff --git a/Engine/Plugins/Experimental/Avalanche/Source/AvalancheMedia/Private/Playback/Transition/AvaPlaybackServerTransition.cpp b/Engine/Plugins/Experimental/Avalanche/Source/AvalancheMedia/Private/Playback/Transition/AvaPlaybackServerTransition.cpp
--- a/Engine/Plugins/Experimental/Avalanche/Source/AvalancheMedia/Private/Playback/Transition/AvaPlaybackServerTransition.cpp
+++ b/Engine/Plugins/Experimental/Avalanche/Source/AvalancheMedia/Private/Playback/Transition/AvaPlaybackServerTransition.cpp
@@ -2,6 +2,7 @@
 
 #include "Playback/Transition/AvaPlaybackServerTransition.h"
 
+#include "Playable/AvaPlayableGroupManager.h"
 #include "Playable/Transition/AvaPlayableTransition.h"
 #include "Playback/AvaPlaybackManager.h"
 #include "Playback/AvaPlaybackServer.h"
@@ -327,11 +328,35 @@ bool UAvaPlaybackServerTransition::CanStart(bool& bOutShouldDiscard)
 }
 
 void UAvaPlaybackServerTransition::Start()
+{
+	TWeakObjectPtr<UAvaPlaybackServerTransition> ThisWeak(this);
+	auto StartEventHandler = [ThisWeak]()
+	{
+		if (UAvaPlaybackServerTransition* This = ThisWeak.Get())
+		{
+			This->Start_Synchronized();
+		}
+	};
+
+	FString StartEventSignature = FString(TEXT("PlaybackServerTransitionStart_")) + TransitionId.ToString();
+	
+	if (TSharedPtr<FAvaPlaybackServer> PlaybackServer = PlaybackServerWeak.Pin())
+	{
+		if (UAvaPlayableGroupManager* GroupManager = PlaybackServer->GetPlaybackManager().GetPlayableGroupManager())
+		{
+			if (!GroupManager->IsSynchronizedEventPushed(StartEventSignature))
+			{
+				GroupManager->PushSynchronizedEvent(MoveTemp(StartEventSignature), MoveTemp(StartEventHandler));
+			}
+		}
+	}
+}
+
+void UAvaPlaybackServerTransition::Start_Synchronized()
 {
 	RegisterToPlayableTransitionEvent();
 
-	// May fail if playables are not loaded yet. Playables are loaded
-	// when the playback object has ticked at least one.
+	// Playables should be loaded at this point since the synchronized part is to wait on loading assets.
 	MakePlayableTransition();
 
 	bool bTransitionStarted = false;
@@ -339,9 +364,6 @@ void UAvaPlaybackServerTransition::Start()
 	if (PlayableTransition)
 	{
 		LogDetailedTransitionInfo();
-
-		// Todo: validate the level streaming playables are finished streaming the asset.
-		// Otherwise, transition start must be queued on playable streaming events. 
 		bTransitionStarted = PlayableTransition->Start();
 	}
 
diff --git a/Engine/Plugins/Experimental/Avalanche/Source/AvalancheMedia/Private/Playback/Transition/AvaPlaybackServerTransition.h b/Engine/Plugins/Experimental/Avalanche/Source/AvalancheMedia/Private/Playback/Transition/AvaPlaybackServerTransition.h
--- a/Engine/Plugins/Experimental/Avalanche/Source/AvalancheMedia/Private/Playback/Transition/AvaPlaybackServerTransition.h
+++ b/Engine/Plugins/Experimental/Avalanche/Source/AvalancheMedia/Private/Playback/Transition/AvaPlaybackServerTransition.h
@@ -68,6 +68,9 @@ public:
 	FString GetBriefTransitionDescription() const;
 	
 protected:
+	/** Implementation of the Start function intended to be synchronized on cluster. */
+	void Start_Synchronized();
+
 	TSharedPtr<FAvaPlaybackInstance> FindInstanceForPlayable(const UAvaPlayable* InPlayable);
 
 	void OnTransitionEvent(UAvaPlayable* InPlayable, UAvaPlayableTransition* InTransition, EAvaPlayableTransitionEventFlags InTransitionFlags);
diff --git a/Engine/Plugins/Experimental/Avalanche/Source/AvalancheMedia/Private/Playback/Transition/AvaPlaybackTransition.h b/Engine/Plugins/Experimental/Avalanche/Source/AvalancheMedia/Private/Playback/Transition/AvaPlaybackTransition.h
--- a/Engine/Plugins/Experimental/Avalanche/Source/AvalancheMedia/Private/Playback/Transition/AvaPlaybackTransition.h
+++ b/Engine/Plugins/Experimental/Avalanche/Source/AvalancheMedia/Private/Playback/Transition/AvaPlaybackTransition.h
@@ -23,16 +23,23 @@ public:
 	 */
 	virtual bool CanStart(bool& bOutShouldDiscard) { bOutShouldDiscard = true; return false; }
 
+	/** 
+	* Start the transition.
+	* @remark Transition start is synchronized on clusters, but it is implemented by the derived classes.
+	 */
 	virtual void Start() {}
 
+	/** Stop transition */
 	virtual void Stop() {}
 
+	/** Returns true if the transition is running, false otherwise. */
 	virtual bool IsRunning() const { return false; }
 	
 	//~ Begin IAvaPlayableVisibilityConstraint
 	virtual bool IsVisibilityConstrained(const UAvaPlayable* InPlayable) const override { return false; }
 	//~ End IAvaPlayableVisibilityConstraint
 
+	/** Returns the TransitionId, a unique identifier (replicated on server) for this transition. */
 	const FGuid& GetTransitionId() const { return TransitionId; }
 	
 protected:
diff --git a/Engine/Plugins/Experimental/Avalanche/Source/AvalancheMedia/Private/Rundown/Transition/AvaRundownPageTransition.cpp b/Engine/Plugins/Experimental/Avalanche/Source/AvalancheMedia/Private/Rundown/Transition/AvaRundownPageTransition.cpp
--- a/Engine/Plugins/Experimental/Avalanche/Source/AvalancheMedia/Private/Rundown/Transition/AvaRundownPageTransition.cpp
+++ b/Engine/Plugins/Experimental/Avalanche/Source/AvalancheMedia/Private/Rundown/Transition/AvaRundownPageTransition.cpp
@@ -2,6 +2,7 @@
 
 #include "Rundown/Transition/AvaRundownPageTransition.h"
 
+#include "Playable/AvaPlayableGroupManager.h"
 #include "Playable/Transition/AvaPlayableTransition.h"
 #include "Playback/AvaPlaybackManager.h"
 #include "Playback/AvaPlaybackUtils.h"
@@ -261,11 +262,36 @@ bool UAvaRundownPageTransition::CanStart(bool& bOutShouldDiscard)
 }
 
 void UAvaRundownPageTransition::Start()
+{
+	TWeakObjectPtr<UAvaRundownPageTransition> ThisWeak(this);
+	auto StartEventHandler = [ThisWeak]()
+	{
+		if (UAvaRundownPageTransition* This = ThisWeak.Get())
+		{
+			This->Start_Synchronized();
+		}
+	};
+
+	// Build unique signature for this event.
+	FString StartEventSignature = FString(TEXT("PageTransitionStart_")) + TransitionId.ToString();
+	
+	if (const UAvaRundown* Rundown = GetRundown())
+	{
+		if (UAvaPlayableGroupManager* GroupManager = Rundown->GetPlaybackManager().GetPlayableGroupManager())
+		{
+			if (!GroupManager->IsSynchronizedEventPushed(StartEventSignature))
+			{
+				GroupManager->PushSynchronizedEvent(MoveTemp(StartEventSignature), MoveTemp(StartEventHandler));
+			}
+		}
+	}
+}
+
+void UAvaRundownPageTransition::Start_Synchronized()
 {
 	RegisterToPlayableTransitionEvent();
 
-	// May fail if playables are not loaded yet. Playables are loaded
-	// when the playback object has ticked at least one.
+	// Playables should be loaded at this point since the synchronized part is to wait on loading assets.
 	MakePlayableTransition();
 
 	bool bTransitionStarted = false;
@@ -273,9 +299,6 @@ void UAvaRundownPageTransition::Start()
 	if (PlayableTransition)
 	{
 		LogDetailedTransitionInfo();
-
-		// Todo: validate the level streaming playables are finished streaming the asset.
-		// Otherwise, transition start must be queued on playable streaming events. 
 		bTransitionStarted = PlayableTransition->Start();
 	}
 	
diff --git a/Engine/Plugins/Experimental/Avalanche/Source/AvalancheMedia/Private/Rundown/Transition/AvaRundownPageTransition.h b/Engine/Plugins/Experimental/Avalanche/Source/AvalancheMedia/Private/Rundown/Transition/AvaRundownPageTransition.h
--- a/Engine/Plugins/Experimental/Avalanche/Source/AvalancheMedia/Private/Rundown/Transition/AvaRundownPageTransition.h
+++ b/Engine/Plugins/Experimental/Avalanche/Source/AvalancheMedia/Private/Rundown/Transition/AvaRundownPageTransition.h
@@ -86,6 +86,9 @@ public:
 	bool  bIsPreviewFrameTransition = false;
 
 protected:
+	/** Implementation of the Start function intended to be synchronized on cluster. */
+	void Start_Synchronized();
+	
 	UAvaRundownPlaybackInstancePlayer* FindInstancePlayerForPlayable(const UAvaPlayable* InPlayable) const;
 	UAvaRundownPagePlayer* FindPagePlayerForPlayable(const UAvaPlayable* InPlayable) const;
 
