Motion Design: fix crash with parameter details

diff --git a/Engine/Plugins/Experimental/Avalanche/Source/AvalancheTransitionEditor/Private/Customizations/AvaTransitionTreeEditorDataCustomization.cpp b/Engine/Plugins/Experimental/Avalanche/Source/AvalancheTransitionEditor/Private/Customizations/AvaTransitionTreeEditorDataCustomization.cpp
--- a/Engine/Plugins/Experimental/Avalanche/Source/AvalancheTransitionEditor/Private/Customizations/AvaTransitionTreeEditorDataCustomization.cpp
+++ b/Engine/Plugins/Experimental/Avalanche/Source/AvalancheTransitionEditor/Private/Customizations/AvaTransitionTreeEditorDataCustomization.cpp
@@ -54,14 +54,11 @@ void FAvaTransitionTreeEditorDataCustomization::CustomizeDetails(IDetailLayoutBu
 
 void FAvaTransitionTreeEditorDataCustomization::CustomizeParameters(IDetailLayoutBuilder& InDetailBuilder)
 {
-	TSharedPtr<IPropertyHandle> RootParametersHandle = InDetailBuilder.GetProperty(TEXT("RootParameters"));
+	TSharedPtr<IPropertyHandle> RootParametersHandle = InDetailBuilder.GetProperty(TEXT("RootParameterPropertyBag"));
 	check(RootParametersHandle);
 	RootParametersHandle->MarkHiddenByCustomization();
 
-	TSharedPtr<IPropertyHandle> PropertyBagParametersHandle = RootParametersHandle->GetChildHandle(TEXT("Parameters")); // FInstancedPropertyBag
-	check(PropertyBagParametersHandle);
-
-	TSharedRef<FPropertyBagInstanceDataDetails> InstanceDetails = MakeShared<FPropertyBagInstanceDataDetails>(PropertyBagParametersHandle
+	TSharedRef<FPropertyBagInstanceDataDetails> InstanceDetails = MakeShared<FPropertyBagInstanceDataDetails>(RootParametersHandle
 		, InDetailBuilder.GetPropertyUtilities()
 		, /*bFixedLayout*/true);
 
