[Motion Design] Adding Playable Blueprint Latent action to synchronize an event.

diff --git a/Engine/Plugins/Experimental/Avalanche/Source/AvalancheMedia/Private/Playable/AvaPlayableLibrary.cpp b/Engine/Plugins/Experimental/Avalanche/Source/AvalancheMedia/Private/Playable/AvaPlayableLibrary.cpp
--- a/Engine/Plugins/Experimental/Avalanche/Source/AvalancheMedia/Private/Playable/AvaPlayableLibrary.cpp
+++ b/Engine/Plugins/Experimental/Avalanche/Source/AvalancheMedia/Private/Playable/AvaPlayableLibrary.cpp
@@ -119,3 +119,127 @@ bool UAvaPlayableLibrary::SetPlayableHidden(const UObject* InWorldContextObject,
 	return false;
 }
 
+/**
+ * Pending Latent Action to push a sync event on the cluster.
+ * The event signature includes the playable id as well.
+ * The given signature need only be unique in terms of any sequence of synchronized events for the given playable.
+ */
+class FAvaPlayableSyncEventLatentAction : public FPendingLatentAction
+{
+public:
+	FAvaPlayableSyncEventLatentAction(const FLatentActionInfo& InLatentInfo, UAvaPlayable* InPlayable, const FString& InEventSignature, bool& InSuccess)
+		: ExecutionFunction(InLatentInfo.ExecutionFunction)
+		, OutputLink(InLatentInfo.Linkage)
+		, CallbackTarget(InLatentInfo.CallbackTarget)
+		, TimeRemaining(10.0)
+		, OutSuccess(InSuccess)
+	{
+		if (UAvaPlayableGroup* Group = (InPlayable ? InPlayable->GetPlayableGroup() : nullptr))
+		{
+			// Build a unique playable id string. 
+			const FString PlayableId = InPlayable->GetInstanceId().IsValid() ? InPlayable->GetInstanceId().ToString() : InPlayable->GetSourceAssetPath().ToString();
+
+			// Build unique signature for this event.
+			FString PlayableEventSignature = FString(TEXT("PlayableSyncEvent_")) + PlayableId + FString(TEXT("_")) + InEventSignature;
+	
+			if (!Group->IsSynchronizedEventPushed(PlayableEventSignature))
+			{
+				SyncFence = MakeShared<bool>(false);
+				TWeakPtr<bool> SyncFenceWeak = SyncFence;
+
+				auto SyncEventHandler = [SyncFenceWeak]()
+				{
+					if (TSharedPtr<bool> SyncFence = SyncFenceWeak.Pin())
+					{
+						*SyncFence = true;
+					}
+				};
+
+				// Keep a copy of the signature for the description and logs.
+				EventSignature = PlayableEventSignature;
+
+				Group->PushSynchronizedEvent(MoveTemp(PlayableEventSignature), MoveTemp(SyncEventHandler));
+			}
+		}
+	}
+	
+	void FailedOperation(FLatentResponse& InResponse) const
+	{
+		OutSuccess = false;
+		InResponse.FinishAndTriggerIf(true, ExecutionFunction, OutputLink, CallbackTarget);
+	}
+
+	//~ Begin FPendingLatentAction
+	virtual void UpdateOperation(FLatentResponse& InResponse) override
+	{
+		if (!SyncFence)
+		{
+			FailedOperation(InResponse);
+		}
+		
+		if (*SyncFence)
+		{
+			OutSuccess = true;
+			InResponse.FinishAndTriggerIf(true, ExecutionFunction, OutputLink, CallbackTarget);
+			return;
+		}
+
+		// Timed out
+		TimeRemaining -= InResponse.ElapsedTime();
+		if (TimeRemaining <= 0.0f)
+		{
+			UE_LOG(LogAvaPlayable, Warning, TEXT("Playable Sync Event %s: Timed out"), *EventSignature);
+			FailedOperation(InResponse);
+		}
+	}
+
+#if WITH_EDITOR
+	virtual FString GetDescription() const override
+	{
+		return FString::Printf(TEXT("Playable Sync Event: %s"), *EventSignature);
+	}
+#endif
+	//~ End FPendingLatentAction
+
+private:
+	/** Latent Action Info - The function to execute. */
+	FName ExecutionFunction;
+
+	/** Latent Action Info - The resume point within the function to execute. */
+	int32 OutputLink;
+
+	/** Latent Action Info - Object to execute the function on. */
+	FWeakObjectPtr CallbackTarget;
+
+	/** Copy of the event signature for logs and description. */
+	FString EventSignature;
+
+	/** Keeps track of the remaining time, in seconds, before the operation times out. */
+	float TimeRemaining;
+
+	/** Output parameter - Indicate if the operation completed successfully. */
+	bool& OutSuccess;
+	
+	/** Simple shared bool as a fence to know when the event is completed. */
+	TSharedPtr<bool> SyncFence;
+};
+
+void UAvaPlayableLibrary::PlayableSyncEventLatent(const UObject* InWorldContextObject, FLatentActionInfo InLatentInfo, const FString& InEventSignature, bool& bInSuccess)
+{
+	bInSuccess = false;
+
+	UWorld* World = GEngine->GetWorldFromContextObject(InWorldContextObject, EGetWorldErrorMode::LogAndReturnNull);
+	UAvaPlayable* Playable = GetPlayable(InWorldContextObject);
+	
+	if (!World || !Playable)
+	{
+		return;
+	}
+	
+	FLatentActionManager& LatentManager = World->GetLatentActionManager();
+	if (LatentManager.FindExistingAction<FAvaPlayableSyncEventLatentAction>(InLatentInfo.CallbackTarget, InLatentInfo.UUID) == nullptr)
+	{
+		FAvaPlayableSyncEventLatentAction* NewAction = new FAvaPlayableSyncEventLatentAction(InLatentInfo, Playable, InEventSignature, bInSuccess);
+		LatentManager.AddNewAction(InLatentInfo.CallbackTarget, InLatentInfo.UUID, NewAction);
+	}
+}
\ No newline at end of file
diff --git a/Engine/Plugins/Experimental/Avalanche/Source/AvalancheMedia/Public/Playable/AvaPlayableLibrary.h b/Engine/Plugins/Experimental/Avalanche/Source/AvalancheMedia/Public/Playable/AvaPlayableLibrary.h
--- a/Engine/Plugins/Experimental/Avalanche/Source/AvalancheMedia/Public/Playable/AvaPlayableLibrary.h
+++ b/Engine/Plugins/Experimental/Avalanche/Source/AvalancheMedia/Public/Playable/AvaPlayableLibrary.h
@@ -44,4 +44,16 @@ public:
 	 */
 	UFUNCTION(BlueprintCallable, Category = "Motion Design|Playable", meta=(WorldContext = "InWorldContextObject"))
 	static AVALANCHEMEDIA_API bool SetPlayableHidden(const UObject* InWorldContextObject, bool bInShouldBeHidden);
+
+	/**
+	 * Push a playable sync event with the given signature on the cluster.
+	 * The action will complete when the event has occured on all nodes of the cluster.
+	 * The event signature includes the playable id as well. Therefore, the given signature need only be
+	 * unique in terms of any sequence of synchronized events for the given playable.
+	 * 
+	 * @param InEventSignature Event signature, must be unique in terms of sequence of events for the given playable. 
+	 * @param bInSuccess If true upon completion, the event has been synchronized successfully.
+	 */
+	UFUNCTION(BlueprintCallable, Category = "Motion Design|Playable", meta=(Latent, LatentInfo="InLatentInfo", WorldContext="InWorldContextObject"))
+	static AVALANCHEMEDIA_API void PlayableSyncEventLatent(const UObject* InWorldContextObject, FLatentActionInfo InLatentInfo, const FString& InEventSignature, bool& bInSuccess);
 };
