[Motion Design] Fix CIS Issue 875933 - shadowed parameter

diff --git a/Engine/Plugins/Experimental/Avalanche/Source/AvalancheMedia/Private/Playable/AvaPlayableLibrary.cpp b/Engine/Plugins/Experimental/Avalanche/Source/AvalancheMedia/Private/Playable/AvaPlayableLibrary.cpp
--- a/Engine/Plugins/Experimental/Avalanche/Source/AvalancheMedia/Private/Playable/AvaPlayableLibrary.cpp
+++ b/Engine/Plugins/Experimental/Avalanche/Source/AvalancheMedia/Private/Playable/AvaPlayableLibrary.cpp
@@ -149,9 +149,9 @@ public:
 
 				auto SyncEventHandler = [SyncFenceWeak]()
 				{
-					if (TSharedPtr<bool> SyncFence = SyncFenceWeak.Pin())
+					if (TSharedPtr<bool> SyncFencePtr = SyncFenceWeak.Pin())
 					{
-						*SyncFence = true;
+						*SyncFencePtr = true;
 					}
 				};
 
