[MotionDesign] Extend Sequencer Navigation Tool to support embedded Motion Design sequences

diff --git a/Engine/Plugins/Experimental/Avalanche/Source/Avalanche/Avalanche.Build.cs b/Engine/Plugins/Experimental/Avalanche/Source/Avalanche/Avalanche.Build.cs
--- a/Engine/Plugins/Experimental/Avalanche/Source/Avalanche/Avalanche.Build.cs
+++ b/Engine/Plugins/Experimental/Avalanche/Source/Avalanche/Avalanche.Build.cs
@@ -47,6 +47,7 @@ public class Avalanche : ModuleRules
 				"HeadMountedDisplay",
 				"InputCore",
 				"MediaCompositing",
+				"MovieScene",
 				"Niagara",
 				"NiagaraCore",
 				"ProceduralMeshComponent",
diff --git a/Engine/Plugins/Experimental/Avalanche/Source/Avalanche/Public/AvaScene.h b/Engine/Plugins/Experimental/Avalanche/Source/Avalanche/Public/AvaScene.h
--- a/Engine/Plugins/Experimental/Avalanche/Source/Avalanche/Public/AvaScene.h
+++ b/Engine/Plugins/Experimental/Avalanche/Source/Avalanche/Public/AvaScene.h
@@ -4,6 +4,7 @@
 
 #include "AvaSceneTree.h"
 #include "Delegates/Delegate.h"
+#include "EditorSequenceNavigationDefs.h"
 #include "GameFramework/Actor.h"
 #include "IAvaRemoteControlInterface.h"
 #include "IAvaSceneInterface.h"
@@ -61,7 +62,7 @@ public:
 
 	//~ Begin IAvaSceneInterface
 	virtual ULevel* GetSceneLevel() const override;
-	virtual UAvaSceneSettings* GetSceneSettings() const { return SceneSettings; }
+	virtual UAvaSceneSettings* GetSceneSettings() const override { return SceneSettings; }
 	virtual UAvaSceneState* GetSceneState() const override { return SceneState; }
 	virtual FAvaSceneTree& GetSceneTree() override { return SceneTree; }
 	virtual const FAvaSceneTree& GetSceneTree() const override { return SceneTree; }
@@ -69,6 +70,7 @@ public:
 	virtual IAvaSequenceProvider* GetSequenceProvider() override { return this; }
 	virtual const IAvaSequenceProvider* GetSequenceProvider() const override { return this; }
 	virtual URemoteControlPreset* GetRemoteControlPreset() const override { return RemoteControlPreset; }
+	virtual FNavigationToolSaveState& GetNavigationToolSaveState() override { return NavigationToolState; }
 	//~ End IAvaSceneInterface
 
 	//~ Begin IAvaSequenceProvider
@@ -169,6 +171,10 @@ protected:
 	/** Whether the motion design mode should auto start */
 	UPROPERTY()
 	bool bAutoStartMode = true;
+
+	/** Saved Navigation Tool state to restore */
+	UPROPERTY()
+	FNavigationToolSaveState NavigationToolState;
 #endif
 
 };
diff --git a/Engine/Plugins/Experimental/Avalanche/Source/Avalanche/Public/IAvaSceneInterface.h b/Engine/Plugins/Experimental/Avalanche/Source/Avalanche/Public/IAvaSceneInterface.h
--- a/Engine/Plugins/Experimental/Avalanche/Source/Avalanche/Public/IAvaSceneInterface.h
+++ b/Engine/Plugins/Experimental/Avalanche/Source/Avalanche/Public/IAvaSceneInterface.h
@@ -13,6 +13,10 @@ class ULevel;
 class URemoteControlPreset;
 struct FAvaSceneTree;
 
+#if WITH_EDITOR
+struct FNavigationToolSaveState;
+#endif
+
 /**
  * Interface for the driver class of an Motion Design Scene
  * This interface's purpose is to provide access to scene data and to core interfaces in the Scene
@@ -45,4 +49,8 @@ public:
 	virtual IAvaSequencePlaybackObject* GetPlaybackObject() const = 0;
 
 	virtual URemoteControlPreset* GetRemoteControlPreset() const = 0;
+
+#if WITH_EDITOR
+	virtual FNavigationToolSaveState& GetNavigationToolSaveState() = 0;
+#endif
 };
diff --git a/Engine/Plugins/Experimental/Avalanche/Source/AvalancheEditorCore/Private/AvaEditorCoreStyle.cpp b/Engine/Plugins/Experimental/Avalanche/Source/AvalancheEditorCore/Private/AvaEditorCoreStyle.cpp
new file mode 100644
--- /dev/null
+++ b/Engine/Plugins/Experimental/Avalanche/Source/AvalancheEditorCore/Private/AvaEditorCoreStyle.cpp
@@ -0,0 +1,27 @@
+﻿// Copyright Epic Games, Inc. All Rights Reserved.
+
+#include "AvaEditorCoreStyle.h"
+#include "Interfaces/IPluginManager.h"
+#include "Styling/SlateStyleMacros.h"
+#include "Styling/SlateStyleRegistry.h"
+
+FAvaEditorCoreStyle::FAvaEditorCoreStyle()
+	: FSlateStyleSet(TEXT("AvaEditorCore"))
+{
+	const FVector2f Icon16(16.f);
+
+	const TSharedPtr<IPlugin> Plugin = IPluginManager::Get().FindPlugin(UE_PLUGIN_NAME);
+	check(Plugin.IsValid());
+
+	FSlateStyleSet::SetContentRoot(Plugin->GetBaseDir() / TEXT("Resources"));
+	FSlateStyleSet::SetCoreContentRoot(FPaths::EngineContentDir() / TEXT("Editor/Slate"));
+
+	Set("Icons.MotionDesign", new IMAGE_BRUSH(TEXT("Icons/ToolboxIcons/toolbox"), Icon16));
+
+	FSlateStyleRegistry::RegisterSlateStyle(*this);
+}
+
+FAvaEditorCoreStyle::~FAvaEditorCoreStyle()
+{
+	FSlateStyleRegistry::UnRegisterSlateStyle(*this);
+}
diff --git a/Engine/Plugins/Experimental/Avalanche/Source/AvalancheEditorCore/Public/AvaEditorCoreStyle.h b/Engine/Plugins/Experimental/Avalanche/Source/AvalancheEditorCore/Public/AvaEditorCoreStyle.h
new file mode 100644
--- /dev/null
+++ b/Engine/Plugins/Experimental/Avalanche/Source/AvalancheEditorCore/Public/AvaEditorCoreStyle.h
@@ -0,0 +1,18 @@
+﻿// Copyright Epic Games, Inc. All Rights Reserved.
+
+#pragma once
+
+#include "Styling/SlateStyle.h"
+
+class AVALANCHEEDITORCORE_API FAvaEditorCoreStyle final : public FSlateStyleSet
+{
+public:
+	static FAvaEditorCoreStyle& Get()
+	{
+		static FAvaEditorCoreStyle Instance;
+		return Instance;
+	}
+
+	FAvaEditorCoreStyle();
+	virtual ~FAvaEditorCoreStyle() override;
+};
diff --git a/Engine/Plugins/Experimental/Avalanche/Source/AvalancheSequencer/AvalancheSequencer.Build.cs b/Engine/Plugins/Experimental/Avalanche/Source/AvalancheSequencer/AvalancheSequencer.Build.cs
--- a/Engine/Plugins/Experimental/Avalanche/Source/AvalancheSequencer/AvalancheSequencer.Build.cs
+++ b/Engine/Plugins/Experimental/Avalanche/Source/AvalancheSequencer/AvalancheSequencer.Build.cs
@@ -26,6 +26,7 @@ public class AvalancheSequencer : ModuleRules
 			new string[]
 			{
 				"ApplicationCore",
+				"Avalanche",
 				"AvalancheCore",
 				"AvalancheEditorCore",
 				"AvalancheOutliner",
diff --git a/Engine/Plugins/Experimental/Avalanche/Source/AvalancheSequencer/Private/AvaSequencer.cpp b/Engine/Plugins/Experimental/Avalanche/Source/AvalancheSequencer/Private/AvaSequencer.cpp
--- a/Engine/Plugins/Experimental/Avalanche/Source/AvalancheSequencer/Private/AvaSequencer.cpp
+++ b/Engine/Plugins/Experimental/Avalanche/Source/AvalancheSequencer/Private/AvaSequencer.cpp
@@ -1,6 +1,7 @@
 // Copyright Epic Games, Inc. All Rights Reserved.
 
 #include "AvaSequencer.h"
+#include "AvaEditorCoreStyle.h"
 #include "AvaSequence.h"
 #include "AvaSequenceActor.h"
 #include "AvaSequencePlaybackObject.h"
@@ -11,6 +12,7 @@
 #include "Clipboard/AvaSequenceImporter.h"
 #include "Commands/AvaSequencerAction.h"
 #include "Commands/AvaSequencerCommands.h"
+#include "ContentBrowserModule.h"
 #include "CoreGlobals.h"
 #include "DetailsView/SAvaMarkDetails.h"
 #include "DetailsView/Section/AvaSequencePlaybackDetails.h"
@@ -19,11 +21,14 @@
 #include "DetailsView/Section/AvaSequencerEaseCurveToolSection.h"
 #include "EaseCurveTool/AvaEaseCurveTool.h"
 #include "EaseCurveTool/AvaEaseCurveToolCommands.h"
+#include "Editor.h"
 #include "Editor/Sequencer/Private/Sequencer.h"
 #include "Framework/Commands/GenericCommands.h"
 #include "Framework/Commands/UIAction.h"
 #include "Framework/MultiBox/MultiBoxBuilder.h"
+#include "Framework/Notifications/NotificationManager.h"
 #include "IAvaSequenceProvider.h"
+#include "IContentBrowserSingleton.h"
 #include "ISequencer.h"
 #include "ISequencerModule.h"
 #include "ISequencerTrackEditor.h"
@@ -33,32 +38,30 @@
 #include "MVVM/Views/SOutlinerView.h"
 #include "Misc/MessageDialog.h"
 #include "Misc/NotifyHook.h"
-#include "Misc/TextFilter.h"
 #include "MovieScene.h"
+#include "NavigationTool/AvaNavigationToolProvider.h"
+#include "NavigationTool/NavigationToolAvaSequence.h"
 #include "Playback/AvaSequencerCleanView.h"
 #include "Playback/AvaSequencerController.h"
 #include "PropertyHandle.h"
 #include "ScopedTransaction.h"
 #include "Selection/AvaEditorSelection.h"
+#include "SequencerCommands.h"
+#include "SequencerSettings.h"
+#include "SequencerUtilities.h"
 #include "SequenceTree/AvaSequenceItem.h"
+#include "SequenceTree/IAvaSequenceItem.h"
 #include "SequenceTree/Columns/AvaSequenceNameColumn.h"
 #include "SequenceTree/Columns/AvaSequenceStatusColumn.h"
-#include "SequenceTree/Columns/IAvaSequenceColumn.h"
 #include "SequenceTree/Widgets/SAvaSequenceTree.h"
-#include "SequencerCommands.h"
-#include "SequencerSettings.h"
-#include "SequencerUtilities.h"
 #include "Settings/AvaSequencerSettings.h"
-#include "Sidebar/SSidebar.h"
-#include "Sidebar/SSidebarContainer.h"
 #include "Sidebar/SidebarDrawerConfig.h"
 #include "StaggerTool/AvaStaggerTool.h"
-#include "ToolMenu.h"
-#include "ToolMenuEntry.h"
-#include "ToolMenuSection.h"
-#include "ToolMenus.h"
 #include "Toolkits/AssetEditorToolkit.h"
-#include "Widgets/Views/STreeView.h"
+#include "Tools/NavigationTool/INavigationTool.h"
+#include "Tools/NavigationTool/NavigationToolDefines.h"
+#include "Widgets/Notifications/SNotificationList.h"
+#include "Widgets/Views/SHeaderRow.h"
 
 DEFINE_LOG_CATEGORY_STATIC(LogAvaSequencer, Log, All);
 
@@ -117,7 +120,7 @@ namespace UE::AvaSequencer::Private
 	};
 }
 
-const FName FAvaSequencer::SidebarDrawerId = TEXT("Sequence");
+const FName FAvaSequencer::SidebarDrawerId = TEXT("MotionDesign");
 
 FAvaSequencer::FAvaSequencer(IAvaSequencerProvider& InProvider, FAvaSequencerArgs&& InArgs)
 	: Provider(InProvider)
@@ -167,7 +170,10 @@ FAvaSequencer::FAvaSequencer(IAvaSequencerProvider& InProvider, FAvaSequencerArg
 	}
 
 	// Register to update when an undo/redo operation has been called to update our list of items
-	GEditor->RegisterForUndo(this);
+	if (GEditor)
+	{
+		GEditor->RegisterForUndo(this);
+	}
 }
 
 FAvaSequencer::~FAvaSequencer()
@@ -177,7 +183,7 @@ FAvaSequencer::~FAvaSequencer()
 	OnSequenceStartedHandle.Reset();
 	OnSequenceFinishedHandle.Reset();
 
-	if (GEngine)
+	if (GEditor)
 	{
 		GEditor->UnregisterForUndo(this);
 	}
@@ -211,6 +217,8 @@ FAvaSequencer::~FAvaSequencer()
 			SidebarExtender.Reset();
 		}
 	}
+
+	UnregisterNavigationTool();
 }
 
 void FAvaSequencer::BindCommands()
@@ -246,7 +254,7 @@ void FAvaSequencer::BindCommands()
 		, FExecuteAction::CreateRaw(this, &FAvaSequencer::AddTransformKey, EMovieSceneTransformChannel::Rotation));
 
 	CommandList->MapAction(SequencerCommands.AddScaleKey
-		, FExecuteAction::CreateRaw(this, &FAvaSequencer::AddTransformKey, EMovieSceneTransformChannel::Scale));	
+		, FExecuteAction::CreateRaw(this, &FAvaSequencer::AddTransformKey, EMovieSceneTransformChannel::Scale));
 }
 
 TSharedPtr<IAvaSequenceColumn> FAvaSequencer::FindSequenceColumn(FName InColumnName) const
@@ -366,6 +374,7 @@ TSharedRef<ISequencer> FAvaSequencer::CreateSequencer()
 		SequencerInitParams.HostCapabilities.bSupportsCurveEditor = true;
 		SequencerInitParams.HostCapabilities.bSupportsSaveMovieSceneAsset = false;
 		SequencerInitParams.HostCapabilities.bSupportsSidebar = true;
+		SequencerInitParams.HostCapabilities.bSupportsNavigationTool = true;
 	};
 
 	InstancedSequencer = FAvaSequencerUtils::GetSequencerModule().CreateSequencer(SequencerInitParams);
@@ -549,7 +558,7 @@ UAvaSequence* FAvaSequencer::CreateSequence() const
 	}
 
 	UObject* const Outer = SequenceProvider->ToUObject();
-	if (!ensure(Outer))
+	if (!IsValid(Outer))
 	{
 		return nullptr;
 	}
@@ -632,13 +641,13 @@ TSharedRef<SWidget> FAvaSequencer::CreatePlayerToolBar(const TSharedRef<FUIComma
 	const FAvaSequencerCommands& Commands = FAvaSequencerCommands::Get();
 
 	ToolBarBuilder.AddToolBarButton(Commands.PlaySelected, NAME_None, TAttribute<FText>(), TAttribute<FText>()
-		, FSlateIcon(FAppStyle::GetAppStyleSetName(), "Icons.Toolbar.Play"));
+		, FSlateIcon(FAppStyle::GetAppStyleSetName(), TEXT("Icons.Toolbar.Play")));
 
 	ToolBarBuilder.AddToolBarButton(Commands.ContinueSelected, NAME_None, TAttribute<FText>(), TAttribute<FText>()
-		, FSlateIcon(FAppStyle::GetAppStyleSetName(), "Icons.JumpToEvent"));
+		, FSlateIcon(FAppStyle::GetAppStyleSetName(), TEXT("Icons.JumpToEvent")));
 
 	ToolBarBuilder.AddToolBarButton(Commands.StopSelected, NAME_None, TAttribute<FText>(), TAttribute<FText>()
-		, FSlateIcon(FAppStyle::GetAppStyleSetName(), "Icons.Toolbar.Stop"));
+		, FSlateIcon(FAppStyle::GetAppStyleSetName(), TEXT("Icons.Toolbar.Stop")));
 
 	return SNew(SBox)
 		.HAlign(EHorizontalAlignment::HAlign_Center)
@@ -759,6 +768,11 @@ void FAvaSequencer::AddTransformKey(EMovieSceneTransformChannel InTransformChann
 {
 	using namespace UE::AvaSequencer::Private;
 
+	if (!GEditor)
+	{
+		return;
+	}
+
 	FEditorModeTools* const EditorModeTools = Provider.GetSequencerModeTools();
 	if (!EditorModeTools)
 	{
@@ -840,7 +854,21 @@ void FAvaSequencer::ApplyCustomPresetToSelection(FName InPresetName)
 
 void FAvaSequencer::ApplyPresetToSelection(const FAvaSequencePreset& InPreset)
 {
-	TArray<TSharedPtr<IAvaSequenceItem>> SelectedItems = SequenceTreeView->GetSelectedItems();
+	using namespace UE::Sequencer::NavigationTool;
+
+	const TSharedPtr<ISequencer> Sequencer = SequencerWeak.Pin();
+	if (!Sequencer.IsValid())
+	{
+		return;
+	}
+
+	const TSharedPtr<INavigationTool> Tool = Sequencer->GetNavigationTool();
+	if (!Tool.IsValid())
+	{
+		return;
+	}
+
+	const TArray<FNavigationToolItemPtr> SelectedItems = Tool->GetSelectedItems();
 	if (SelectedItems.IsEmpty())
 	{
 		return;
@@ -848,11 +876,11 @@ void FAvaSequencer::ApplyPresetToSelection(const FAvaSequencePreset& InPreset)
 
 	FScopedTransaction Transaction(LOCTEXT("ApplySequencePreset", "Apply Sequence Preset"));
 
-	for (const TSharedPtr<IAvaSequenceItem>& SequenceItem : SelectedItems)
+	for (const FNavigationToolItemPtr& Item : SelectedItems)
 	{
-		if (SequenceItem.IsValid())
+		if (FNavigationToolAvaSequence* const AvaSequenceItem = Item->CastTo<FNavigationToolAvaSequence>())
 		{
-			InPreset.ApplyPreset(SequenceItem->GetSequence());
+			InPreset.ApplyPreset(AvaSequenceItem->GetAvaSequence());
 		}
 	}
 }
@@ -876,12 +904,12 @@ void FAvaSequencer::AddSequence_Execute()
 	SequenceProvider->AddSequence(Sequence);
 }
 
-bool FAvaSequencer::DuplicateSequence_CanExecute() const
+bool FAvaSequencer::DuplicateSequences_CanExecute() const
 {
 	return GetProvider().CanEditOrPlaySequences();
 }
 
-void FAvaSequencer::DuplicateSequence_Execute()
+void FAvaSequencer::DuplicateSequences_Execute()
 {
 	IAvaSequenceProvider* const SequenceProvider = Provider.GetSequenceProvider();
 	if (!SequenceTreeView.IsValid() || !SequenceProvider)
@@ -895,9 +923,14 @@ void FAvaSequencer::DuplicateSequence_Execute()
 		return;
 	}
 
+	UObject* const Outer = SequenceProvider->ToUObject();
+	if (!IsValid(Outer))
+	{
+		return;
+	}
+
 	const FScopedTransaction Transaction(LOCTEXT("DuplicateSequenceTransaction", "Duplicate Sequence"));
 
-	UObject* const Outer = SequenceProvider->ToUObject();
 	Outer->Modify();
 
 	for (const FAvaSequenceItemPtr& Item : SelectedItems)
@@ -910,24 +943,24 @@ void FAvaSequencer::DuplicateSequence_Execute()
 	}
 }
 
-bool FAvaSequencer::ExportSequence_CanExecute() const
+bool FAvaSequencer::ExportSequences_CanExecute() const
 {
 	return SequenceTreeView.IsValid()
 		&& !SequenceTreeView->GetSelectedItems().IsEmpty();
 }
 
-void FAvaSequencer::ExportSequence_Execute()
+void FAvaSequencer::ExportSequences_Execute()
 {
 	Provider.ExportSequences(GetSelectedSequences());
 }
 
-bool FAvaSequencer::SpawnPlayer_CanExecute() const
+bool FAvaSequencer::SpawnPlayers_CanExecute() const
 {
 	return SequenceTreeView.IsValid()
 		&& !SequenceTreeView->GetSelectedItems().IsEmpty();
 }
 
-void FAvaSequencer::SpawnPlayer_Execute()
+void FAvaSequencer::SpawnPlayers_Execute()
 {
 	if (!GEditor)
 	{
@@ -946,7 +979,11 @@ void FAvaSequencer::SpawnPlayer_Execute()
 		return;
 	}
 
-	TArray<UAvaSequence*, TInlineAllocator<1>> Sequences = GetSelectedSequences();
+	TArray<UAvaSequence*> Sequences = GetSelectedSequences();
+	if (Sequences.Num() != 1)
+	{
+		return;
+	}
 
 	FScopedTransaction Transaction(LOCTEXT("SpawnSequencePlayers", "Spawn Sequence Players"));
 
@@ -957,12 +994,12 @@ void FAvaSequencer::SpawnPlayer_Execute()
 	}
 }
 
-bool FAvaSequencer::DeleteSequence_CanExecute() const
+bool FAvaSequencer::DeleteSequences_CanExecute() const
 {
 	return GetProvider().CanEditOrPlaySequences();
 }
 
-void FAvaSequencer::DeleteSequence_Execute()
+void FAvaSequencer::DeleteSequences_Execute()
 {
 	IAvaSequenceProvider* const SequenceProvider = Provider.GetSequenceProvider();
 	if (!SequenceProvider)
@@ -970,12 +1007,16 @@ void FAvaSequencer::DeleteSequence_Execute()
 		return;
 	}
 
+	UObject* const Outer = SequenceProvider->ToUObject();
+	if (!IsValid(Outer))
+	{
+		return;
+	}
+
 	TArray<FAvaSequenceItemPtr> SelectedItems = SequenceTreeView->GetSelectedItems();
 
 	const FScopedTransaction Transaction(LOCTEXT("DeleteSequenceTransaction", "Delete Sequence"));
 
-	UObject* const Outer = SequenceProvider->ToUObject();
-	check(Outer);
 	Outer->Modify();
 
 	TArray<UAvaSequence*> RemovedSequences;
@@ -1023,13 +1064,13 @@ void FAvaSequencer::RelabelSequence_Execute()
 	SelectedItem->RequestRelabel();
 }
 
-bool FAvaSequencer::PlaySelected_CanExecute() const
+bool FAvaSequencer::PlaySequences_CanExecute() const
 {
 	const bool bCanEditSequences = GetProvider().CanEditOrPlaySequences();
 	return bCanEditSequences && SequenceTreeView->GetNumItemsSelected() > 0;
 }
 
-void FAvaSequencer::PlaySelected_Execute()
+void FAvaSequencer::PlaySequences_Execute()
 {
 	const TArray<FAvaSequenceItemPtr> SelectedItems = SequenceTreeView->GetSelectedItems();
 
@@ -1050,13 +1091,13 @@ void FAvaSequencer::PlaySelected_Execute()
 	}
 }
 
-bool FAvaSequencer::ContinueSelected_CanExecute() const
+bool FAvaSequencer::ContinueSequences_CanExecute() const
 {
 	const bool bCanEditSequences = GetProvider().CanEditOrPlaySequences();
 	return bCanEditSequences && SequenceTreeView->GetNumItemsSelected() > 0;
 }
 
-void FAvaSequencer::ContinueSelected_Execute()
+void FAvaSequencer::ContinueSequences_Execute()
 {
 	const TArray<FAvaSequenceItemPtr> SelectedItems = SequenceTreeView->GetSelectedItems();
 	
@@ -1074,12 +1115,12 @@ void FAvaSequencer::ContinueSelected_Execute()
 	}
 }
 
-bool FAvaSequencer::StopSelected_CanExecute() const
+bool FAvaSequencer::StopSequences_CanExecute() const
 {
 	return SequenceTreeView->GetNumItemsSelected() > 0;
 }
 
-void FAvaSequencer::StopSelected_Execute()
+void FAvaSequencer::StopSequences_Execute()
 {
 	const TArray<FAvaSequenceItemPtr> SelectedItems = SequenceTreeView->GetSelectedItems();
 
@@ -1097,6 +1138,302 @@ void FAvaSequencer::StopSelected_Execute()
 	}
 }
 
+bool FAvaSequencer::CanAddSequence() const
+{
+	return GetProvider().CanEditOrPlaySequences();
+}
+
+UAvaSequence* FAvaSequencer::AddSequence(UAvaSequence* const InParentSequence)
+{
+	IAvaSequenceProvider* const SequenceProvider = Provider.GetSequenceProvider();
+	if (!SequenceProvider)
+	{
+		return nullptr;
+	}
+
+	const FScopedTransaction Transaction(LOCTEXT("AddSequenceTransaction", "Add Sequence"));
+
+	UAvaSequence* const NewSequence = CreateSequence();
+
+	if (IsValid(InParentSequence))
+	{
+		InParentSequence->Modify();
+		InParentSequence->AddChild(NewSequence);
+	}
+	else
+	{
+		SequenceProvider->AddSequence(NewSequence);
+	}
+
+	OnSequenceAddedDelegate.Broadcast(NewSequence);
+
+	return NewSequence;
+}
+
+bool FAvaSequencer::CanDuplicateSequences() const
+{
+	return GetProvider().CanEditOrPlaySequences();
+}
+
+void FAvaSequencer::DuplicateSequences()
+{
+	const TArray<UAvaSequence*> SelectedSequences = GetSelectedSequences();
+	if (SelectedSequences.IsEmpty())
+	{
+		return;
+	}
+
+	IAvaSequenceProvider* const SequenceProvider = Provider.GetSequenceProvider();
+	if (!SequenceProvider)
+	{
+		return;
+	}
+
+	UObject* const Outer = SequenceProvider->ToUObject();
+	if (!IsValid(Outer))
+	{
+		return;
+	}
+
+	const FScopedTransaction Transaction(LOCTEXT("DuplicateSequenceTransaction", "Duplicate Sequence"));
+
+	Outer->Modify();
+
+	for (UAvaSequence* const TemplateSequence : SelectedSequences)
+	{
+		check(IsValid(TemplateSequence));
+
+		UAvaSequence* const DupedSequence = DuplicateObject<UAvaSequence>(TemplateSequence, Outer);
+		SequenceProvider->AddSequence(DupedSequence);
+	}
+}
+
+bool FAvaSequencer::CanExportSequences() const
+{
+	return GetSelectedSequenceItems().Num() > 0;
+}
+
+void FAvaSequencer::ExportSequences()
+{
+	const TArray<UAvaSequence*> SelectedSequences = GetSelectedSequences();
+
+	TArray<TWeakObjectPtr<>> WeakSelectedObjects;
+	Algo::Transform(SelectedSequences, WeakSelectedObjects, [](UAvaSequence* const InSequence)
+		{
+			return InSequence;
+		});
+
+	Provider.ExportSequences(SelectedSequences);
+
+	FNotificationInfo Info(LOCTEXT("ExportSuccess", "Sequence Exported Successfully!"));
+	Info.bUseLargeFont = false;
+	Info.bUseSuccessFailIcons = false;
+	Info.ExpireDuration = 5.0f;
+	Info.HyperlinkText = LOCTEXT("ShowNewAssetsInContentBrowser", "Show in content browser");
+	Info.Hyperlink = FSimpleDelegate::CreateLambda([WeakSelectedObjects]()
+		{
+			TArray<UObject*> SyncObjects;
+
+			for (const TWeakObjectPtr<> WeakObject : WeakSelectedObjects)
+			{
+				UObject* const SyncObject = WeakObject.Get();
+				if (IsValid(SyncObject))
+				{
+					SyncObjects.Add(SyncObject);
+				}
+			}
+
+			FContentBrowserModule& ContentBrowserModule = FModuleManager::LoadModuleChecked<FContentBrowserModule>(TEXT("ContentBrowser"));
+			ContentBrowserModule.Get().SyncBrowserToAssets(SyncObjects);
+		});
+
+	FSlateNotificationManager::Get().AddNotification(Info);
+}
+
+bool FAvaSequencer::CanSpawnPlayers() const
+{
+	return GetSelectedSequenceItems().Num() > 0;
+}
+
+void FAvaSequencer::SpawnPlayers()
+{
+	if (!GEditor)
+	{
+		return;
+	}
+
+	UWorld* World = Provider.GetPlaybackContext()->GetWorld();
+	if (!World)
+	{
+		return;
+	}
+
+	UActorFactory* ActorFactory = GEditor->FindActorFactoryForActorClass(AAvaSequenceActor::StaticClass());
+	if (!ensure(ActorFactory))
+	{
+		return;
+	}
+
+	TArray<UAvaSequence*> Sequences = GetSelectedSequences();
+
+	FScopedTransaction Transaction(LOCTEXT("SpawnSequencePlayers", "Spawn Sequence Players"));
+
+	for (UAvaSequence* Sequence : Sequences)
+	{
+		check(Sequence);
+		GEditor->UseActorFactory(ActorFactory, FAssetData(Sequence), &FTransform::Identity);
+	}
+
+	const FText NotifyText = FText::Format(LOCTEXT("SpawnPlayerSuccess", "{0} Sequence Players Spawned Successfully!")
+		, FText::FromString(FString::FromInt(Sequences.Num())));
+
+	FNotificationInfo Info(NotifyText);
+	Info.bUseLargeFont = false;
+	Info.bUseSuccessFailIcons = false;
+	FSlateNotificationManager::Get().AddNotification(Info);
+}
+
+bool FAvaSequencer::CanDeleteSequence() const
+{
+	return GetProvider().CanEditOrPlaySequences();
+}
+
+void FAvaSequencer::DeleteSequence()
+{
+	DeleteSequences(TSet<UAvaSequence*> { GetSelectedSequences() });
+}
+
+void FAvaSequencer::DeleteSequences(const TSet<UAvaSequence*>& InSequences)
+{
+	if (InSequences.IsEmpty())
+	{
+		return;
+	}
+
+	IAvaSequenceProvider* const SequenceProvider = Provider.GetSequenceProvider();
+	if (!SequenceProvider)
+	{
+		return;
+	}
+
+	UObject* const Outer = SequenceProvider->ToUObject();
+	if (!IsValid(Outer))
+	{
+		return;
+	}
+
+	const FScopedTransaction Transaction(LOCTEXT("DeleteSequencesTransaction", "Delete Sequence(s)"));
+
+	Outer->Modify();
+
+	TArray<UAvaSequence*> RemovedSequences;
+	RemovedSequences.Reserve(InSequences.Num());
+
+	// Remove the sequences from the list (not marked as garbage yet)
+	for (UAvaSequence* const Sequence : InSequences)
+	{
+		check(IsValid(Sequence));
+
+		Sequence->Modify();
+		SequenceProvider->RemoveSequence(Sequence);
+
+		RemovedSequences.Add(Sequence);
+
+		OnSequenceRemovedDelegate.Broadcast(Sequence);
+	}
+
+	// Set the Viewed Sequence to the Default one
+	SetViewedSequence(GetDefaultSequence());
+
+	// Once a new viewed sequence is set, the removed sequences can now be marked as garbage 
+	for (UAvaSequence* Sequence : RemovedSequences)
+	{
+		Sequence->OnSequenceRemoved();
+	}
+
+	if (UBlueprint* const Blueprint = Cast<UBlueprint>(Outer))
+	{
+		FBlueprintEditorUtils::MarkBlueprintAsStructurallyModified(Blueprint);
+	}
+}
+
+bool FAvaSequencer::CanRelabelSequence() const
+{
+	const bool bCanEditSequences = GetProvider().CanEditOrPlaySequences();
+	return bCanEditSequences && SequenceTreeView->GetNumItemsSelected() == 1;
+}
+
+void FAvaSequencer::RelabelSequence()
+{
+	TArray<FAvaSequenceItemPtr> SelectedItems = SequenceTreeView->GetSelectedItems();
+	check(SelectedItems.Num() == 1);
+
+	const FAvaSequenceItemPtr& SelectedItem = SelectedItems[0];
+	SelectedItem->RequestRelabel();
+}
+
+bool FAvaSequencer::CanPlaySequences() const
+{
+	const bool bCanEditSequences = GetProvider().CanEditOrPlaySequences();
+	return bCanEditSequences && GetSelectedSequenceItems().Num() > 0;
+}
+
+void FAvaSequencer::PlaySequences()
+{
+	const TArray<UAvaSequence*> SelectedSequences = GetSelectedSequences();
+
+	FAvaSequencePlayParams PlaySettings;
+	PlaySettings.AdvancedSettings.bRestoreState = true;
+
+	for (UAvaSequence* const Sequence : SelectedSequences)
+	{
+		IAvaSequencePlaybackObject* const PlaybackObject = Provider.GetPlaybackObject();
+		if (IsValid(Sequence) && PlaybackObject)
+		{
+			PlaybackObject->PlaySequence(Sequence, PlaySettings);
+		}
+	}
+}
+
+bool FAvaSequencer::CanContinueSequences() const
+{
+	const bool bCanEditSequences = GetProvider().CanEditOrPlaySequences();
+	return bCanEditSequences && GetSelectedSequenceItems().Num() > 0;
+}
+
+void FAvaSequencer::ContinueSequences()
+{
+	const TArray<UAvaSequence*> SelectedSequences = GetSelectedSequences();
+
+	for (UAvaSequence* const Sequence : SelectedSequences)
+	{
+		IAvaSequencePlaybackObject* const PlaybackObject = Provider.GetPlaybackObject();
+		if (IsValid(Sequence) && PlaybackObject)
+		{
+			PlaybackObject->ContinueSequence(Sequence);
+		}
+	}
+}
+
+bool FAvaSequencer::CanStopSequences() const
+{
+	return GetSelectedSequenceItems().Num() > 0;
+}
+
+void FAvaSequencer::StopSequences()
+{
+	const TArray<UAvaSequence*> SelectedSequences = GetSelectedSequences();
+
+	for (UAvaSequence* const Sequence : SelectedSequences)
+	{
+		IAvaSequencePlaybackObject* const PlaybackObject = Provider.GetPlaybackObject();
+		if (IsValid(Sequence) && PlaybackObject)
+		{
+			PlaybackObject->StopSequence(Sequence);
+		}
+	}
+}
+
 UObject* FAvaSequencer::FindObjectToPossess(UObject* InParentObject, const FMovieScenePossessable& InPossessable)
 {
 	if (!IsValid(InParentObject))
@@ -1519,48 +1856,31 @@ TSharedRef<SWidget> FAvaSequencer::CreateSequenceWidget()
 		SequencerWeak.Reset();
 	}
 
-	TSharedRef<ISequencer> Sequencer = GetSequencer();
+	const TSharedRef<ISequencer> Sequencer = GetSequencer();
 
 	UAvaSequencerSettings* const SequencerSettings = GetMutableDefault<UAvaSequencerSettings>();
 	check(IsValid(SequencerSettings));
 	FSidebarState& SidebarState = SequencerSettings->GetSidebarState();
 
-	// Make sure the sequence tree widget is created
-	GetSequenceTreeWidget();
-
-	if (!SidebarState.IsVisible())
+	if (SidebarState.IsVisible())
 	{
-		return Sequencer->GetSequencerWidget();
-	}
-
-	TSharedRef<SSidebarContainer> SidebarContainer = SNew(SSidebarContainer);
-
-	TSharedRef<SSidebar> LeftSidebar = SNew(SSidebar, SidebarContainer)
-		.TabLocation(ESidebarTabLocation::Left)
-		.InitialDrawerSize(SidebarState.GetDrawerSize())
-		.OnStateChanged(this, &FAvaSequencer::OnSidebarStateChanged)
-		.OnGetContent(FOnGetContent::CreateLambda([this]()
-		{
-			return GetSequencer()->GetSequencerWidget();
-		}));
-
-	SidebarContainer->RebuildSidebar(LeftSidebar, SidebarState);
+		FSidebarDrawerConfig MotionDesignDrawerConfig;
+		MotionDesignDrawerConfig.UniqueId = SidebarDrawerId;
+		MotionDesignDrawerConfig.ButtonText = LOCTEXT("MotionDesignLabel", "Motion Design");
+		MotionDesignDrawerConfig.ToolTipText = LOCTEXT("SequenceTooltip", "Open the Sequence options panel");
+		MotionDesignDrawerConfig.Icon = FAvaEditorCoreStyle::Get().GetBrush(TEXT("Icons.MotionDesign"));
+		MotionDesignDrawerConfig.InitialState = SidebarState.FindOrAddDrawerState(SidebarDrawerId);
+		Sequencer->RegisterDrawer(MoveTemp(MotionDesignDrawerConfig));
 
-	FSidebarDrawerConfig SequenceTreeDrawerConfig;
-	SequenceTreeDrawerConfig.UniqueId = SidebarDrawerId;
-	SequenceTreeDrawerConfig.ButtonText = LOCTEXT("SequenceLabel", "Sequence");
-	SequenceTreeDrawerConfig.ToolTipText = LOCTEXT("SequenceTooltip", "Open the Sequence options panel");
-	SequenceTreeDrawerConfig.Icon = FAppStyle::GetBrush(TEXT("Persona.EditInSequencer"));
-	SequenceTreeDrawerConfig.InitialState = SidebarState.FindOrAddDrawerState(SidebarDrawerId);
-
-	LeftSidebar->RegisterDrawer(MoveTemp(SequenceTreeDrawerConfig));
+		const TSharedRef<FAvaSequencer> ThisSequencerRef = SharedThis(this);
+		Sequencer->RegisterDrawerSection(SidebarDrawerId, MakeShared<FAvaSequenceTreeDetails>(ThisSequencerRef));
+		Sequencer->RegisterDrawerSection(SidebarDrawerId, MakeShared<FAvaSequencePlaybackDetails>(ThisSequencerRef));
+		Sequencer->RegisterDrawerSection(SidebarDrawerId, MakeShared<FAvaSequenceSettingsDetails>(ThisSequencerRef));
+	}
 
-	const TSharedRef<FAvaSequencer> ThisSequencerRef = SharedThis(this);
-	LeftSidebar->RegisterDrawerSection(SidebarDrawerId, MakeShared<FAvaSequenceTreeDetails>(ThisSequencerRef));
-	LeftSidebar->RegisterDrawerSection(SidebarDrawerId, MakeShared<FAvaSequencePlaybackDetails>(ThisSequencerRef));
-	LeftSidebar->RegisterDrawerSection(SidebarDrawerId, MakeShared<FAvaSequenceSettingsDetails>(ThisSequencerRef));
+	RegisterNavigationTool();
 
-	return SidebarContainer;
+	return Sequencer->GetSequencerWidget();
 }
 
 void FAvaSequencer::OnActorsCopied(FString& InOutCopiedData, TConstArrayView<AActor*> InCopiedActors)
@@ -1631,6 +1951,13 @@ void FAvaSequencer::OnEditorSelectionChanged(const FAvaEditorSelection& InEditor
 
 void FAvaSequencer::NotifyOnSequenceTreeChanged()
 {
+	using namespace UE::Sequencer::NavigationTool;
+
+	if (const TSharedPtr<INavigationTool> NavigationTool = GetSequencer()->GetNavigationTool())
+	{
+		NavigationTool->RequestRefresh();
+	}
+
 	IAvaSequenceProvider* const SequenceProvider = Provider.GetSequenceProvider();
 	if (!SequenceTreeView.IsValid() || !SequenceProvider)
 	{
@@ -1701,42 +2028,76 @@ void FAvaSequencer::NotifyOnSequenceTreeChanged()
 	SequenceTreeView->RequestTreeRefresh();
 }
 
-void FAvaSequencer::PostUndo(bool bSuccess)
+void FAvaSequencer::PostUndo(const bool bInSuccess)
 {
 	// A just-added sequence might be removed due to this undo, so refresh 
 	NotifyOnSequenceTreeChanged();
 }
 
-TArray<UAvaSequence*, TInlineAllocator<1>> FAvaSequencer::GetSelectedSequences() const
+void FAvaSequencer::PostRedo(const bool bInSuccess)
 {
-	if (!SequenceTreeView.IsValid())
+	PostUndo(bInSuccess);
+}
+
+TArray<TSharedPtr<UE::Sequencer::NavigationTool::FNavigationToolAvaSequence>> FAvaSequencer::GetSelectedSequenceItems() const
+{
+	using namespace UE::Sequencer::NavigationTool;
+
+	const TSharedPtr<ISequencer> Sequencer = SequencerWeak.Pin();
+	if (!Sequencer.IsValid())
 	{
 		return {};
 	}
 
-	const TArray<FAvaSequenceItemPtr> SelectedItems = SequenceTreeView->GetSelectedItems();
+	const TSharedPtr<INavigationTool> Tool = Sequencer->GetNavigationTool();
+	if (!Tool.IsValid())
+	{
+		return {};
+	}
+
+	const TArray<FNavigationToolItemPtr> SelectedItems = Tool->GetSelectedItems();
 	if (SelectedItems.IsEmpty())
 	{
 		return {};
 	}
 
-	TArray<UAvaSequence*, TInlineAllocator<1>> SelectedSequences;
-	SelectedSequences.Reserve(SelectedItems.Num());
+	TArray<TSharedPtr<FNavigationToolAvaSequence>> SelectedSequenceItems;
+	SelectedSequenceItems.Reserve(SelectedItems.Num());
 
-	for (const FAvaSequenceItemPtr& Item : SelectedItems)
+	for (const FNavigationToolItemPtr& Item : SelectedItems)
 	{
-		if (!Item.IsValid())
+		if (FNavigationToolAvaSequence* const AvaSequenceItem = Item->CastTo<FNavigationToolAvaSequence>())
 		{
-			continue;
+			SelectedSequenceItems.Add(SharedThis(AvaSequenceItem));
 		}
+	}
 
-		if (UAvaSequence* Sequence = Item->GetSequence())
+	return SelectedSequenceItems;
+}
+
+TArray<UAvaSequence*> FAvaSequencer::GetSelectedSequences() const
+{
+	using namespace UE::Sequencer::NavigationTool;
+
+	const TArray<TSharedPtr<FNavigationToolAvaSequence>> SelectedItems = GetSelectedSequenceItems();
+	if (SelectedItems.IsEmpty())
+	{
+		return {};
+	}
+
+	TArray<UAvaSequence*> SelectedAvaSequences;
+	SelectedAvaSequences.Reserve(SelectedItems.Num());
+
+	for (const TSharedPtr<FNavigationToolAvaSequence>& SequenceItem : SelectedItems)
+	{
+		UAvaSequence* const AvaSequence = SequenceItem->GetAvaSequence();
+		if (IsValid(AvaSequence))
 		{
-			SelectedSequences.Add(Sequence);
+			SelectedAvaSequences.Add(AvaSequence);
 		}
 	}
 
-	return SelectedSequences;
+	return SelectedAvaSequences;
 }
 
 TSharedPtr<UE::Sequencer::SOutlinerView> FAvaSequencer::GetOutlinerView() const
@@ -1909,4 +2270,64 @@ void FAvaSequencer::OnSequencerClosed(const TSharedRef<ISequencer> InSequencer)
 	}
 }
 
+const TArray<TWeakObjectPtr<UAvaSequence>>& FAvaSequencer::GetRootSequences() const
+{
+	IAvaSequenceProvider* const SequenceProvider = Provider.GetSequenceProvider();
+	check(SequenceProvider);
+	return SequenceProvider->GetRootSequences();
+}
+
+void FAvaSequencer::RegisterNavigationTool()
+{
+	using namespace UE::Sequencer::NavigationTool;
+
+	const TSharedPtr<ISequencer> Sequencer = SequencerWeak.Pin();
+	if (!Sequencer.IsValid())
+	{
+		return;
+	}
+
+	const TSharedPtr<INavigationTool> Tool = Sequencer->GetNavigationTool();
+	if (!Tool.IsValid())
+	{
+		return;
+	}
+
+	const FName ToolId = FNavigationToolExtender::GetToolInstanceId(*Tool);
+
+	const TSharedPtr<FNavigationToolProvider> FoundProvider = FNavigationToolExtender::FindToolProvider(ToolId, UE::AvaSequencer::FAvaNavigationToolProvider::Identifier);
+	if (!FoundProvider.IsValid())
+	{
+		NavigationToolProvider = MakeShared<UE::AvaSequencer::FAvaNavigationToolProvider>(SharedThis(this));
+		FNavigationToolExtender::RegisterToolProvider(ToolId, NavigationToolProvider.ToSharedRef());
+	}
+}
+
+void FAvaSequencer::UnregisterNavigationTool()
+{
+	using namespace UE::Sequencer::NavigationTool;
+
+	if (!NavigationToolProvider.IsValid())
+	{
+		return;
+	}
+
+	const TSharedPtr<ISequencer> Sequencer = SequencerWeak.Pin();
+	if (!Sequencer.IsValid())
+	{
+		return;
+	}
+
+	const TSharedPtr<INavigationTool> Tool = Sequencer->GetNavigationTool();
+	if (!Tool.IsValid())
+	{
+		return;
+	}
+
+	const FName ToolId = FNavigationToolExtender::GetToolInstanceId(*Tool);
+	FNavigationToolExtender::UnregisterToolProvider(ToolId, NavigationToolProvider->GetIdentifier());
+
+	NavigationToolProvider.Reset();
+}
+
 #undef LOCTEXT_NAMESPACE
diff --git a/Engine/Plugins/Experimental/Avalanche/Source/AvalancheSequencer/Private/AvaSequencer.h b/Engine/Plugins/Experimental/Avalanche/Source/AvalancheSequencer/Private/AvaSequencer.h
--- a/Engine/Plugins/Experimental/Avalanche/Source/AvalancheSequencer/Private/AvaSequencer.h
+++ b/Engine/Plugins/Experimental/Avalanche/Source/AvalancheSequencer/Private/AvaSequencer.h
@@ -9,49 +9,50 @@
 #include "IAvaSequencerProvider.h"
 #include "MovieSceneSequenceID.h"
 #include "Widgets/Layout/SSplitter.h"
+#include "Widgets/Views/STreeView.h"
 
 class AActor;
 class FAvaEaseCurveTool;
 class FAvaEditorSelection;
+class FAvaNavigationToolProvider;
 class FAvaSequencerAction;
 class FAvaSequencerCleanView;
-class FAvaSequencerController;
-class FAvaStaggerTool;
 class FExtender;
 class FMenuBuilder;
-class FProperty;
 class FUICommandList;
 class IAvaSequenceColumn;
 class IAvaSequencerController;
-class IPropertyHandle;
 class ISequencer;
-class ISequencerObjectChangeListener;
 class SAvaSequenceTree;
-class SBox;
 class SHeaderRow;
-class SSidebar;
-class SSidebarContainer;
 class SWidget;
 class UAvaSequence;
-class UAvaSequencerSettings;
-class UMovieScene3DTransformTrack;
 class UMovieScene;
-class UMovieSceneTrack;
 class USequencerSettings;
-class UToolMenu;
 enum class EMovieSceneTransformChannel : uint32;
 struct FAvaSequencePreset;
 struct FAvaSequencerArgs;
 struct FMovieSceneBinding;
 struct FMovieScenePossessable;
-struct FSequencerInitParams;
 struct FSidebarState;
-struct FToolMenuContext;
-template<typename ItemType> class STreeView;
 
-namespace UE::Sequencer
+namespace UE
 {
-	class SOutlinerView;
+	namespace Sequencer
+	{
+		class SOutlinerView;
+
+		namespace NavigationTool
+		{
+			class FNavigationToolAvaSequence;
+			class FNavigationToolProvider;
+		}
+	}
+
+	namespace AvaSequencer
+	{
+		class FAvaNavigationToolProvider;
+	}
 }
 
 class FAvaSequencer : public IAvaSequencer, public FEditorUndoClient, public TSharedFromThis<FAvaSequencer>
@@ -121,39 +122,15 @@ public:
 	DECLARE_MULTICAST_DELEGATE_OneParam(FOnViewedSequenceChanged, UAvaSequence*);
 	FOnViewedSequenceChanged& GetOnViewedSequenceChanged() { return OnViewedSequenceChanged; }
 
+	FOnSequenceAdded& OnSequenceAdded() override { return OnSequenceAddedDelegate; }
+	FOnSequenceRemoved& OnSequenceRemoved() override { return OnSequenceRemovedDelegate; }
+
 	void AddTransformKey(EMovieSceneTransformChannel InTransformChannel);
 
 	void ApplyDefaultPresetToSelection(FName InPresetName);
 	void ApplyCustomPresetToSelection(FName InPresetName);
 	void ApplyPresetToSelection(const FAvaSequencePreset& InPreset);
 
-	bool AddSequence_CanExecute() const;
-	void AddSequence_Execute();
-
-	bool DuplicateSequence_CanExecute() const;
-	void DuplicateSequence_Execute();
-
-	bool ExportSequence_CanExecute() const;
-	void ExportSequence_Execute();
-
-	bool SpawnPlayer_CanExecute() const;
-	void SpawnPlayer_Execute();
-
-	bool DeleteSequence_CanExecute() const;
-	void DeleteSequence_Execute();
-
-	bool RelabelSequence_CanExecute() const;
-	void RelabelSequence_Execute();
-
-	bool PlaySelected_CanExecute() const;
-	void PlaySelected_Execute();
-
-	bool ContinueSelected_CanExecute() const;
-	void ContinueSelected_Execute();
-
-	bool StopSelected_CanExecute() const;
-	void StopSelected_Execute();
-
 	/**
  	 * Attempts to find the best matching subobject of the given Parent Object
  	 * based on the information provided by the given possessable
@@ -197,9 +174,10 @@ public:
 
 	TSharedRef<FAvaEaseCurveTool> GetEaseCurveTool() const;
 	
-	TArrayView<TWeakObjectPtr<>> ResolveBoundObjects(const FGuid& InBindingId, class UMovieSceneSequence* Sequence) const;
+	TArrayView<TWeakObjectPtr<>> ResolveBoundObjects(const FGuid& InBindingId, UMovieSceneSequence* Sequence) const;
 
 	//~ Begin IAvaSequencer
+
 	virtual const IAvaSequencerProvider& GetProvider() const override { return Provider; }
 	virtual TSharedRef<ISequencer> GetSequencer() const override;
 	virtual USequencerSettings* GetSequencerSettings() const override;
@@ -213,15 +191,74 @@ public:
 	virtual void OnActorsPasted(FStringView InPastedData, const TMap<FName, AActor*>& InPastedActors) override;
 	virtual void OnEditorSelectionChanged(const FAvaEditorSelection& InEditorSelection) override;
 	virtual void NotifyOnSequenceTreeChanged() override;
+	virtual const TArray<TWeakObjectPtr<UAvaSequence>>& GetRootSequences() const override;
+
 	//~ End IAvaSequencer
 
-	//~ Begin FEditorUndoClient Interface
-	virtual void PostUndo(bool bSuccess) override;
-	virtual void PostRedo(bool bSuccess) override { PostUndo(bSuccess); }
+	virtual bool CanAddSequence() const;
+	virtual UAvaSequence* AddSequence(UAvaSequence* const InParentSequence = nullptr);
+
+	virtual bool CanDuplicateSequences() const;
+	virtual void DuplicateSequences();
+
+	virtual bool CanExportSequences() const;
+	virtual void ExportSequences();
+
+	virtual bool CanSpawnPlayers() const;
+	virtual void SpawnPlayers();
+
+	virtual bool CanDeleteSequence() const;
+	virtual void DeleteSequence();
+	virtual void DeleteSequences(const TSet<UAvaSequence*>& InSequences);
+
+	virtual bool CanRelabelSequence() const;
+	virtual void RelabelSequence();
+
+	virtual bool CanPlaySequences() const;
+	virtual void PlaySequences();
+
+	virtual bool CanContinueSequences() const;
+	virtual void ContinueSequences();
+
+	virtual bool CanStopSequences() const;
+	virtual void StopSequences();
+
+	bool AddSequence_CanExecute() const;
+	void AddSequence_Execute();
+
+	bool DuplicateSequences_CanExecute() const;
+	void DuplicateSequences_Execute();
+
+	bool ExportSequences_CanExecute() const;
+	void ExportSequences_Execute();
+
+	bool SpawnPlayers_CanExecute() const;
+	void SpawnPlayers_Execute();
+
+	bool DeleteSequences_CanExecute() const;
+	void DeleteSequences_Execute();
+
+	bool RelabelSequence_CanExecute() const;
+	void RelabelSequence_Execute();
+
+	bool PlaySequences_CanExecute() const;
+	void PlaySequences_Execute();
+
+	bool ContinueSequences_CanExecute() const;
+	void ContinueSequences_Execute();
+
+	bool StopSequences_CanExecute() const;
+	void StopSequences_Execute();
+
+	//~ Begin FEditorUndoClient
+	virtual void PostUndo(bool bInSuccess) override;
+	virtual void PostRedo(bool bInSuccess) override;
 	// End of FEditorUndoClient
 
 private:
-	TArray<UAvaSequence*, TInlineAllocator<1>> GetSelectedSequences() const;
+	TArray<TSharedPtr<UE::Sequencer::NavigationTool::FNavigationToolAvaSequence>> GetSelectedSequenceItems() const;
+
+	TArray<UAvaSequence*> GetSelectedSequences() const;
 
 	TSharedPtr<UE::Sequencer::SOutlinerView> GetOutlinerView() const;
 
@@ -238,8 +275,14 @@ private:
 
 	void OnSequencerClosed(const TSharedRef<ISequencer> InSequencer);
 
+	void RegisterNavigationTool();
+	void UnregisterNavigationTool();
+
 	FOnViewedSequenceChanged OnViewedSequenceChanged;
 
+	FOnSequenceAdded OnSequenceAddedDelegate;
+	FOnSequenceRemoved OnSequenceRemovedDelegate;
+
 	IAvaSequencerProvider& Provider;
 
 	/** The active sequencer to use. This is set to Weak as FSequencer::Tick does not allow to have more than 1 ref (outside of menus) */
@@ -258,6 +301,7 @@ private:
 
 	FDelegateHandle SequencerAddTrackExtenderHandle;
 
+	// BEGIN REMOVE AFTER NAVIGATION TOOL INTEGRATION
 	TMap<FName, TSharedPtr<IAvaSequenceColumn>> SequenceColumns;
 
 	TArray<FAvaSequenceItemPtr> RootSequenceItems;
@@ -267,6 +311,7 @@ private:
 	TSharedPtr<SAvaSequenceTree> SequenceTree;
 
 	TSharedPtr<STreeView<FAvaSequenceItemPtr>> SequenceTreeView;
+	// END REMOVE AFTER NAVIGATION TOOL INTEGRATION
 
 	TSharedRef<FAvaSequencerCleanView> CleanView;
 
@@ -292,4 +337,6 @@ private:
 	TSharedPtr<FExtender> SidebarExtender;
 
 	FDelegateHandle SidebarSelectionExtenderHandle;
+
+	TSharedPtr<UE::AvaSequencer::FAvaNavigationToolProvider> NavigationToolProvider;
 };
diff --git a/Engine/Plugins/Experimental/Avalanche/Source/AvalancheSequencer/Private/AvaSequencerUtils.cpp b/Engine/Plugins/Experimental/Avalanche/Source/AvalancheSequencer/Private/AvaSequencerUtils.cpp
--- a/Engine/Plugins/Experimental/Avalanche/Source/AvalancheSequencer/Private/AvaSequencerUtils.cpp
+++ b/Engine/Plugins/Experimental/Avalanche/Source/AvalancheSequencer/Private/AvaSequencerUtils.cpp
@@ -1,10 +1,74 @@
 // Copyright Epic Games, Inc. All Rights Reserved.
 
 #include "AvaSequencerUtils.h"
+#include "AvaSceneSubsystem.h"
+#include "AvaSequencerSubsystem.h"
+#include "Engine/World.h"
+#include "IAvaSceneInterface.h"
+#include "IAvaSequencer.h"
 #include "Playback/AvaSequencerController.h"
 #include "Templates/SharedPointer.h"
+#include "Toolkits/IToolkitHost.h"
 
 TSharedRef<IAvaSequencerController> FAvaSequencerUtils::CreateSequencerController()
 {
 	return MakeShared<FAvaSequencerController>();
 }
+
+UWorld* FAvaSequencerUtils::GetSequencerWorld(const TSharedRef<ISequencer>& InSequencer)
+{
+	if (const TSharedPtr<IToolkitHost> ToolkitHost = InSequencer->GetToolkitHost())
+	{
+		return ToolkitHost->GetWorld();
+	}
+	return nullptr;
+}
+
+UAvaSequencerSubsystem* FAvaSequencerUtils::GetSequencerSubsystem(const TSharedRef<ISequencer>& InSequencer)
+{
+	const UWorld* const World = GetSequencerWorld(InSequencer);
+	if (IsValid(World))
+	{
+		return World->GetSubsystem<UAvaSequencerSubsystem>();
+	}
+	return nullptr;
+}
+
+UAvaSceneSubsystem* FAvaSequencerUtils::GetSceneSubsystem(const TSharedRef<ISequencer>& InSequencer)
+{
+	const UWorld* const World = GetSequencerWorld(InSequencer);
+	if (IsValid(World))
+	{
+		return World->GetSubsystem<UAvaSceneSubsystem>();
+	}
+	return nullptr;
+}
+
+IAvaSceneInterface* FAvaSequencerUtils::GetSceneInterface(const TSharedRef<ISequencer>& InSequencer)
+{
+	UAvaSceneSubsystem* const SceneSubsystem = GetSceneSubsystem(InSequencer);
+	if (IsValid(SceneSubsystem))
+	{
+		return SceneSubsystem->GetSceneInterface();
+	}
+	return nullptr;
+}
+
+IAvaSequenceProvider* FAvaSequencerUtils::GetSequenceProvider(const TSharedRef<ISequencer>& InSequencer)
+{
+	if (IAvaSceneInterface* const SceneInterface = GetSceneInterface(InSequencer))
+	{
+		return SceneInterface->GetSequenceProvider();
+	}
+	return nullptr;
+}
+
+TSharedPtr<IAvaSequencer> FAvaSequencerUtils::GetAvaSequencer(const TSharedRef<ISequencer>& InSequencer)
+{
+	UAvaSequencerSubsystem* const SequencerSubsystem = GetSequencerSubsystem(InSequencer);
+	if (IsValid(SequencerSubsystem))
+	{
+		return SequencerSubsystem->GetSequencer();
+	}
+	return nullptr;
+}
diff --git a/Engine/Plugins/Experimental/Avalanche/Source/AvalancheSequencer/Private/Commands/AvaSequencerCommands.cpp b/Engine/Plugins/Experimental/Avalanche/Source/AvalancheSequencer/Private/Commands/AvaSequencerCommands.cpp
--- a/Engine/Plugins/Experimental/Avalanche/Source/AvalancheSequencer/Private/Commands/AvaSequencerCommands.cpp
+++ b/Engine/Plugins/Experimental/Avalanche/Source/AvalancheSequencer/Private/Commands/AvaSequencerCommands.cpp
@@ -6,6 +6,12 @@
 
 void FAvaSequencerCommands::RegisterCommands()
 {
+	UI_COMMAND(AddNew
+		, "Add New"
+		, "Adds a new embedded Sequence to the level"
+		, EUserInterfaceActionType::Button
+		, FInputChord());
+
 	UI_COMMAND(PlaySelected
 		, "Play"
 		, "Plays all selected Sequences"
diff --git a/Engine/Plugins/Experimental/Avalanche/Source/AvalancheSequencer/Private/Commands/AvaSequencerCommands.h b/Engine/Plugins/Experimental/Avalanche/Source/AvalancheSequencer/Private/Commands/AvaSequencerCommands.h
--- a/Engine/Plugins/Experimental/Avalanche/Source/AvalancheSequencer/Private/Commands/AvaSequencerCommands.h
+++ b/Engine/Plugins/Experimental/Avalanche/Source/AvalancheSequencer/Private/Commands/AvaSequencerCommands.h
@@ -16,7 +16,9 @@ public:
 	{
 	}
 
-	virtual void RegisterCommands() override;	
+	virtual void RegisterCommands() override;
+
+	TSharedPtr<FUICommandInfo> AddNew;
 
 	TSharedPtr<FUICommandInfo> PlaySelected;
 
diff --git a/Engine/Plugins/Experimental/Avalanche/Source/AvalancheSequencer/Private/NavigationTool/AvaNavigationToolProvider.cpp b/Engine/Plugins/Experimental/Avalanche/Source/AvalancheSequencer/Private/NavigationTool/AvaNavigationToolProvider.cpp
new file mode 100644
--- /dev/null
+++ b/Engine/Plugins/Experimental/Avalanche/Source/AvalancheSequencer/Private/NavigationTool/AvaNavigationToolProvider.cpp
@@ -0,0 +1,605 @@
+﻿// Copyright Epic Games, Inc. All Rights Reserved.
+
+#include "AvaNavigationToolProvider.h"
+#include "AvaSequence.h"
+#include "AvaSequencer.h"
+#include "AvaSequencerUtils.h"
+#include "Engine/Level.h"
+#include "IAvaSceneInterface.h"
+#include "IAvaSequenceProvider.h"
+#include "ISequencer.h"
+#include "ISettingsModule.h"
+#include "ISourceControlModule.h"
+#include "Commands/AvaSequencerCommands.h"
+#include "DragDrop/NavigationToolAvaSequenceDropHandler.h"
+#include "Framework/Commands/GenericCommands.h"
+#include "NavigationTool/Columns/AvaNavigationToolStatusColumn.h"
+#include "NavigationTool/NavigationToolAvaSequence.h"
+#include "Settings/AvaSequencerSettings.h"
+#include "Tools/NavigationTool/Columns/NavigationToolColorColumn.h"
+#include "Tools/NavigationTool/Columns/NavigationToolColumnExtender.h"
+#include "Tools/NavigationTool/Columns/NavigationToolCommentColumn.h"
+#include "Tools/NavigationTool/Columns/NavigationToolHBiasColumn.h"
+#include "Tools/NavigationTool/Columns/NavigationToolInTimeColumn.h"
+#include "Tools/NavigationTool/Columns/NavigationToolItemsColumn.h"
+#include "Tools/NavigationTool/Columns/NavigationToolLabelColumn.h"
+#include "Tools/NavigationTool/Columns/NavigationToolLengthColumn.h"
+#include "Tools/NavigationTool/Columns/NavigationToolLockColumn.h"
+#include "Tools/NavigationTool/Columns/NavigationToolMarkerVisibilityColumn.h"
+#include "Tools/NavigationTool/Columns/NavigationToolOutTimeColumn.h"
+#include "Tools/NavigationTool/Columns/NavigationToolPlayheadColumn.h"
+#include "Tools/NavigationTool/Columns/NavigationToolRevisionControlColumn.h"
+#include "Tools/NavigationTool/Columns/NavigationToolStartFrameOffsetColumn.h"
+#include "Tools/NavigationTool/Columns/NavigationToolTakeColumn.h"
+#include "Tools/NavigationTool/Columns/NavigationToolDeactiveStateColumn.h"
+#include "Tools/NavigationTool/Filters/Filters/NavigationToolBuiltInFilterParams.h"
+#include "Tools/NavigationTool/INavigationTool.h"
+#include "Tools/NavigationTool/ItemActions/NavigationToolAddItem.h"
+#include "Tools/NavigationTool/Items/NavigationToolItemParameters.h"
+#include "Tools/NavigationTool/NavigationToolSettings.h"
+
+#define LOCTEXT_NAMESPACE "AvaNavigationToolProvider"
+
+namespace UE::AvaSequencer
+{
+
+const FName FAvaNavigationToolProvider::Identifier = TEXT("MotionDesign");
+
+const FText FAvaNavigationToolProvider::MotionDesignColumnViewName = LOCTEXT("MotionDesignColumnViewName", "Motion Design");
+
+const FName FAvaNavigationToolProvider::ToolbarSectionName = TEXT("MotionDesign");
+const FName FAvaNavigationToolProvider::ContextMenuSectionName = TEXT("MotionDesignActions");
+
+FAvaNavigationToolProvider::FAvaNavigationToolProvider(const TSharedRef<FAvaSequencer>& InAvaSequencer)
+	: WeakAvaSequencer(InAvaSequencer)
+	, ToolCommands(MakeShared<FUICommandList>())
+{
+}
+
+FName FAvaNavigationToolProvider::GetIdentifier() const
+{
+	return Identifier;
+}
+
+TSet<TSubclassOf<UMovieSceneSequence>> FAvaNavigationToolProvider::GetSupportedSequenceClasses() const
+{
+	return { UAvaSequence::StaticClass() };
+}
+
+FText FAvaNavigationToolProvider::GetDefaultColumnView() const
+{
+	return MotionDesignColumnViewName;
+}
+
+FNavigationToolSaveState* FAvaNavigationToolProvider::GetSaveState(const UE::Sequencer::NavigationTool::INavigationTool& InTool) const
+{
+	if (IAvaSceneInterface* const AvaScene = GetSceneInterface(InTool))
+	{
+		return &AvaScene->GetNavigationToolSaveState();
+	}
+	return nullptr;
+}
+
+void FAvaNavigationToolProvider::SetSaveState(const UE::Sequencer::NavigationTool::INavigationTool& InTool, const FNavigationToolSaveState& InSaveState) const
+{
+	if (IAvaSceneInterface* const AvaScene = GetSceneInterface(InTool))
+	{
+		AvaScene->GetNavigationToolSaveState() = InSaveState;
+	}
+}
+
+void FAvaNavigationToolProvider::BindCommands(const TSharedRef<FUICommandList>& InCommandList)
+{
+	const TSharedPtr<FAvaSequencer> AvaSequencer = WeakAvaSequencer.Pin();
+	if (!AvaSequencer.IsValid())
+	{
+		return;
+	}
+
+	InCommandList->Append(ToolCommands);
+
+	const TSharedRef<FAvaSequencer> AvaSequencerRef = AvaSequencer.ToSharedRef();
+	const FAvaSequencerCommands& AvaSequencerCommands = FAvaSequencerCommands::Get();
+	const FGenericCommands& GenericCommands = FGenericCommands::Get();
+
+	ToolCommands->MapAction(GenericCommands.Duplicate
+		, FExecuteAction::CreateSP(AvaSequencerRef, &FAvaSequencer::DuplicateSequences)
+		, FCanExecuteAction::CreateSP(AvaSequencerRef, &FAvaSequencer::CanDuplicateSequences));
+
+	ToolCommands->MapAction(GenericCommands.Delete
+		, FExecuteAction::CreateSP(AvaSequencerRef, &FAvaSequencer::DeleteSequence)
+		, FCanExecuteAction::CreateSP(AvaSequencerRef, &FAvaSequencer::CanDeleteSequence));
+
+	ToolCommands->MapAction(GenericCommands.Rename
+		, FExecuteAction::CreateSP(this, &FAvaNavigationToolProvider::RelabelSequence)
+		, FCanExecuteAction::CreateSP(this, &FAvaNavigationToolProvider::CanRelabelSequence));
+
+	ToolCommands->MapAction(AvaSequencerCommands.AddNew
+		, FExecuteAction::CreateSP(this, &FAvaNavigationToolProvider::AddSequenceToSelection)
+		, FCanExecuteAction::CreateSP(this, &FAvaNavigationToolProvider::CanAddSequenceToSelection));
+
+	ToolCommands->MapAction(AvaSequencerCommands.PlaySelected
+		, FExecuteAction::CreateSP(AvaSequencerRef, &FAvaSequencer::PlaySequences)
+		, FCanExecuteAction::CreateSP(AvaSequencerRef, &FAvaSequencer::CanPlaySequences));
+
+	ToolCommands->MapAction(AvaSequencerCommands.ContinueSelected
+		, FExecuteAction::CreateSP(AvaSequencerRef, &FAvaSequencer::ContinueSequences)
+		, FCanExecuteAction::CreateSP(AvaSequencerRef, &FAvaSequencer::CanContinueSequences));
+
+	ToolCommands->MapAction(AvaSequencerCommands.StopSelected
+		, FExecuteAction::CreateSP(AvaSequencerRef, &FAvaSequencer::StopSequences)
+		, FCanExecuteAction::CreateSP(AvaSequencerRef, &FAvaSequencer::CanStopSequences));
+
+	ToolCommands->MapAction(AvaSequencerCommands.ExportSequence
+		, FExecuteAction::CreateSP(AvaSequencerRef, &FAvaSequencer::ExportSequences)
+		, FCanExecuteAction::CreateSP(AvaSequencerRef, &FAvaSequencer::CanExportSequences));
+
+	ToolCommands->MapAction(AvaSequencerCommands.SpawnSequencePlayer
+		, FExecuteAction::CreateSP(AvaSequencerRef, &FAvaSequencer::SpawnPlayers)
+		, FCanExecuteAction::CreateSP(AvaSequencerRef, &FAvaSequencer::CanSpawnPlayers));
+}
+
+void FAvaNavigationToolProvider::OnActivate()
+{
+	ExtendToolToolBar();
+	ExtendToolItemContextMenu();
+
+	DragDropInitializedDelegate = UE::Sequencer::NavigationTool::FNavigationToolExtender::OnItemDragDropOpInitialized()
+		.AddSP(this, &FAvaNavigationToolProvider::ExtendToolItemDragDropOp);
+}
+
+void FAvaNavigationToolProvider::OnDeactivate()
+{
+	RemoveToolToolBarExtension();
+	RemoveToolItemContextMenuExtension();
+
+	UE::Sequencer::NavigationTool::FNavigationToolExtender::OnItemDragDropOpInitialized()
+		.Remove(DragDropInitializedDelegate);
+}
+
+void FAvaNavigationToolProvider::OnExtendColumns(UE::Sequencer::NavigationTool::FNavigationToolColumnExtender& OutExtender)
+{
+	using namespace UE::Sequencer::NavigationTool;
+
+	// Support built in columns
+	OutExtender.AddColumn<FNavigationToolPlayheadColumn>();
+	OutExtender.AddColumn<FNavigationToolDeactiveStateColumn>();
+	OutExtender.AddColumn<FNavigationToolMarkerVisibilityColumn>();
+	OutExtender.AddColumn<FNavigationToolLockColumn>();
+	OutExtender.AddColumn<FNavigationToolColorColumn>();
+	OutExtender.AddColumn<FNavigationToolLabelColumn>();
+	OutExtender.AddColumn<FNavigationToolItemsColumn>();
+	OutExtender.AddColumn<FNavigationToolInTimeColumn>();
+	OutExtender.AddColumn<FNavigationToolOutTimeColumn>();
+	OutExtender.AddColumn<FNavigationToolLengthColumn>();
+	OutExtender.AddColumn<FNavigationToolHBiasColumn>();
+	OutExtender.AddColumn<FNavigationToolStartFrameOffsetColumn>();
+	OutExtender.AddColumn<FNavigationToolTakeColumn>();
+	OutExtender.AddColumn<FNavigationToolCommentColumn>();
+
+	ISourceControlModule& SourceControlModule = ISourceControlModule::Get();
+	if (SourceControlModule.IsEnabled())
+	{
+		OutExtender.AddColumn<FNavigationToolRevisionControlColumn>();
+	}
+
+	// Add Motion Design specific columns
+	OutExtender.AddColumn<FAvaNavigationToolStatusColumn>();
+
+	FNavigationToolProvider::OnExtendColumns(OutExtender);
+}
+
+void FAvaNavigationToolProvider::OnExtendColumnViews(TSet<FNavigationToolColumnView>& OutColumnViews)
+{
+	using namespace UE::Sequencer::NavigationTool;
+
+	FNavigationToolColumnView& MotionDesignColumnView = OutColumnViews
+		.FindOrAdd(FNavigationToolColumnView(MotionDesignColumnViewName));
+
+	MotionDesignColumnView.VisibleColumns =
+		{
+			FNavigationToolColorColumn::StaticColumnId(),
+			FNavigationToolLabelColumn::StaticColumnId(),
+			FNavigationToolItemsColumn::StaticColumnId(),
+			FAvaNavigationToolStatusColumn::StaticColumnId()
+		};
+
+	FNavigationToolProvider::OnExtendColumnViews(OutColumnViews);
+}
+
+void FAvaNavigationToolProvider::OnExtendItemChildren(UE::Sequencer::NavigationTool::INavigationTool& InTool
+	, const UE::Sequencer::NavigationTool::FNavigationToolItemPtr& InParentItem
+	, TArray<UE::Sequencer::NavigationTool::FNavigationToolItemPtr>& OutChildren
+	, const bool bInRecursive)
+{
+	using namespace UE::Sequencer::NavigationTool;
+
+	FNavigationToolProvider::OnExtendItemChildren(InTool, InParentItem, OutChildren, bInRecursive);
+
+	// Only extending root item
+	if (InParentItem->GetItemId() != FNavigationToolItemId::RootId)
+	{
+		return;
+	}
+
+	IAvaSequenceProvider* const SequenceProvider = GetSequenceProvider(InTool);
+	if (!SequenceProvider)
+	{
+		return;
+	}
+
+	TArray<TWeakObjectPtr<UAvaSequence>> RootWeakSequences = SequenceProvider->GetRootSequences();
+
+	// Remove invalid sequences and sort
+	RootWeakSequences.RemoveAll([](const TWeakObjectPtr<UAvaSequence>& InRoot)
+		{
+			return !InRoot.IsValid();
+		});
+	RootWeakSequences.Sort([](const TWeakObjectPtr<UAvaSequence>& InA, const TWeakObjectPtr<UAvaSequence>& InB)
+		{
+			return InA->GetDisplayName().CompareTo(InB->GetDisplayName()) > 0;
+		});
+
+	const TSharedRef<FAvaNavigationToolProvider> Provider = SharedThis(this);
+
+	// Add child sequence items
+	for (const TWeakObjectPtr<UAvaSequence>& RootSequenceWeak : RootWeakSequences)
+	{
+		UAvaSequence* const RootSequence = RootSequenceWeak.Get();
+		if (IsValid(RootSequence))
+		{
+			const FNavigationToolItemPtr NewItem = InTool.FindOrAdd<FNavigationToolAvaSequence>(Provider
+				, InParentItem, RootSequence);
+			const FNavigationToolItemFlagGuard Guard(NewItem, ENavigationToolItemFlags::IgnorePendingKill);
+			OutChildren.Add(NewItem);
+			if (bInRecursive)
+			{
+				NewItem->FindChildren(OutChildren, bInRecursive);
+			}
+		}
+	}
+}
+
+void FAvaNavigationToolProvider::OnExtendBuiltInFilters(TArray<FNavigationToolBuiltInFilterParams>& OutFilterParams)
+{
+	OutFilterParams.Add(FNavigationToolBuiltInFilterParams::CreateSequenceFilter());
+	OutFilterParams.Add(FNavigationToolBuiltInFilterParams::CreateTrackFilter());
+	OutFilterParams.Add(FNavigationToolBuiltInFilterParams::CreateBindingFilter());
+	OutFilterParams.Add(FNavigationToolBuiltInFilterParams::CreateMarkerFilter());
+
+	FNavigationToolProvider::OnExtendBuiltInFilters(OutFilterParams);
+}
+
+IAvaSceneInterface* FAvaNavigationToolProvider::GetSceneInterface(const UE::Sequencer::NavigationTool::INavigationTool& InTool) const
+{
+	if (const TSharedPtr<ISequencer> Sequencer = InTool.GetSequencer())
+	{
+		return FAvaSequencerUtils::GetSceneInterface(Sequencer.ToSharedRef());
+	}
+	return nullptr;
+}
+
+IAvaSequenceProvider* FAvaNavigationToolProvider::GetSequenceProvider(const UE::Sequencer::NavigationTool::INavigationTool& InTool) const
+{
+	if (const TSharedPtr<ISequencer> Sequencer = InTool.GetSequencer())
+	{
+		return FAvaSequencerUtils::GetSequenceProvider(Sequencer.ToSharedRef());
+	}
+	return nullptr;
+}
+
+void FAvaNavigationToolProvider::ExtendToolToolBar()
+{
+	UToolMenu* const ToolMenu = UToolMenus::Get()->ExtendMenu(UE::Sequencer::NavigationTool::GetToolBarMenuName());
+	if (!IsValid(ToolMenu))
+	{
+		return;
+	}
+
+	FToolMenuSection& MotionDesignSection = ToolMenu->FindOrAddSection(ToolbarSectionName);
+
+	const FAvaSequencerCommands& AvaSequencerCommands = FAvaSequencerCommands::Get();
+	const FName SequencerToolbarStyleName = TEXT("SequencerToolbar");
+
+	FSlimHorizontalToolBarBuilder ToolBarBuilder(ToolCommands, FMultiBoxCustomization::None, nullptr, true);
+	ToolBarBuilder.SetLabelVisibility(EVisibility::Collapsed);
+
+	FToolMenuEntry SequenceAddNewEntry = FToolMenuEntry::InitToolBarButton(AvaSequencerCommands.AddNew
+		, TAttribute<FText>(), TAttribute<FText>()
+		, FSlateIcon(FAppStyle::GetAppStyleSetName(), TEXT("AnimationEditor.SetKey")));
+	SequenceAddNewEntry.StyleNameOverride = SequencerToolbarStyleName;
+	SequenceAddNewEntry.SetCommandList(ToolCommands);
+	MotionDesignSection.AddEntry(SequenceAddNewEntry);
+
+	FToolMenuEntry SequencePlaySelectedEntry = FToolMenuEntry::InitToolBarButton(AvaSequencerCommands.PlaySelected
+		, TAttribute<FText>(), TAttribute<FText>()
+		, FSlateIcon(FAppStyle::GetAppStyleSetName(), TEXT("Icons.Toolbar.Play")));
+	SequencePlaySelectedEntry.StyleNameOverride = SequencerToolbarStyleName;
+	SequencePlaySelectedEntry.SetCommandList(ToolCommands);
+	MotionDesignSection.AddEntry(SequencePlaySelectedEntry);
+
+	FToolMenuEntry SequenceContinueSelectedEntry = FToolMenuEntry::InitToolBarButton(AvaSequencerCommands.ContinueSelected
+		, TAttribute<FText>(), TAttribute<FText>()
+		, FSlateIcon(FAppStyle::GetAppStyleSetName(), TEXT("Icons.JumpToEvent")));
+	SequenceContinueSelectedEntry.StyleNameOverride = SequencerToolbarStyleName;
+	SequenceContinueSelectedEntry.SetCommandList(ToolCommands);
+	MotionDesignSection.AddEntry(SequenceContinueSelectedEntry);
+
+	FToolMenuEntry SequenceStopSelectedEntry = FToolMenuEntry::InitToolBarButton(AvaSequencerCommands.StopSelected
+		, TAttribute<FText>(), TAttribute<FText>()
+		, FSlateIcon(FAppStyle::GetAppStyleSetName(), TEXT("Icons.Toolbar.Stop")));
+	SequenceStopSelectedEntry.StyleNameOverride = SequencerToolbarStyleName;
+	SequenceStopSelectedEntry.SetCommandList(ToolCommands);
+	MotionDesignSection.AddEntry(SequenceStopSelectedEntry);
+}
+
+void FAvaNavigationToolProvider::ExtendToolItemContextMenu()
+{
+	UToolMenus* ToolMenus = UToolMenus::Get();
+	if (!ToolMenus)
+	{
+		return;
+	}
+
+	UToolMenu* const ExtendedMenu = ToolMenus->ExtendMenu(UE::Sequencer::NavigationTool::GetItemContextMenuName());
+	if (!IsValid(ExtendedMenu))
+	{
+		return;
+	}
+
+	const FAvaSequencerCommands& SequencerCommands = FAvaSequencerCommands::Get();
+
+	FToolMenuSection& MotionDesignSection = ExtendedMenu->FindOrAddSection(ContextMenuSectionName
+		, LOCTEXT("MotionDesign", "Motion Design Actions")
+		, FToolMenuInsert(TEXT("ToolActions"), EToolMenuInsertType::Before));
+
+	MotionDesignSection.AddSubMenu(TEXT("ApplyPreset")
+		, LOCTEXT("ApplyPresetLabel", "Apply Preset")
+		, LOCTEXT("ApplyPresetTooltip", "Apply a Preset to the Selected Sequences")
+		, FNewToolMenuDelegate::CreateSP(this, &FAvaNavigationToolProvider::GeneratePresetMenu));
+	MotionDesignSection.AddMenuEntry(SequencerCommands.SpawnSequencePlayer);
+	MotionDesignSection.AddMenuEntry(SequencerCommands.ExportSequence);
+
+	const FGenericCommands& GenericCommands = FGenericCommands::Get();
+	FToolMenuSection& GenericSection = ExtendedMenu->FindOrAddSection(TEXT("GenericActions"), LOCTEXT("GenericActionsHeader", "Generic Actions"));
+
+	GenericSection.AddMenuEntry(GenericCommands.Duplicate);
+	GenericSection.AddMenuEntry(GenericCommands.Delete);
+	GenericSection.AddMenuEntry(GenericCommands.Rename);
+}
+
+void FAvaNavigationToolProvider::RemoveToolToolBarExtension()
+{
+	UToolMenu* const Menu = UToolMenus::Get()->FindMenu(UE::Sequencer::NavigationTool::GetToolBarMenuName());
+	if (IsValid(Menu))
+	{
+		Menu->RemoveSection(ToolbarSectionName);
+	}
+}
+
+void FAvaNavigationToolProvider::RemoveToolItemContextMenuExtension()
+{
+	UToolMenu* const Menu = UToolMenus::Get()->FindMenu(UE::Sequencer::NavigationTool::GetItemContextMenuName());
+	if (IsValid(Menu))
+	{
+		Menu->RemoveSection(ContextMenuSectionName);
+	}
+}
+
+void FAvaNavigationToolProvider::ExtendToolItemDragDropOp(UE::Sequencer::NavigationTool::FNavigationToolItemDragDropOp& InDragDropOp)
+{
+	InDragDropOp.AddDropHandler<UE::Sequencer::NavigationTool::FNavigationToolAvaSequenceDropHandler>(WeakAvaSequencer);
+}
+
+void FAvaNavigationToolProvider::OnSequenceAdded(UAvaSequence* const InAvaSequence)
+{
+	const TSharedPtr<FAvaSequencer> AvaSequencer = StaticCastSharedPtr<FAvaSequencer>(WeakAvaSequencer.Pin());
+	if (!AvaSequencer.IsValid())
+	{
+		return;
+	}
+
+	const TSharedPtr<ISequencer> Sequencer = AvaSequencer->GetSequencer();
+	if (!Sequencer.IsValid())
+	{
+		return;
+	}
+
+	const TSharedPtr<UE::Sequencer::NavigationTool::INavigationTool> NavigationTool = Sequencer->GetNavigationTool();
+	if (!NavigationTool.IsValid())
+	{
+		return;
+	}
+
+	NavigationTool->RequestRefresh();
+}
+
+void FAvaNavigationToolProvider::GeneratePresetMenu(UToolMenu* const InToolMenu)
+{
+	const UAvaSequencerSettings* const SequencerSettings = GetDefault<UAvaSequencerSettings>();
+	if (!IsValid(SequencerSettings))
+	{
+		return;
+	}
+
+	const TSharedPtr<FAvaSequencer> AvaSequencer = WeakAvaSequencer.Pin();
+	if (!AvaSequencer.IsValid())
+	{
+		return;
+	}
+
+	const TSharedRef<FAvaSequencer> AvaSequencerRef = AvaSequencer.ToSharedRef();
+
+	// Default Presets
+	const TConstArrayView<FAvaSequencePreset> DefaultPresets = SequencerSettings->GetDefaultSequencePresets();
+	if (!DefaultPresets.IsEmpty())
+	{
+		FToolMenuSection& DefaultPresetSection = InToolMenu->FindOrAddSection(TEXT("DefaultPresets"), LOCTEXT("DefaultPresetsLabel", "Default Presets"));
+		for (const FAvaSequencePreset& Preset : DefaultPresets)
+		{
+			DefaultPresetSection.AddMenuEntry(Preset.PresetName
+				, FText::FromName(Preset.PresetName)
+				, FText::FromName(Preset.PresetName)
+				, FSlateIcon()
+				, FToolUIActionChoice(FExecuteAction::CreateSP(AvaSequencerRef, &FAvaSequencer::ApplyDefaultPresetToSelection, Preset.PresetName)));
+		}
+	}
+
+	// Custom Presets
+	const TSet<FAvaSequencePreset>& CustomPresets = SequencerSettings->GetCustomSequencePresets();
+	if (!CustomPresets.IsEmpty())
+	{
+		FToolMenuSection& CustomPresetSection = InToolMenu->FindOrAddSection(TEXT("CustomPresets"), LOCTEXT("CustomPresetsLabel", "Custom Presets"));
+		for (const FAvaSequencePreset& Preset : CustomPresets)
+		{
+			CustomPresetSection.AddMenuEntry(Preset.PresetName
+				, FText::FromName(Preset.PresetName)
+				, FText::FromName(Preset.PresetName)
+				, FSlateIcon()
+				, FToolUIActionChoice(FExecuteAction::CreateSP(AvaSequencerRef, &FAvaSequencer::ApplyCustomPresetToSelection, Preset.PresetName)));
+		}
+	}
+
+	// Settings
+	FToolMenuSection& SettingsSection = InToolMenu->FindOrAddSection(TEXT("Settings"));
+	SettingsSection.AddSeparator(NAME_None);
+	SettingsSection.AddMenuEntry(TEXT("OpenSettings")
+		, LOCTEXT("OpenSettingsLabel", "Open Settings")
+		, LOCTEXT("OpenSettingsTooltip", "Opens the Settings to customize the sequence presets")
+		, FSlateIcon(FAppStyle::GetAppStyleSetName(), "Icons.Settings")
+		, FExecuteAction::CreateLambda([]
+		{
+			ISettingsModule& SettingsModule = FModuleManager::LoadModuleChecked<ISettingsModule>(TEXT("Settings"));
+			const UAvaSequencerSettings* const AvaSequencerSettings = GetDefault<UAvaSequencerSettings>();
+			if (IsValid(AvaSequencerSettings))
+			{
+				SettingsModule.ShowViewer(AvaSequencerSettings->GetContainerName(), AvaSequencerSettings->GetCategoryName(), AvaSequencerSettings->GetSectionName());
+			}
+		}));
+}
+
+bool FAvaNavigationToolProvider::CanRelabelSequence() const
+{
+	using namespace UE::Sequencer::NavigationTool;
+
+	const TSharedPtr<IAvaSequencer> AvaSequencer = WeakAvaSequencer.Pin();
+	if (!AvaSequencer.IsValid())
+	{
+		return false;
+	}
+
+	const TSharedPtr<INavigationTool> NavigationTool = AvaSequencer->GetSequencer()->GetNavigationTool();
+	if (!NavigationTool.IsValid())
+	{
+		return false;
+	}
+
+	const TArray<FNavigationToolItemPtr> SelectedItems = NavigationTool->GetSelectedItems();
+
+	return SelectedItems.Num() == 1 && AvaSequencer->GetProvider().CanEditOrPlaySequences();
+}
+
+void FAvaNavigationToolProvider::RelabelSequence()
+{
+	using namespace UE::Sequencer::NavigationTool;
+
+	const TSharedPtr<IAvaSequencer> AvaSequencer = WeakAvaSequencer.Pin();
+	if (!AvaSequencer.IsValid())
+	{
+		return;
+	}
+
+	const TSharedPtr<INavigationTool> NavigationTool = AvaSequencer->GetSequencer()->GetNavigationTool();
+	if (!NavigationTool.IsValid())
+	{
+		return;
+	}
+
+	const TArray<FNavigationToolItemPtr> SelectedItems = NavigationTool->GetSelectedItems();
+	if (SelectedItems.Num() != 1)
+	{
+		return;
+	}
+
+	const FNavigationToolAvaSequence* const AvaSequenceItem = SelectedItems[0]->CastTo<FNavigationToolAvaSequence>();
+	if (!AvaSequenceItem)
+	{
+		return;
+	}
+
+	// @Todo: Implement BeginRename function in base classes
+	//AvaSequenceItem->BeginRename();
+}
+
+bool FAvaNavigationToolProvider::CanAddSequenceToSelection() const
+{
+	using namespace UE::Sequencer::NavigationTool;
+
+	const TSharedPtr<IAvaSequencer> AvaSequencer = WeakAvaSequencer.Pin();
+	if (!AvaSequencer.IsValid())
+	{
+		return false;
+	}
+
+	const TSharedPtr<INavigationTool> NavigationTool = AvaSequencer->GetSequencer()->GetNavigationTool();
+	if (!NavigationTool.IsValid())
+	{
+		return false;
+	}
+
+	const TArray<FNavigationToolItemPtr> SelectedItems = NavigationTool->GetSelectedItems();
+
+	return SelectedItems.Num() <= 1 && AvaSequencer->GetProvider().CanEditOrPlaySequences();
+}
+
+void FAvaNavigationToolProvider::AddSequenceToSelection()
+{
+	using namespace UE::Sequencer::NavigationTool;
+
+	const TSharedPtr<FAvaSequencer> AvaSequencer = WeakAvaSequencer.Pin();
+	if (!AvaSequencer.IsValid())
+	{
+		return;
+	}
+
+	const TSharedPtr<INavigationTool> NavigationTool = AvaSequencer->GetSequencer()->GetNavigationTool();
+	if (!NavigationTool.IsValid())
+	{
+		return;
+	}
+
+	TSharedPtr<FNavigationToolAvaSequence> ParentSequenceItem;
+	UAvaSequence* ParentSequence = nullptr;
+
+	const TArray<FNavigationToolItemPtr> SelectedItems = NavigationTool->GetSelectedItems();
+	if (SelectedItems.Num() == 1)
+	{
+		if (FNavigationToolAvaSequence* const AvaSequenceItem = SelectedItems[0]->CastTo<FNavigationToolAvaSequence>())
+		{
+			ParentSequenceItem = SharedThis(AvaSequenceItem);
+			ParentSequence = AvaSequenceItem->GetAvaSequence();
+		}
+	}
+
+	UAvaSequence* const NewSequence = AvaSequencer->AddSequence(ParentSequence);
+	if (!IsValid(NewSequence))
+	{
+		return;
+	}
+
+	FNavigationToolAddItemParams AddChildItemParams;
+	AddChildItemParams.Item = NavigationTool->FindOrAdd<FNavigationToolAvaSequence>(SharedThis(this)
+		, AddChildItemParams.RelativeItem, NewSequence);
+	AddChildItemParams.RelativeItem = ParentSequenceItem.IsValid()
+		? ParentSequenceItem.ToSharedRef() : NavigationTool->GetTreeRoot();
+	AddChildItemParams.RelativeDropZone = EItemDropZone::OntoItem;
+	AddChildItemParams.Flags = ENavigationToolAddItemFlags::Select | ENavigationToolAddItemFlags::Transact;
+	AddChildItemParams.SelectionFlags = ENavigationToolItemSelectionFlags::SignalSelectionChange
+		| ENavigationToolItemSelectionFlags::ScrollIntoView;
+
+	NavigationTool->EnqueueItemAction<FNavigationToolAddItem>(AddChildItemParams);
+}
+
+} // UE::AvaSequencer
+
+#undef LOCTEXT_NAMESPACE
diff --git a/Engine/Plugins/Experimental/Avalanche/Source/AvalancheSequencer/Private/NavigationTool/AvaNavigationToolProvider.h b/Engine/Plugins/Experimental/Avalanche/Source/AvalancheSequencer/Private/NavigationTool/AvaNavigationToolProvider.h
new file mode 100644
--- /dev/null
+++ b/Engine/Plugins/Experimental/Avalanche/Source/AvalancheSequencer/Private/NavigationTool/AvaNavigationToolProvider.h
@@ -0,0 +1,97 @@
+﻿// Copyright Epic Games, Inc. All Rights Reserved.
+
+#pragma once
+
+#include "Delegates/IDelegateInstance.h"
+#include "Templates/SharedPointer.h"
+#include "Tools/NavigationTool/NavigationToolDefines.h"
+#include "Tools/NavigationTool/Providers/NavigationToolProvider.h"
+#include "UObject/WeakObjectPtr.h"
+
+class FAvaSequencer;
+class FUICommandList;
+class IAvaSceneInterface;
+class IAvaSequenceProvider;
+class UAvaSequence;
+class UObject;
+struct FNavigationToolSaveState;
+
+namespace UE::Sequencer::NavigationTool
+{
+	class FNavigationToolItem;
+	class FNavigationToolItemDragDropOp;
+	class INavigationTool;
+}
+
+namespace UE::AvaSequencer
+{
+
+class FAvaNavigationToolProvider : public UE::Sequencer::NavigationTool::FNavigationToolProvider
+{
+public:
+	static const FName Identifier;
+
+	static const FText MotionDesignColumnViewName;
+
+	FAvaNavigationToolProvider(const TSharedRef<FAvaSequencer>& InAvaSequencer);
+	virtual ~FAvaNavigationToolProvider() override = default;
+
+	//~ Begin INavigationToolProvider
+
+	virtual FName GetIdentifier() const override;
+
+	virtual TSet<TSubclassOf<UMovieSceneSequence>> GetSupportedSequenceClasses() const override;
+
+	virtual FText GetDefaultColumnView() const override;
+
+	virtual FNavigationToolSaveState* GetSaveState(const UE::Sequencer::NavigationTool::INavigationTool& InTool) const override;
+	virtual void SetSaveState(const UE::Sequencer::NavigationTool::INavigationTool& InTool, const FNavigationToolSaveState& InSaveState) const override;
+
+	virtual void BindCommands(const TSharedRef<FUICommandList>& InCommandList) override;
+
+	virtual void OnActivate() override;
+	virtual void OnDeactivate() override;
+
+	virtual void OnExtendColumns(UE::Sequencer::NavigationTool::FNavigationToolColumnExtender& OutExtender) override;
+	virtual void OnExtendColumnViews(TSet<FNavigationToolColumnView>& OutColumnViews) override;
+	virtual void OnExtendItemChildren(UE::Sequencer::NavigationTool::INavigationTool& InTool
+		, const UE::Sequencer::NavigationTool::FNavigationToolItemPtr& InParentItem
+		, TArray<UE::Sequencer::NavigationTool::FNavigationToolItemPtr>& OutChildren
+		, const bool bInRecursive) override;
+	virtual void OnExtendBuiltInFilters(TArray<FNavigationToolBuiltInFilterParams>& OutFilterParams) override;
+
+	//~ End INavigationToolProvider
+
+private:
+	static const FName ToolbarSectionName;
+	static const FName ContextMenuSectionName;
+
+	IAvaSceneInterface* GetSceneInterface(const UE::Sequencer::NavigationTool::INavigationTool& InTool) const;
+	IAvaSequenceProvider* GetSequenceProvider(const UE::Sequencer::NavigationTool::INavigationTool& InTool) const;
+
+	void ExtendToolToolBar();
+	void ExtendToolItemContextMenu();
+
+	void RemoveToolToolBarExtension();
+	void RemoveToolItemContextMenuExtension();
+
+	void ExtendToolItemDragDropOp(UE::Sequencer::NavigationTool::FNavigationToolItemDragDropOp& InDragDropOp);
+
+	void OnSequenceAdded(UAvaSequence* const InAvaSequence);
+
+	void GeneratePresetMenu(UToolMenu* const InToolMenu);
+
+	bool CanRelabelSequence() const;
+	void RelabelSequence();
+
+	bool CanAddSequenceToSelection() const;
+	void AddSequenceToSelection();
+
+	TWeakPtr<FAvaSequencer> WeakAvaSequencer;
+
+	TSharedRef<FUICommandList> ToolCommands;
+
+	FDelegateHandle DragDropInitializedDelegate;
+};
+
+} // UE::AvaSequencer
diff --git a/Engine/Plugins/Experimental/Avalanche/Source/AvalancheSequencer/Private/NavigationTool/Columns/AvaNavigationToolStatusColumn.cpp b/Engine/Plugins/Experimental/Avalanche/Source/AvalancheSequencer/Private/NavigationTool/Columns/AvaNavigationToolStatusColumn.cpp
new file mode 100644
--- /dev/null
+++ b/Engine/Plugins/Experimental/Avalanche/Source/AvalancheSequencer/Private/NavigationTool/Columns/AvaNavigationToolStatusColumn.cpp
@@ -0,0 +1,67 @@
+﻿// Copyright Epic Games, Inc. All Rights Reserved.
+
+#include "AvaNavigationToolStatusColumn.h"
+#include "ISequencer.h"
+#include "NavigationTool/NavigationToolAvaSequence.h"
+#include "NavigationTool/Columns/SAvaNavigationToolStatus.h"
+#include "Styling/AppStyle.h"
+#include "Tools/NavigationTool/INavigationTool.h"
+#include "Tools/NavigationTool/INavigationToolView.h"
+#include "Tools/NavigationTool/NavigationToolDefines.h"
+
+#define LOCTEXT_NAMESPACE "AvaNavigationToolStatusColumn"
+
+namespace UE::Sequencer::NavigationTool
+{
+
+FText FAvaNavigationToolStatusColumn::GetColumnDisplayNameText() const
+{
+	return LOCTEXT("StatusColumn", "Status");
+}
+
+const FSlateBrush* FAvaNavigationToolStatusColumn::GetIconBrush() const
+{
+	return FAppStyle::GetBrush(TEXT("FoliageEditMode.BubbleBorder"));
+}
+
+SHeaderRow::FColumn::FArguments FAvaNavigationToolStatusColumn::ConstructHeaderRowColumn(const TSharedRef<INavigationToolView>& InToolView, const float InFillSize)
+{
+	const FName ColumnId = GetColumnId();
+
+	return SHeaderRow::Column(ColumnId)
+		.FillWidth(InFillSize)
+		.HAlignHeader(HAlign_Center)
+		.VAlignHeader(VAlign_Center)
+		.HAlignCell(HAlign_Fill)
+		.VAlignCell(VAlign_Center)
+		.DefaultLabel(GetColumnDisplayNameText())
+		.OnGetMenuContent(InToolView, &INavigationToolView::GetColumnMenuContent, ColumnId);
+}
+
+TSharedRef<SWidget> FAvaNavigationToolStatusColumn::ConstructRowWidget(const FNavigationToolItemPtr& InItem
+	, const TSharedRef<INavigationToolView>& InToolView
+	, const TSharedRef<SNavigationToolTreeRow>& InRow)
+{
+	if (!InItem->IsA<FNavigationToolAvaSequence>())
+	{
+		return SNullWidget::NullWidget;
+	}
+
+	const TSharedPtr<INavigationTool> OwnerTool = InToolView->GetOwnerTool();
+	if (!OwnerTool.IsValid())
+	{
+		return SNullWidget::NullWidget;
+	}
+
+	const TSharedPtr<ISequencer> Sequencer = OwnerTool->GetSequencer();
+	if (!Sequencer.IsValid())
+	{
+		return SNullWidget::NullWidget;
+	}
+
+	return SNew(SAvaNavigationToolStatus, InItem, InToolView, InRow, Sequencer.ToSharedRef());
+}
+
+} // UE::Sequencer::NavigationTool
+
+#undef LOCTEXT_NAMESPACE
diff --git a/Engine/Plugins/Experimental/Avalanche/Source/AvalancheSequencer/Private/NavigationTool/Columns/AvaNavigationToolStatusColumn.h b/Engine/Plugins/Experimental/Avalanche/Source/AvalancheSequencer/Private/NavigationTool/Columns/AvaNavigationToolStatusColumn.h
new file mode 100644
--- /dev/null
+++ b/Engine/Plugins/Experimental/Avalanche/Source/AvalancheSequencer/Private/NavigationTool/Columns/AvaNavigationToolStatusColumn.h
@@ -0,0 +1,31 @@
+﻿// Copyright Epic Games, Inc. All Rights Reserved.
+
+#pragma once
+
+#include "Tools/NavigationTool/Columns/NavigationToolColumn.h"
+
+namespace UE::Sequencer::NavigationTool
+{
+
+class FAvaNavigationToolStatusColumn : public FNavigationToolColumn
+{
+public:
+	UE_NAVIGATIONTOOL_INHERITS_WITH_SUPER(FAvaNavigationToolStatusColumn, FNavigationToolColumn);
+
+	static FName StaticColumnId() { return TEXT("Status"); }
+
+protected:
+	//~ Begin INavigationToolColumn
+	virtual FName GetColumnId() const override { return StaticColumnId(); }
+	virtual FText GetColumnDisplayNameText() const override;
+	virtual const FSlateBrush* GetIconBrush() const override;
+	virtual bool ShouldShowColumnByDefault() const override { return false; }
+	virtual float GetFillWidth() const override { return 10.f; }
+	virtual SHeaderRow::FColumn::FArguments ConstructHeaderRowColumn(const TSharedRef<INavigationToolView>& InToolView, const float InFillSize) override;
+	virtual TSharedRef<SWidget> ConstructRowWidget(const FNavigationToolItemPtr& InItem
+		, const TSharedRef<INavigationToolView>& InToolView
+		, const TSharedRef<SNavigationToolTreeRow>& InRow) override;
+	//~ End INavigationToolColumn
+};
+
+} // UE::Sequencer::NavigationTool
diff --git a/Engine/Plugins/Experimental/Avalanche/Source/AvalancheSequencer/Private/NavigationTool/Columns/SAvaNavigationToolStatus.cpp b/Engine/Plugins/Experimental/Avalanche/Source/AvalancheSequencer/Private/NavigationTool/Columns/SAvaNavigationToolStatus.cpp
new file mode 100644
--- /dev/null
+++ b/Engine/Plugins/Experimental/Avalanche/Source/AvalancheSequencer/Private/NavigationTool/Columns/SAvaNavigationToolStatus.cpp
@@ -0,0 +1,146 @@
+﻿// Copyright Epic Games, Inc. All Rights Reserved.
+
+#include "SAvaNavigationToolStatus.h"
+#include "AvaSequence.h"
+#include "AvaSequencePlayer.h"
+#include "AvaSequencerUtils.h"
+#include "IAvaSceneInterface.h"
+#include "IAvaSequencerProvider.h"
+#include "ISequencer.h"
+#include "NavigationTool/NavigationToolAvaSequence.h"
+#include "Tools/NavigationTool/INavigationToolView.h"
+#include "Tools/NavigationTool/Widgets/SNavigationToolTreeRow.h"
+#include "Widgets/Layout/SScrollBox.h"
+#include "Widgets/Notifications/SProgressBar.h"
+
+#define LOCTEXT_NAMESPACE "SAvaNavigationToolStatus"
+
+namespace UE::Sequencer::NavigationTool
+{
+
+void SAvaNavigationToolStatus::Construct(const FArguments& InArgs
+	, const FNavigationToolItemPtr& InItem
+	, const TSharedRef<INavigationToolView>& InView
+	, const TSharedRef<SNavigationToolTreeRow>& InRowWidget
+	, const TSharedRef<ISequencer>& InSequencer)
+{
+	WeakItem = InItem;
+	WeakView = InView;
+	WeakRowWidget = InRowWidget;
+	WeakSequencer = InSequencer;
+
+	FNavigationToolAvaSequence* const AvaSequenceItem = InItem->CastTo<FNavigationToolAvaSequence>();
+	if (!AvaSequenceItem)
+	{
+		return;
+	}
+
+	ChildSlot
+	[
+		SNew(SBox)
+		.Padding(5.f, 2.f)
+		[
+			SNew(SOverlay)
+			+ SOverlay::Slot()
+			.Padding(1.f, 1.f)
+			[
+				SNew(SProgressBar)
+				.Percent(this, &SAvaNavigationToolStatus::GetProgressPercent)
+				.Visibility(EVisibility::Visible)
+			]
+			+ SOverlay::Slot()
+			.Padding(1.f, 0.f)
+			.VAlign(VAlign_Center)
+			[
+				SNew(STextBlock)
+				.Text(this, &SAvaNavigationToolStatus::GetProgressText)
+				.TextStyle(&FAppStyle::GetWidgetStyle<FTextBlockStyle>(TEXT("SmallText")))
+				.Justification(ETextJustify::Type::Center)
+			]
+		]
+	];
+}
+
+void SAvaNavigationToolStatus::Tick(const FGeometry& InAllottedGeometry, const double InCurrentTime, const float InDeltaTime)
+{
+	if (!WeakItem.IsValid() || !WeakSequencer.IsValid())
+	{
+		return;
+	}
+
+	const TSharedPtr<ISequencer> Sequencer = WeakSequencer.Pin();
+	if (!ensure(Sequencer.IsValid()))
+	{
+		return;
+	}
+
+	const FNavigationToolAvaSequence* const AvaSequenceItem = WeakItem.Pin()->CastTo<FNavigationToolAvaSequence>();
+	if (!ensure(AvaSequenceItem))
+	{
+		return;
+	}
+
+	UAvaSequence* const Sequence = AvaSequenceItem->GetAvaSequence();
+	if (!ensure(IsValid(Sequence)))
+	{
+		return;
+	}
+
+	IAvaSceneInterface* const SceneInterface = FAvaSequencerUtils::GetSceneInterface(Sequencer.ToSharedRef());
+	if (!ensure(SceneInterface))
+	{
+		return;
+	}
+
+	IAvaSequencePlaybackObject* const PlaybackObject = SceneInterface->GetPlaybackObject();
+	if (!PlaybackObject)
+	{
+		return;
+	}
+
+	UAvaSequencePlayer* const Player = PlaybackObject->GetSequencePlayer(Sequence);
+
+	bSequenceInProgress = IsValid(Player);
+	if (bSequenceInProgress)
+	{
+		const EMovieScenePlayerStatus::Type Status = Player->GetPlaybackStatus();
+		CurrentFrame = Player->GetCurrentTime().ConvertTo(Player->GetDisplayRate());
+		TotalFrames = Player->GetDuration().ConvertTo(Player->GetDisplayRate());
+
+		FFormatNamedArguments Args;
+		if (Status == EMovieScenePlayerStatus::Playing)
+		{
+			Args.Add("Status", LOCTEXT("SequenceStatus_Playing", "Playing"));
+		}
+		else
+		{
+			Args.Add("Status", LOCTEXT("SequenceStatus_Stopped" , "Stopped"));
+		}
+		Args.Add("Current", FText::AsNumber(CurrentFrame.GetFrame().Value));
+		Args.Add("Total", FText::AsNumber(TotalFrames.GetFrame().Value));
+
+		StatusText =  FText::Format(LOCTEXT("SequenceStatus_Text", "{Status} ({Current} / {Total})"), Args);
+
+		Progress = FMath::IsNearlyZero(TotalFrames.AsDecimal())
+			? 0.f : (CurrentFrame.AsDecimal() / TotalFrames.AsDecimal());
+	}
+	else
+	{
+		StatusText = LOCTEXT("SequenceStatus_Unknown", "Not Playing");
+		Progress = 0.f;
+	}
+}
+
+FText SAvaNavigationToolStatus::GetProgressText() const
+{
+	return StatusText;
+}
+
+TOptional<float> SAvaNavigationToolStatus::GetProgressPercent() const
+{
+	return Progress;
+}
+
+} // UE::Sequencer::NavigationTool
+
+#undef LOCTEXT_NAMESPACE
diff --git a/Engine/Plugins/Experimental/Avalanche/Source/AvalancheSequencer/Private/NavigationTool/Columns/SAvaNavigationToolStatus.h b/Engine/Plugins/Experimental/Avalanche/Source/AvalancheSequencer/Private/NavigationTool/Columns/SAvaNavigationToolStatus.h
new file mode 100644
--- /dev/null
+++ b/Engine/Plugins/Experimental/Avalanche/Source/AvalancheSequencer/Private/NavigationTool/Columns/SAvaNavigationToolStatus.h
@@ -0,0 +1,51 @@
+﻿// Copyright Epic Games, Inc. All Rights Reserved.
+
+#pragma once
+
+#include "Misc/FrameTime.h"
+#include "Tools/NavigationTool/NavigationToolDefines.h"
+#include "Widgets/SCompoundWidget.h"
+
+class FText;
+class ISequencer;
+
+namespace UE::Sequencer::NavigationTool
+{
+
+class INavigationToolView;
+class SNavigationToolTreeRow;
+
+class SAvaNavigationToolStatus : public SCompoundWidget
+{
+public:
+	SLATE_BEGIN_ARGS(SAvaNavigationToolStatus) {}
+	SLATE_END_ARGS()
+
+	void Construct(const FArguments& InArgs
+		, const FNavigationToolItemPtr& InItem
+		, const TSharedRef<INavigationToolView>& InView
+		, const TSharedRef<SNavigationToolTreeRow>& InRowWidget
+		, const TSharedRef<ISequencer>& InSequencer);
+
+private:
+	//~ Begin SWidget
+	virtual void Tick(const FGeometry& InAllottedGeometry, const double InCurrentTime, const float InDeltaTime) override;
+	//~ End SWidget
+
+	FText GetProgressText() const;
+
+	TOptional<float> GetProgressPercent() const;
+
+	FNavigationToolItemWeakPtr WeakItem;
+	TWeakPtr<INavigationToolView> WeakView;
+	TWeakPtr<SNavigationToolTreeRow> WeakRowWidget;
+	TWeakPtr<ISequencer> WeakSequencer;
+
+	FText StatusText;
+	FFrameTime CurrentFrame;
+	FFrameTime TotalFrames;
+	float Progress = 0.f;
+	bool bSequenceInProgress = false;
+};
+
+} // UE::Sequencer::NavigationTool
diff --git a/Engine/Plugins/Experimental/Avalanche/Source/AvalancheSequencer/Private/NavigationTool/DragDrop/NavigationToolAvaSequenceDropHandler.cpp b/Engine/Plugins/Experimental/Avalanche/Source/AvalancheSequencer/Private/NavigationTool/DragDrop/NavigationToolAvaSequenceDropHandler.cpp
new file mode 100644
--- /dev/null
+++ b/Engine/Plugins/Experimental/Avalanche/Source/AvalancheSequencer/Private/NavigationTool/DragDrop/NavigationToolAvaSequenceDropHandler.cpp
@@ -0,0 +1,215 @@
+// Copyright Epic Games, Inc. All Rights Reserved.
+
+#include "NavigationToolAvaSequenceDropHandler.h"
+#include "AvaSequence.h"
+#include "IAvaSequenceProvider.h"
+#include "IAvaSequencer.h"
+#include "IAvaSequencerProvider.h"
+#include "Misc/Optional.h"
+#include "NavigationTool/NavigationToolAvaSequence.h"
+#include "ScopedTransaction.h"
+#include "Tools/NavigationTool/INavigationTool.h"
+#include "Tools/NavigationTool/ItemActions/NavigationToolAddItem.h"
+#include "Tools/NavigationTool/ItemActions/NavigationToolRemoveItem.h"
+
+#define LOCTEXT_NAMESPACE "NavigationToolAvaSequenceDropHandler"
+
+namespace UE::Sequencer::NavigationTool
+{
+
+FNavigationToolAvaSequenceDropHandler::FNavigationToolAvaSequenceDropHandler(const TWeakPtr<IAvaSequencer>& InWeakAvaSequencer)
+	: WeakAvaSequencer(InWeakAvaSequencer)
+{
+	
+}
+
+bool FNavigationToolAvaSequenceDropHandler::IsDraggedItemSupported(const FNavigationToolItemPtr& InDraggedItem) const
+{
+	return InDraggedItem->IsA<FNavigationToolAvaSequence>();
+}
+
+TOptional<EItemDropZone> FNavigationToolAvaSequenceDropHandler::CanDrop(const EItemDropZone InDropZone, const FNavigationToolItemPtr& InTargetItem) const
+{
+	const bool bTargetIsRoot = (InTargetItem->GetItemId() == FNavigationToolItemId::RootId);
+
+	UAvaSequence* TargetAvaSequence = nullptr;
+
+	if (!bTargetIsRoot)
+	{
+		if (FNavigationToolAvaSequence* const TargetAvaSequenceItem = InTargetItem->CastTo<FNavigationToolAvaSequence>())
+		{
+			TargetAvaSequence = TargetAvaSequenceItem->GetAvaSequence();
+		}
+		if (!TargetAvaSequence)
+		{
+			return TOptional<EItemDropZone>();
+		}
+	}
+
+	for (const FNavigationToolItemPtr& Item : Items)
+	{
+		// Only allow moving sequences to other Motion Design sequences
+		const FNavigationToolAvaSequence* const AvaSequenceItemToMove = Item->CastTo<FNavigationToolAvaSequence>();
+		if (!AvaSequenceItemToMove)
+		{
+			return TOptional<EItemDropZone>();
+		}
+
+		// Only allow moving valid sequences and sequences that aren't the target
+		UAvaSequence* const AvaSequenceToMove = AvaSequenceItemToMove->GetAvaSequence();
+		if (!IsValid(AvaSequenceToMove) || AvaSequenceToMove == TargetAvaSequence)
+		{
+			return TOptional<EItemDropZone>();
+		}
+
+		if (!bTargetIsRoot && IsValid(TargetAvaSequence))
+		{
+			if (TargetAvaSequence->GetChildren().Contains(AvaSequenceToMove))
+			{
+				return TOptional<EItemDropZone>();
+			}
+		}
+	}
+
+	switch (ActionType)
+	{
+		case ENavigationToolDragDropActionType::Move:
+			// Make sure the destination is not one of the items we're moving
+			if (!Items.Contains(InTargetItem))
+			{
+				return InDropZone;
+			}
+			break;
+
+		case ENavigationToolDragDropActionType::Copy:
+			return InDropZone;
+	}
+
+	return TOptional<EItemDropZone>();
+}
+
+bool FNavigationToolAvaSequenceDropHandler::Drop(const EItemDropZone InDropZone, const FNavigationToolItemPtr& InTargetItem)
+{
+	switch (ActionType)
+	{
+	case ENavigationToolDragDropActionType::Move:
+		MoveItems(InDropZone, InTargetItem);
+		break;
+
+	case ENavigationToolDragDropActionType::Copy:
+		DuplicateItems(Items, InTargetItem, InDropZone);
+		break;
+
+	default:
+		return false;
+	}
+
+	return true;
+}
+
+void FNavigationToolAvaSequenceDropHandler::MoveItems(const EItemDropZone InDropZone, const FNavigationToolItemPtr& InTargetItem)
+{
+	TArray<TSharedPtr<INavigationToolAction>> ItemActions;
+
+	FNavigationToolAddItemParams AddParams;
+	AddParams.RelativeItem = InTargetItem;
+	AddParams.RelativeDropZone = InDropZone;
+	AddParams.Flags = ENavigationToolAddItemFlags::Select | ENavigationToolAddItemFlags::Transact;
+	AddParams.SelectionFlags = ENavigationToolItemSelectionFlags::AppendToCurrentSelection;
+
+	const TSet<FNavigationToolItemPtr> DraggedItemSet(Items);
+
+	// Remove all Items whose Parent is in the Item Set 
+	Items.RemoveAll([&DraggedItemSet](const FNavigationToolItemPtr& InItem)
+		{
+			return !InItem.IsValid()
+				|| !InItem->GetParent().IsValid()
+				|| DraggedItemSet.Contains(InItem->GetParent());
+		});
+
+	// Reverse order for onto since Item->AddChild(...) adds it to Index 0, so last item would be at the top, which is reversed
+	if (InDropZone == EItemDropZone::OntoItem)
+	{
+		Algo::Reverse(Items);
+	}
+
+	for (const FNavigationToolItemPtr& Item : Items)
+	{
+		// Only allow moving sequences to other Motion Design sequences
+		if (FNavigationToolAvaSequence* const AvaSequenceItemToMove = Item->CastTo<FNavigationToolAvaSequence>())
+		{
+			AddParams.Item = SharedThis(AvaSequenceItemToMove);
+			ItemActions.Add(MakeShared<FNavigationToolAddItem>(AddParams));
+		}
+	}
+
+	InTargetItem->GetOwnerTool().EnqueueItemActions(MoveTemp(ItemActions));
+}
+
+void FNavigationToolAvaSequenceDropHandler::DuplicateItems(const TArray<FNavigationToolItemPtr>& InItems
+	, const FNavigationToolItemPtr& InRelativeItem
+	, const TOptional<EItemDropZone>& InRelativeDropZone)
+{
+	const TSharedPtr<IAvaSequencer> AvaSequencer = WeakAvaSequencer.Pin();
+	if (!AvaSequencer.IsValid())
+	{
+		return;
+	}
+
+	const IAvaSequencerProvider& SequencerProvider = AvaSequencer->GetProvider();
+
+	IAvaSequenceProvider* const SequenceProvider = SequencerProvider.GetSequenceProvider();
+	if (!SequenceProvider)
+	{
+		return;
+	}
+
+	UObject* const Outer = SequenceProvider->ToUObject();
+	if (!IsValid(Outer))
+	{
+		return;
+	}
+
+	// Gather list of sequences to duplicate
+	TSet<UAvaSequence*> SequencesToDuplicate;
+
+	for (const FNavigationToolItemPtr& Item : InItems)
+	{
+		const FNavigationToolAvaSequence* const AvaSequenceItem = Item->CastTo<FNavigationToolAvaSequence>();
+		if (!AvaSequenceItem)
+		{
+			continue;
+		}
+
+		UAvaSequence* const TemplateSequence = AvaSequenceItem->GetAvaSequence();
+		if (!IsValid(TemplateSequence))
+		{
+			continue;
+		}
+
+		SequencesToDuplicate.Add(TemplateSequence);
+	}
+
+	if (SequencesToDuplicate.IsEmpty())
+	{
+		return;
+	}
+
+	// Duplicate sequence objects
+	const FScopedTransaction Transaction(LOCTEXT("DuplicateSequencesTransaction", "Duplicate Sequence(s)"));
+
+	Outer->Modify();
+
+	for (const UAvaSequence* const TemplateSequence : SequencesToDuplicate)
+	{
+		UAvaSequence* const DupedSequence = DuplicateObject<UAvaSequence>(TemplateSequence, Outer);
+		if (IsValid(DupedSequence))
+		{
+			SequenceProvider->AddSequence(DupedSequence);
+		}
+	}
+}
+
+} // UE::Sequencer::NavigationTool
+
+#undef LOCTEXT_NAMESPACE
diff --git a/Engine/Plugins/Experimental/Avalanche/Source/AvalancheSequencer/Private/NavigationTool/DragDrop/NavigationToolAvaSequenceDropHandler.h b/Engine/Plugins/Experimental/Avalanche/Source/AvalancheSequencer/Private/NavigationTool/DragDrop/NavigationToolAvaSequenceDropHandler.h
new file mode 100644
--- /dev/null
+++ b/Engine/Plugins/Experimental/Avalanche/Source/AvalancheSequencer/Private/NavigationTool/DragDrop/NavigationToolAvaSequenceDropHandler.h
@@ -0,0 +1,38 @@
+﻿// Copyright Epic Games, Inc. All Rights Reserved.
+
+#pragma once
+
+#include "Templates/SharedPointer.h"
+#include "Tools/NavigationTool/DragDropOps/Handlers/NavigationToolItemDropHandler.h"
+
+class IAvaSequencer;
+enum class EItemDropZone;
+template<typename OptionalType> struct TOptional;
+
+namespace UE::Sequencer::NavigationTool
+{
+
+class FNavigationToolAvaSequenceDropHandler : public FNavigationToolItemDropHandler
+{
+public:
+	UE_NAVIGATIONTOOL_INHERITS(FNavigationToolAvaSequenceDropHandler, FNavigationToolItemDropHandler);
+
+	FNavigationToolAvaSequenceDropHandler(const TWeakPtr<IAvaSequencer>& InWeakAvaSequencer);
+
+protected:
+	//~ Begin FNavigationToolItemDropHandler
+	virtual bool IsDraggedItemSupported(const FNavigationToolItemPtr& InDraggedItem) const override;
+	virtual TOptional<EItemDropZone> CanDrop(const EItemDropZone InDropZone, const FNavigationToolItemPtr& InTargetItem) const override;
+	virtual bool Drop(const EItemDropZone InDropZone, const FNavigationToolItemPtr& InTargetItem) override;
+	//~ End FNavigationToolItemDropHandler
+
+	void MoveItems(const EItemDropZone InDropZone, const FNavigationToolItemPtr& InTargetItem);
+
+	void DuplicateItems(const TArray<FNavigationToolItemPtr>& InItems
+		, const FNavigationToolItemPtr& InRelativeItem
+		, const TOptional<EItemDropZone>& InRelativeDropZone);
+
+	TWeakPtr<IAvaSequencer> WeakAvaSequencer;
+};
+
+} // UE::Sequencer::NavigationTool
diff --git a/Engine/Plugins/Experimental/Avalanche/Source/AvalancheSequencer/Private/NavigationTool/NavigationToolAvaSequence.cpp b/Engine/Plugins/Experimental/Avalanche/Source/AvalancheSequencer/Private/NavigationTool/NavigationToolAvaSequence.cpp
new file mode 100644
--- /dev/null
+++ b/Engine/Plugins/Experimental/Avalanche/Source/AvalancheSequencer/Private/NavigationTool/NavigationToolAvaSequence.cpp
@@ -0,0 +1,321 @@
+﻿// Copyright Epic Games, Inc. All Rights Reserved.
+
+#include "NavigationToolAvaSequence.h"
+#include "AvaEditorCoreStyle.h"
+#include "AvaSequence.h"
+#include "AvaSequencer.h"
+#include "AvaSequencerUtils.h"
+#include "IAvaSequencer.h"
+#include "ISequencer.h"
+#include "ScopedTransaction.h"
+#include "Styling/StyleColors.h"
+#include "Tools/NavigationTool/INavigationTool.h"
+#include "Tools/NavigationTool/Items/NavigationToolItemParameters.h"
+#include "Tools/NavigationTool/Items/NavigationToolSequence.h"
+
+#define LOCTEXT_NAMESPACE "NavigationToolAvaSequence"
+
+namespace UE::Sequencer::NavigationTool
+{
+
+using namespace UE::Sequencer;
+
+FNavigationToolAvaSequence::FNavigationToolAvaSequence(INavigationTool& InTool
+	, const FNavigationToolItemPtr& InParentItem
+	, UAvaSequence* const InAvaSequence)
+	: FNavigationToolSequence(InTool
+	, InParentItem
+	, InAvaSequence
+	, nullptr
+	, 0)
+{
+}
+
+bool FNavigationToolAvaSequence::AddChild(const FNavigationToolAddItemParams& InAddItemParams)
+{
+	if (!CanAddChild(InAddItemParams.Item))
+	{
+		return false;
+	}
+
+	const FNavigationToolAvaSequence* const AvaSequenceItemToAdd = InAddItemParams.Item->CastTo<FNavigationToolAvaSequence>();
+	if (!AvaSequenceItemToAdd)
+	{
+		return false;
+	}
+
+	AddChildChecked(InAddItemParams);
+
+	UAvaSequence* const ParentAvaSequence = GetAvaSequence();
+	UAvaSequence* const AvaSequenceToAdd = AvaSequenceItemToAdd->GetAvaSequence();
+
+	const TArray<TWeakObjectPtr<UAvaSequence>> TargetChildren = ParentAvaSequence->GetChildren();
+	if (TargetChildren.Contains(AvaSequenceToAdd))
+	{
+		return false;
+	}
+
+	// Remove the sequence from any parent it's a child of
+	UAvaSequence* const AvaSequenceToMoveParent = AvaSequenceToAdd->GetParent();
+	if (IsValid(AvaSequenceToMoveParent))
+	{
+		AvaSequenceToMoveParent->RemoveChild(AvaSequenceToAdd);
+	}
+
+	// Process moving sequence if it's parent is the target
+	if (ParentAvaSequence == AvaSequenceToAdd->GetParent())
+	{
+		ParentAvaSequence->RemoveChild(AvaSequenceToAdd);
+
+		if (UAvaSequence* const TargetParent = ParentAvaSequence->GetParent())
+		{
+			TargetParent->AddChild(AvaSequenceToAdd);
+		}
+
+		return true;
+	}
+
+	// Check if the Target Item is a Child of Sequence Item (to prevent cycle)
+	UAvaSequence* CurrentItem = ParentAvaSequence;
+	while (CurrentItem)
+	{
+		if (CurrentItem->GetParent() == AvaSequenceToAdd)
+		{
+			AvaSequenceToAdd->RemoveChild(CurrentItem);
+			CurrentItem->SetParent(AvaSequenceToAdd->GetParent());
+
+			break;
+		}
+
+		CurrentItem = CurrentItem->GetParent();
+	}
+
+	ParentAvaSequence->AddChild(AvaSequenceToAdd);
+
+	return true;
+}
+
+bool FNavigationToolAvaSequence::RemoveChild(const FNavigationToolRemoveItemParams& InRemoveItemParams)
+{
+	// Check that we're removing a child item that is directly under us, and not a grandchild
+	if (!InRemoveItemParams.Item.IsValid() || !Children.Contains(InRemoveItemParams.Item))
+	{
+		return false;
+	}
+
+	if (const FNavigationToolAvaSequence* const AvaSequenceItemToRemove = InRemoveItemParams.Item->CastTo<FNavigationToolAvaSequence>())
+	{
+		UAvaSequence* const AvaSequenceToRemove = AvaSequenceItemToRemove->GetAvaSequence();
+		if (IsValid(AvaSequenceToRemove))
+		{
+			// Remove the sequence from any parent it's a child of
+			UAvaSequence* const AvaSequenceToRemoveParent = AvaSequenceToRemove->GetParent();
+			if (IsValid(AvaSequenceToRemoveParent))
+			{
+				AvaSequenceToRemoveParent->RemoveChild(AvaSequenceToRemove);
+			}
+		}
+	}
+
+	return RemoveChildChecked(InRemoveItemParams.Item);
+}
+
+void FNavigationToolAvaSequence::FindChildren(TArray<FNavigationToolItemPtr>& OutChildren, bool bRecursive)
+{
+	FNavigationToolSequence::FindChildren(OutChildren, bRecursive);
+
+	UAvaSequence* const AvaSequence = GetAvaSequence();
+	if (!IsValid(AvaSequence))
+	{
+		return;
+	}
+
+	const TSharedPtr<ISequencer> Sequencer = Tool.GetSequencer();
+	if (!Sequencer.IsValid())
+	{
+		return;
+	}
+
+	const TSharedPtr<FNavigationToolProvider> Provider = GetProvider();
+	if (!Provider.IsValid())
+	{
+		return;
+	}
+
+	TArray<TWeakObjectPtr<UAvaSequence>> ChildSequences = AvaSequence->GetChildren();
+	ChildSequences.RemoveAll([](const TWeakObjectPtr<UAvaSequence>& InChild)
+		{
+			return !InChild.IsValid();
+		});
+	ChildSequences.Sort([](const TWeakObjectPtr<UAvaSequence>& InA, const TWeakObjectPtr<UAvaSequence>& InB)
+		{
+			return InA->GetDisplayName().CompareTo(InB->GetDisplayName()) > 0;
+		});
+
+	const TSharedRef<FNavigationToolAvaSequence> SharedThisRef = SharedThis(this);
+
+	for (const TWeakObjectPtr<UAvaSequence>& SequenceWeak : ChildSequences)
+	{
+		const FNavigationToolItemPtr NewItem = Tool.FindOrAdd<FNavigationToolAvaSequence>(Provider
+			, SharedThisRef, SequenceWeak.Get());
+		OutChildren.Add(NewItem);
+		if (bRecursive)
+		{
+			NewItem->FindChildren(OutChildren, bRecursive);
+		}
+	}
+}
+
+void FNavigationToolAvaSequence::GetItemProxies(TArray<TSharedPtr<FNavigationToolItemProxy>>& OutItemProxies)
+{
+	FNavigationToolSequence::GetItemProxies(OutItemProxies);
+}
+
+bool FNavigationToolAvaSequence::CanRename() const
+{
+	return IsValid(GetAvaSequence());
+}
+
+bool FNavigationToolAvaSequence::Rename(const FString& InName)
+{
+	UAvaSequence* const UnderlyingSequence = GetAvaSequence();
+	if (!IsValid(UnderlyingSequence))
+	{
+		return false;
+	}
+
+	bool bRenamed = false;
+
+	if (!InName.Equals(UnderlyingSequence->GetLabel().ToString(), ESearchCase::CaseSensitive))
+	{
+		const FScopedTransaction Transaction(LOCTEXT("AvaSequenceRename", "Rename Sequence"));
+
+		UnderlyingSequence->Modify();
+		UnderlyingSequence->SetLabel(FName(InName));
+
+		Tool.NotifyToolItemRenamed(SharedThis(this));
+
+		bRenamed = true;
+	}
+
+	return bRenamed;
+}
+
+bool FNavigationToolAvaSequence::CanDelete() const
+{
+	return IsValid(GetAvaSequence());
+}
+
+bool FNavigationToolAvaSequence::Delete()
+{
+	const TSharedPtr<ISequencer> Sequencer = Tool.GetSequencer();
+	if (!Sequencer.IsValid())
+	{
+		return false;
+	}
+
+	UAvaSequence* const AvaSequence = GetAvaSequence();
+	if (!IsValid(AvaSequence))
+	{
+		return false;
+	}
+
+	const TSharedPtr<FAvaSequencer> AvaSequencer = StaticCastSharedPtr<FAvaSequencer>(FAvaSequencerUtils::GetAvaSequencer(Sequencer.ToSharedRef()));
+	if (!AvaSequencer.IsValid())
+	{
+		return false;
+	}
+
+	AvaSequencer->DeleteSequences({ AvaSequence });
+
+	return true;
+}
+
+FText FNavigationToolAvaSequence::GetDisplayName() const
+{
+	UAvaSequence* const AvaSequence = GetAvaSequence();
+	if (IsValid(AvaSequence))
+	{
+		return AvaSequence->GetDisplayName();
+	}
+	return FText::GetEmpty();
+}
+
+FSlateIcon FNavigationToolAvaSequence::GetIcon() const
+{
+	return FSlateIcon(FAvaEditorCoreStyle::Get().GetStyleSetName(), TEXT("Icons.MotionDesign"));
+}
+
+FSlateColor FNavigationToolAvaSequence::GetIconColor() const
+{
+	const TSharedPtr<ISequencer> Sequencer = Tool.GetSequencer();
+	if (!Sequencer.IsValid())
+	{
+		return FStyleColors::Foreground;
+	}
+
+	const UMovieSceneSequence* const RootSequence = Sequencer->GetRootMovieSceneSequence();
+	if (!IsValid(RootSequence))
+	{
+		return FStyleColors::Foreground;
+	}
+
+	if (GetAvaSequence() == RootSequence)
+	{
+		return FStyleColors::AccentGreen;
+	}
+
+	return FStyleColors::Foreground;
+}
+
+void FNavigationToolAvaSequence::OnSelect()
+{
+	Super::OnSelect();
+}
+
+void FNavigationToolAvaSequence::OnDoubleClick()
+{
+	UAvaSequence* const AvaSequence = GetAvaSequence();
+	if (!IsValid(AvaSequence))
+	{
+		return;
+	}
+
+	const TSharedPtr<ISequencer> Sequencer = Tool.GetSequencer();
+	if (!Sequencer.IsValid())
+	{
+		return;
+	}
+
+	const TSharedPtr<IAvaSequencer> AvaSequencer = FAvaSequencerUtils::GetAvaSequencer(Sequencer.ToSharedRef());
+	if (!AvaSequencer.IsValid())
+	{
+		return;
+	}
+
+	AvaSequencer->SetViewedSequence(AvaSequence);
+}
+
+UAvaSequence* FNavigationToolAvaSequence::GetAvaSequence() const
+{
+	return Cast<UAvaSequence>(GetSequence());
+}
+
+FNavigationToolItemId FNavigationToolAvaSequence::CalculateItemId() const
+{
+	return FNavigationToolItemId(GetParent(), GetAvaSequence());
+}
+
+TOptional<FColor> FNavigationToolAvaSequence::GetColor() const
+{
+	return Tool.FindItemColor(SharedThis(const_cast<FNavigationToolAvaSequence*>(this)));
+}
+
+void FNavigationToolAvaSequence::SetColor(const TOptional<FColor>& InColor)
+{
+	Tool.SetItemColor(SharedThis(this), InColor.Get(FColor()));
+}
+
+} // UE::Sequencer::NavigationTool
+
+#undef LOCTEXT_NAMESPACE
diff --git a/Engine/Plugins/Experimental/Avalanche/Source/AvalancheSequencer/Private/NavigationTool/NavigationToolAvaSequence.h b/Engine/Plugins/Experimental/Avalanche/Source/AvalancheSequencer/Private/NavigationTool/NavigationToolAvaSequence.h
new file mode 100644
--- /dev/null
+++ b/Engine/Plugins/Experimental/Avalanche/Source/AvalancheSequencer/Private/NavigationTool/NavigationToolAvaSequence.h
@@ -0,0 +1,62 @@
+﻿// Copyright Epic Games, Inc. All Rights Reserved.
+
+#pragma once
+
+#include "Tools/NavigationTool/Items/NavigationToolSequence.h"
+
+class UAvaSequence;
+
+namespace UE::Sequencer::NavigationTool
+{
+
+/**
+ * Item in Navigation Tool representing a Motion Design Sequence.
+ */
+class FNavigationToolAvaSequence : public FNavigationToolSequence
+{
+public:
+	UE_NAVIGATIONTOOL_INHERITS_WITH_SUPER(FNavigationToolAvaSequence, FNavigationToolSequence);
+
+	FNavigationToolAvaSequence(INavigationTool& InTool
+		, const FNavigationToolItemPtr& InParentItem
+		, UAvaSequence* const InAvaSequence);
+
+	//~ Begin INavigationToolItem
+
+	virtual bool AddChild(const FNavigationToolAddItemParams& InAddItemParams) override;
+	virtual bool RemoveChild(const FNavigationToolRemoveItemParams& InRemoveItemParams) override;
+	virtual void FindChildren(TArray<FNavigationToolItemPtr>& OutChildren, bool bRecursive) override;
+	virtual void GetItemProxies(TArray<TSharedPtr<FNavigationToolItemProxy>>& OutItemProxies) override;
+
+	virtual bool ShouldSort() const override { return true; }
+	virtual bool CanBeTopLevel() const override { return true; }
+
+	virtual bool CanRename() const override;
+	virtual bool Rename(const FString& InName) override;
+
+	virtual bool CanDelete() const override;
+	virtual bool Delete() override;
+
+	virtual FText GetDisplayName() const override;
+	virtual FSlateIcon GetIcon() const override;
+	virtual FSlateColor GetIconColor() const override;
+
+	virtual void OnSelect() override;
+	virtual void OnDoubleClick() override;
+
+	//~ End INavigationToolItem
+
+	//~ Begin IColorExtension
+	virtual TOptional<FColor> GetColor() const override;
+	virtual void SetColor(const TOptional<FColor>& InColor) override;
+	//~ End IColorExtension
+
+	UAvaSequence* GetAvaSequence() const;
+
+protected:
+	//~Begin FNavigationToolItem
+	virtual FNavigationToolItemId CalculateItemId() const override;
+	//~End FNavigationToolItem
+};
+
+} // UE::Sequencer::NavigationTool
diff --git a/Engine/Plugins/Experimental/Avalanche/Source/AvalancheSequencer/Private/SequenceTree/Widgets/SAvaSequenceTree.cpp b/Engine/Plugins/Experimental/Avalanche/Source/AvalancheSequencer/Private/SequenceTree/Widgets/SAvaSequenceTree.cpp
--- a/Engine/Plugins/Experimental/Avalanche/Source/AvalancheSequencer/Private/SequenceTree/Widgets/SAvaSequenceTree.cpp
+++ b/Engine/Plugins/Experimental/Avalanche/Source/AvalancheSequencer/Private/SequenceTree/Widgets/SAvaSequenceTree.cpp
@@ -196,36 +196,36 @@ void SAvaSequenceTree::BindCommands(const TSharedRef<FAvaSequencer>& InSequencer
 	const FAvaSequencerCommands& AvaSequencerCommands = FAvaSequencerCommands::Get();
 
 	CommandList->MapAction(GenericCommands.Duplicate
-		, FExecuteAction::CreateSP(InSequencer, &FAvaSequencer::DuplicateSequence_Execute)
-		, FCanExecuteAction::CreateSP(InSequencer, &FAvaSequencer::DuplicateSequence_CanExecute));
+		, FExecuteAction::CreateSP(InSequencer, &FAvaSequencer::DuplicateSequences_Execute)
+		, FCanExecuteAction::CreateSP(InSequencer, &FAvaSequencer::DuplicateSequences_CanExecute));
 
 	CommandList->MapAction(GenericCommands.Delete
-		, FExecuteAction::CreateSP(InSequencer, &FAvaSequencer::DeleteSequence_Execute)
-		, FCanExecuteAction::CreateSP(InSequencer, &FAvaSequencer::DeleteSequence_CanExecute));
+		, FExecuteAction::CreateSP(InSequencer, &FAvaSequencer::DeleteSequences_Execute)
+		, FCanExecuteAction::CreateSP(InSequencer, &FAvaSequencer::DeleteSequences_CanExecute));
 
 	CommandList->MapAction(GenericCommands.Rename
 		, FExecuteAction::CreateSP(InSequencer, &FAvaSequencer::RelabelSequence_Execute)
 		, FCanExecuteAction::CreateSP(InSequencer, &FAvaSequencer::RelabelSequence_CanExecute));
 
 	CommandList->MapAction(AvaSequencerCommands.PlaySelected
-		, FExecuteAction::CreateSP(InSequencer, &FAvaSequencer::PlaySelected_Execute)
-		, FCanExecuteAction::CreateSP(InSequencer, &FAvaSequencer::PlaySelected_CanExecute));
+		, FExecuteAction::CreateSP(InSequencer, &FAvaSequencer::PlaySequences_Execute)
+		, FCanExecuteAction::CreateSP(InSequencer, &FAvaSequencer::PlaySequences_CanExecute));
 
 	CommandList->MapAction(AvaSequencerCommands.ContinueSelected
-		, FExecuteAction::CreateSP(InSequencer, &FAvaSequencer::ContinueSelected_Execute)
-		, FCanExecuteAction::CreateSP(InSequencer, &FAvaSequencer::ContinueSelected_CanExecute));
+		, FExecuteAction::CreateSP(InSequencer, &FAvaSequencer::ContinueSequences_Execute)
+		, FCanExecuteAction::CreateSP(InSequencer, &FAvaSequencer::ContinueSequences_CanExecute));
 
 	CommandList->MapAction(AvaSequencerCommands.StopSelected
-		, FExecuteAction::CreateSP(InSequencer, &FAvaSequencer::StopSelected_Execute)
-		, FCanExecuteAction::CreateSP(InSequencer, &FAvaSequencer::StopSelected_CanExecute));
+		, FExecuteAction::CreateSP(InSequencer, &FAvaSequencer::StopSequences_Execute)
+		, FCanExecuteAction::CreateSP(InSequencer, &FAvaSequencer::StopSequences_CanExecute));
 
 	CommandList->MapAction(AvaSequencerCommands.ExportSequence
-		, FExecuteAction::CreateSP(InSequencer, &FAvaSequencer::ExportSequence_Execute)
-		, FCanExecuteAction::CreateSP(InSequencer, &FAvaSequencer::ExportSequence_CanExecute));
+		, FExecuteAction::CreateSP(InSequencer, &FAvaSequencer::ExportSequences_Execute)
+		, FCanExecuteAction::CreateSP(InSequencer, &FAvaSequencer::ExportSequences_CanExecute));
 
 	CommandList->MapAction(AvaSequencerCommands.SpawnSequencePlayer
-		, FExecuteAction::CreateSP(InSequencer, &FAvaSequencer::SpawnPlayer_Execute)
-		, FCanExecuteAction::CreateSP(InSequencer, &FAvaSequencer::SpawnPlayer_CanExecute));
+		, FExecuteAction::CreateSP(InSequencer, &FAvaSequencer::SpawnPlayers_Execute)
+		, FCanExecuteAction::CreateSP(InSequencer, &FAvaSequencer::SpawnPlayers_CanExecute));
 }
 
 TSharedRef<ITableRow> SAvaSequenceTree::OnGenerateRow(FAvaSequenceItemPtr InListItem
@@ -324,7 +324,7 @@ FReply SAvaSequenceTree::OnNewSequenceClicked()
 {
 	if (SequencerWeak.IsValid())
 	{
-		SequencerWeak.Pin()->AddSequence_Execute();
+		SequencerWeak.Pin()->AddSequence();
 		return FReply::Handled();
 	}
 	return FReply::Unhandled();
diff --git a/Engine/Plugins/Experimental/Avalanche/Source/AvalancheSequencer/Public/AvaSequencerUtils.h b/Engine/Plugins/Experimental/Avalanche/Source/AvalancheSequencer/Public/AvaSequencerUtils.h
--- a/Engine/Plugins/Experimental/Avalanche/Source/AvalancheSequencer/Public/AvaSequencerUtils.h
+++ b/Engine/Plugins/Experimental/Avalanche/Source/AvalancheSequencer/Public/AvaSequencerUtils.h
@@ -7,7 +7,14 @@
 #include "Templates/SharedPointerFwd.h"
 #include "Templates/Function.h"
 
+class IAvaSceneInterface;
+class IAvaSequenceProvider;
+class IAvaSequencer;
 class IAvaSequencerController;
+class ISequencer;
+class UAvaSceneSubsystem;
+class UAvaSequencerSubsystem;
+class UWorld;
 enum class EMovieSceneTransformChannel : uint32;
 
 struct FAvaSequencerUtils
@@ -29,4 +36,22 @@ struct FAvaSequencerUtils
 	}
 
 	AVALANCHESEQUENCER_API static TSharedRef<IAvaSequencerController> CreateSequencerController();
+
+	/** @return Motion Design Sequencer World from an ISequencer */
+	static UWorld* GetSequencerWorld(const TSharedRef<ISequencer>& InSequencer);
+
+	/** @return Motion Design Sequencer Subsystem from an ISequencer */
+	static UAvaSequencerSubsystem* GetSequencerSubsystem(const TSharedRef<ISequencer>& InSequencer);
+
+	/** @return Motion Design Scene Subsystem from an ISequencer */
+	static UAvaSceneSubsystem* GetSceneSubsystem(const TSharedRef<ISequencer>& InSequencer);
+
+	/** @return Motion Design Scene Interface from an ISequencer */
+	static IAvaSceneInterface* GetSceneInterface(const TSharedRef<ISequencer>& InSequencer);
+
+	/** @return Motion Design Sequence Provider from an ISequencer */
+	static IAvaSequenceProvider* GetSequenceProvider(const TSharedRef<ISequencer>& InSequencer);
+
+	/** @return Motion Design Sequencer from an ISequencer */
+	static TSharedPtr<IAvaSequencer> GetAvaSequencer(const TSharedRef<ISequencer>& InSequencer);
 };
diff --git a/Engine/Plugins/Experimental/Avalanche/Source/AvalancheSequencer/Public/IAvaSequencer.h b/Engine/Plugins/Experimental/Avalanche/Source/AvalancheSequencer/Public/IAvaSequencer.h
--- a/Engine/Plugins/Experimental/Avalanche/Source/AvalancheSequencer/Public/IAvaSequencer.h
+++ b/Engine/Plugins/Experimental/Avalanche/Source/AvalancheSequencer/Public/IAvaSequencer.h
@@ -2,11 +2,13 @@
 
 #pragma once
 
+#include "AvaSequence.h"
 #include "Containers/ArrayView.h"
 #include "Containers/ContainersFwd.h"
 #include "Containers/StringView.h"
 #include "Templates/SharedPointer.h"
 #include "Templates/UnrealTypeTraits.h"
+#include "UObject/WeakObjectPtrTemplatesFwd.h"
 
 class AActor;
 class FAvaEditorSelection;
@@ -16,7 +18,6 @@ class IAvaSequencerProvider;
 class IPropertyHandle;
 class ISequencer;
 class SWidget;
-class UAvaSequence;
 class UObject;
 class USequencerSettings;
 
@@ -59,4 +60,13 @@ public:
 
 	/** Should be called when the UAvaSequence tree has changed, so that FAvaSequencer can trigger a UI refresh */
 	virtual void NotifyOnSequenceTreeChanged() = 0;
+
+	/** Gets the list of embedded level sequences */
+	virtual const TArray<TWeakObjectPtr<UAvaSequence>>& GetRootSequences() const = 0;
+
+	DECLARE_MULTICAST_DELEGATE_OneParam(FOnSequenceAdded, UAvaSequence*);
+	virtual FOnSequenceAdded& OnSequenceAdded() = 0;
+
+	DECLARE_MULTICAST_DELEGATE_OneParam(FOnSequenceRemoved, UAvaSequence*);
+	virtual FOnSequenceRemoved& OnSequenceRemoved() = 0;
 };
