[Backout] - CL41856403
#rnx
Original CL Desc
-----------------------------------------------------------------
[MotionDesign] Extend Sequencer Navigation Tool to support embedded Motion Design sequences

diff --git a/Engine/Plugins/Experimental/Avalanche/Source/Avalanche/Avalanche.Build.cs b/Engine/Plugins/Experimental/Avalanche/Source/Avalanche/Avalanche.Build.cs
--- a/Engine/Plugins/Experimental/Avalanche/Source/Avalanche/Avalanche.Build.cs
+++ b/Engine/Plugins/Experimental/Avalanche/Source/Avalanche/Avalanche.Build.cs
@@ -47,7 +47,6 @@ public class Avalanche : ModuleRules
 				"HeadMountedDisplay",
 				"InputCore",
 				"MediaCompositing",
-				"MovieScene",
 				"Niagara",
 				"NiagaraCore",
 				"ProceduralMeshComponent",
diff --git a/Engine/Plugins/Experimental/Avalanche/Source/Avalanche/Public/AvaScene.h b/Engine/Plugins/Experimental/Avalanche/Source/Avalanche/Public/AvaScene.h
--- a/Engine/Plugins/Experimental/Avalanche/Source/Avalanche/Public/AvaScene.h
+++ b/Engine/Plugins/Experimental/Avalanche/Source/Avalanche/Public/AvaScene.h
@@ -4,7 +4,6 @@
 
 #include "AvaSceneTree.h"
 #include "Delegates/Delegate.h"
-#include "EditorSequenceNavigationDefs.h"
 #include "GameFramework/Actor.h"
 #include "IAvaRemoteControlInterface.h"
 #include "IAvaSceneInterface.h"
@@ -62,7 +61,7 @@ public:
 
 	//~ Begin IAvaSceneInterface
 	virtual ULevel* GetSceneLevel() const override;
-	virtual UAvaSceneSettings* GetSceneSettings() const override { return SceneSettings; }
+	virtual UAvaSceneSettings* GetSceneSettings() const { return SceneSettings; }
 	virtual UAvaSceneState* GetSceneState() const override { return SceneState; }
 	virtual FAvaSceneTree& GetSceneTree() override { return SceneTree; }
 	virtual const FAvaSceneTree& GetSceneTree() const override { return SceneTree; }
@@ -70,7 +69,6 @@ public:
 	virtual IAvaSequenceProvider* GetSequenceProvider() override { return this; }
 	virtual const IAvaSequenceProvider* GetSequenceProvider() const override { return this; }
 	virtual URemoteControlPreset* GetRemoteControlPreset() const override { return RemoteControlPreset; }
-	virtual FNavigationToolSaveState& GetNavigationToolSaveState() override { return NavigationToolState; }
 	//~ End IAvaSceneInterface
 
 	//~ Begin IAvaSequenceProvider
@@ -171,10 +169,6 @@ protected:
 	/** Whether the motion design mode should auto start */
 	UPROPERTY()
 	bool bAutoStartMode = true;
-
-	/** Saved Navigation Tool state to restore */
-	UPROPERTY()
-	FNavigationToolSaveState NavigationToolState;
 #endif
 
 };
diff --git a/Engine/Plugins/Experimental/Avalanche/Source/Avalanche/Public/IAvaSceneInterface.h b/Engine/Plugins/Experimental/Avalanche/Source/Avalanche/Public/IAvaSceneInterface.h
--- a/Engine/Plugins/Experimental/Avalanche/Source/Avalanche/Public/IAvaSceneInterface.h
+++ b/Engine/Plugins/Experimental/Avalanche/Source/Avalanche/Public/IAvaSceneInterface.h
@@ -13,10 +13,6 @@ class ULevel;
 class URemoteControlPreset;
 struct FAvaSceneTree;
 
-#if WITH_EDITOR
-struct FNavigationToolSaveState;
-#endif
-
 /**
  * Interface for the driver class of an Motion Design Scene
  * This interface's purpose is to provide access to scene data and to core interfaces in the Scene
@@ -49,8 +45,4 @@ public:
 	virtual IAvaSequencePlaybackObject* GetPlaybackObject() const = 0;
 
 	virtual URemoteControlPreset* GetRemoteControlPreset() const = 0;
-
-#if WITH_EDITOR
-	virtual FNavigationToolSaveState& GetNavigationToolSaveState() = 0;
-#endif
 };
diff --git a/Engine/Plugins/Experimental/Avalanche/Source/AvalancheEditorCore/Private/AvaEditorCoreStyle.cpp b/Engine/Plugins/Experimental/Avalanche/Source/AvalancheEditorCore/Private/AvaEditorCoreStyle.cpp
deleted file mode 100644
diff --git a/Engine/Plugins/Experimental/Avalanche/Source/AvalancheEditorCore/Public/AvaEditorCoreStyle.h b/Engine/Plugins/Experimental/Avalanche/Source/AvalancheEditorCore/Public/AvaEditorCoreStyle.h
deleted file mode 100644
diff --git a/Engine/Plugins/Experimental/Avalanche/Source/AvalancheSequencer/AvalancheSequencer.Build.cs b/Engine/Plugins/Experimental/Avalanche/Source/AvalancheSequencer/AvalancheSequencer.Build.cs
--- a/Engine/Plugins/Experimental/Avalanche/Source/AvalancheSequencer/AvalancheSequencer.Build.cs
+++ b/Engine/Plugins/Experimental/Avalanche/Source/AvalancheSequencer/AvalancheSequencer.Build.cs
@@ -26,7 +26,6 @@ public class AvalancheSequencer : ModuleRules
 			new string[]
 			{
 				"ApplicationCore",
-				"Avalanche",
 				"AvalancheCore",
 				"AvalancheEditorCore",
 				"AvalancheOutliner",
diff --git a/Engine/Plugins/Experimental/Avalanche/Source/AvalancheSequencer/Private/AvaSequencer.cpp b/Engine/Plugins/Experimental/Avalanche/Source/AvalancheSequencer/Private/AvaSequencer.cpp
--- a/Engine/Plugins/Experimental/Avalanche/Source/AvalancheSequencer/Private/AvaSequencer.cpp
+++ b/Engine/Plugins/Experimental/Avalanche/Source/AvalancheSequencer/Private/AvaSequencer.cpp
@@ -1,7 +1,6 @@
 // Copyright Epic Games, Inc. All Rights Reserved.
 
 #include "AvaSequencer.h"
-#include "AvaEditorCoreStyle.h"
 #include "AvaSequence.h"
 #include "AvaSequenceActor.h"
 #include "AvaSequencePlaybackObject.h"
@@ -12,7 +11,6 @@
 #include "Clipboard/AvaSequenceImporter.h"
 #include "Commands/AvaSequencerAction.h"
 #include "Commands/AvaSequencerCommands.h"
-#include "ContentBrowserModule.h"
 #include "CoreGlobals.h"
 #include "DetailsView/SAvaMarkDetails.h"
 #include "DetailsView/Section/AvaSequencePlaybackDetails.h"
@@ -21,14 +19,11 @@
 #include "DetailsView/Section/AvaSequencerEaseCurveToolSection.h"
 #include "EaseCurveTool/AvaEaseCurveTool.h"
 #include "EaseCurveTool/AvaEaseCurveToolCommands.h"
-#include "Editor.h"
 #include "Editor/Sequencer/Private/Sequencer.h"
 #include "Framework/Commands/GenericCommands.h"
 #include "Framework/Commands/UIAction.h"
 #include "Framework/MultiBox/MultiBoxBuilder.h"
-#include "Framework/Notifications/NotificationManager.h"
 #include "IAvaSequenceProvider.h"
-#include "IContentBrowserSingleton.h"
 #include "ISequencer.h"
 #include "ISequencerModule.h"
 #include "ISequencerTrackEditor.h"
@@ -38,30 +33,32 @@
 #include "MVVM/Views/SOutlinerView.h"
 #include "Misc/MessageDialog.h"
 #include "Misc/NotifyHook.h"
+#include "Misc/TextFilter.h"
 #include "MovieScene.h"
-#include "NavigationTool/AvaNavigationToolProvider.h"
-#include "NavigationTool/NavigationToolAvaSequence.h"
 #include "Playback/AvaSequencerCleanView.h"
 #include "Playback/AvaSequencerController.h"
 #include "PropertyHandle.h"
 #include "ScopedTransaction.h"
 #include "Selection/AvaEditorSelection.h"
-#include "SequencerCommands.h"
-#include "SequencerSettings.h"
-#include "SequencerUtilities.h"
 #include "SequenceTree/AvaSequenceItem.h"
-#include "SequenceTree/IAvaSequenceItem.h"
 #include "SequenceTree/Columns/AvaSequenceNameColumn.h"
 #include "SequenceTree/Columns/AvaSequenceStatusColumn.h"
+#include "SequenceTree/Columns/IAvaSequenceColumn.h"
 #include "SequenceTree/Widgets/SAvaSequenceTree.h"
+#include "SequencerCommands.h"
+#include "SequencerSettings.h"
+#include "SequencerUtilities.h"
 #include "Settings/AvaSequencerSettings.h"
+#include "Sidebar/SSidebar.h"
+#include "Sidebar/SSidebarContainer.h"
 #include "Sidebar/SidebarDrawerConfig.h"
 #include "StaggerTool/AvaStaggerTool.h"
+#include "ToolMenu.h"
+#include "ToolMenuEntry.h"
+#include "ToolMenuSection.h"
+#include "ToolMenus.h"
 #include "Toolkits/AssetEditorToolkit.h"
-#include "Tools/NavigationTool/INavigationTool.h"
-#include "Tools/NavigationTool/NavigationToolDefines.h"
-#include "Widgets/Notifications/SNotificationList.h"
-#include "Widgets/Views/SHeaderRow.h"
+#include "Widgets/Views/STreeView.h"
 
 DEFINE_LOG_CATEGORY_STATIC(LogAvaSequencer, Log, All);
 
@@ -120,7 +117,7 @@ namespace UE::AvaSequencer::Private
 	};
 }
 
-const FName FAvaSequencer::SidebarDrawerId = TEXT("MotionDesign");
+const FName FAvaSequencer::SidebarDrawerId = TEXT("Sequence");
 
 FAvaSequencer::FAvaSequencer(IAvaSequencerProvider& InProvider, FAvaSequencerArgs&& InArgs)
 	: Provider(InProvider)
@@ -170,10 +167,7 @@ FAvaSequencer::FAvaSequencer(IAvaSequencerProvider& InProvider, FAvaSequencerArg
 	}
 
 	// Register to update when an undo/redo operation has been called to update our list of items
-	if (GEditor)
-	{
-		GEditor->RegisterForUndo(this);
-	}
+	GEditor->RegisterForUndo(this);
 }
 
 FAvaSequencer::~FAvaSequencer()
@@ -183,7 +177,7 @@ FAvaSequencer::~FAvaSequencer()
 	OnSequenceStartedHandle.Reset();
 	OnSequenceFinishedHandle.Reset();
 
-	if (GEditor)
+	if (GEngine)
 	{
 		GEditor->UnregisterForUndo(this);
 	}
@@ -217,8 +211,6 @@ FAvaSequencer::~FAvaSequencer()
 			SidebarExtender.Reset();
 		}
 	}
-
-	UnregisterNavigationTool();
 }
 
 void FAvaSequencer::BindCommands()
@@ -254,7 +246,7 @@ void FAvaSequencer::BindCommands()
 		, FExecuteAction::CreateRaw(this, &FAvaSequencer::AddTransformKey, EMovieSceneTransformChannel::Rotation));
 
 	CommandList->MapAction(SequencerCommands.AddScaleKey
-		, FExecuteAction::CreateRaw(this, &FAvaSequencer::AddTransformKey, EMovieSceneTransformChannel::Scale));
+		, FExecuteAction::CreateRaw(this, &FAvaSequencer::AddTransformKey, EMovieSceneTransformChannel::Scale));	
 }
 
 TSharedPtr<IAvaSequenceColumn> FAvaSequencer::FindSequenceColumn(FName InColumnName) const
@@ -374,7 +366,6 @@ TSharedRef<ISequencer> FAvaSequencer::CreateSequencer()
 		SequencerInitParams.HostCapabilities.bSupportsCurveEditor = true;
 		SequencerInitParams.HostCapabilities.bSupportsSaveMovieSceneAsset = false;
 		SequencerInitParams.HostCapabilities.bSupportsSidebar = true;
-		SequencerInitParams.HostCapabilities.bSupportsNavigationTool = true;
 	};
 
 	InstancedSequencer = FAvaSequencerUtils::GetSequencerModule().CreateSequencer(SequencerInitParams);
@@ -558,7 +549,7 @@ UAvaSequence* FAvaSequencer::CreateSequence() const
 	}
 
 	UObject* const Outer = SequenceProvider->ToUObject();
-	if (!IsValid(Outer))
+	if (!ensure(Outer))
 	{
 		return nullptr;
 	}
@@ -641,13 +632,13 @@ TSharedRef<SWidget> FAvaSequencer::CreatePlayerToolBar(const TSharedRef<FUIComma
 	const FAvaSequencerCommands& Commands = FAvaSequencerCommands::Get();
 
 	ToolBarBuilder.AddToolBarButton(Commands.PlaySelected, NAME_None, TAttribute<FText>(), TAttribute<FText>()
-		, FSlateIcon(FAppStyle::GetAppStyleSetName(), TEXT("Icons.Toolbar.Play")));
+		, FSlateIcon(FAppStyle::GetAppStyleSetName(), "Icons.Toolbar.Play"));
 
 	ToolBarBuilder.AddToolBarButton(Commands.ContinueSelected, NAME_None, TAttribute<FText>(), TAttribute<FText>()
-		, FSlateIcon(FAppStyle::GetAppStyleSetName(), TEXT("Icons.JumpToEvent")));
+		, FSlateIcon(FAppStyle::GetAppStyleSetName(), "Icons.JumpToEvent"));
 
 	ToolBarBuilder.AddToolBarButton(Commands.StopSelected, NAME_None, TAttribute<FText>(), TAttribute<FText>()
-		, FSlateIcon(FAppStyle::GetAppStyleSetName(), TEXT("Icons.Toolbar.Stop")));
+		, FSlateIcon(FAppStyle::GetAppStyleSetName(), "Icons.Toolbar.Stop"));
 
 	return SNew(SBox)
 		.HAlign(EHorizontalAlignment::HAlign_Center)
@@ -768,11 +759,6 @@ void FAvaSequencer::AddTransformKey(EMovieSceneTransformChannel InTransformChann
 {
 	using namespace UE::AvaSequencer::Private;
 
-	if (!GEditor)
-	{
-		return;
-	}
-
 	FEditorModeTools* const EditorModeTools = Provider.GetSequencerModeTools();
 	if (!EditorModeTools)
 	{
@@ -854,21 +840,7 @@ void FAvaSequencer::ApplyCustomPresetToSelection(FName InPresetName)
 
 void FAvaSequencer::ApplyPresetToSelection(const FAvaSequencePreset& InPreset)
 {
-	using namespace UE::Sequencer::NavigationTool;
-
-	const TSharedPtr<ISequencer> Sequencer = SequencerWeak.Pin();
-	if (!Sequencer.IsValid())
-	{
-		return;
-	}
-
-	const TSharedPtr<INavigationTool> Tool = Sequencer->GetNavigationTool();
-	if (!Tool.IsValid())
-	{
-		return;
-	}
-
-	const TArray<FNavigationToolItemPtr> SelectedItems = Tool->GetSelectedItems();
+	TArray<TSharedPtr<IAvaSequenceItem>> SelectedItems = SequenceTreeView->GetSelectedItems();
 	if (SelectedItems.IsEmpty())
 	{
 		return;
@@ -876,11 +848,11 @@ void FAvaSequencer::ApplyPresetToSelection(const FAvaSequencePreset& InPreset)
 
 	FScopedTransaction Transaction(LOCTEXT("ApplySequencePreset", "Apply Sequence Preset"));
 
-	for (const FNavigationToolItemPtr& Item : SelectedItems)
+	for (const TSharedPtr<IAvaSequenceItem>& SequenceItem : SelectedItems)
 	{
-		if (FNavigationToolAvaSequence* const AvaSequenceItem = Item->CastTo<FNavigationToolAvaSequence>())
+		if (SequenceItem.IsValid())
 		{
-			InPreset.ApplyPreset(AvaSequenceItem->GetAvaSequence());
+			InPreset.ApplyPreset(SequenceItem->GetSequence());
 		}
 	}
 }
@@ -904,12 +876,12 @@ void FAvaSequencer::AddSequence_Execute()
 	SequenceProvider->AddSequence(Sequence);
 }
 
-bool FAvaSequencer::DuplicateSequences_CanExecute() const
+bool FAvaSequencer::DuplicateSequence_CanExecute() const
 {
 	return GetProvider().CanEditOrPlaySequences();
 }
 
-void FAvaSequencer::DuplicateSequences_Execute()
+void FAvaSequencer::DuplicateSequence_Execute()
 {
 	IAvaSequenceProvider* const SequenceProvider = Provider.GetSequenceProvider();
 	if (!SequenceTreeView.IsValid() || !SequenceProvider)
@@ -923,14 +895,9 @@ void FAvaSequencer::DuplicateSequences_Execute()
 		return;
 	}
 
-	UObject* const Outer = SequenceProvider->ToUObject();
-	if (!IsValid(Outer))
-	{
-		return;
-	}
-
 	const FScopedTransaction Transaction(LOCTEXT("DuplicateSequenceTransaction", "Duplicate Sequence"));
 
+	UObject* const Outer = SequenceProvider->ToUObject();
 	Outer->Modify();
 
 	for (const FAvaSequenceItemPtr& Item : SelectedItems)
@@ -943,24 +910,24 @@ void FAvaSequencer::DuplicateSequences_Execute()
 	}
 }
 
-bool FAvaSequencer::ExportSequences_CanExecute() const
+bool FAvaSequencer::ExportSequence_CanExecute() const
 {
 	return SequenceTreeView.IsValid()
 		&& !SequenceTreeView->GetSelectedItems().IsEmpty();
 }
 
-void FAvaSequencer::ExportSequences_Execute()
+void FAvaSequencer::ExportSequence_Execute()
 {
 	Provider.ExportSequences(GetSelectedSequences());
 }
 
-bool FAvaSequencer::SpawnPlayers_CanExecute() const
+bool FAvaSequencer::SpawnPlayer_CanExecute() const
 {
 	return SequenceTreeView.IsValid()
 		&& !SequenceTreeView->GetSelectedItems().IsEmpty();
 }
 
-void FAvaSequencer::SpawnPlayers_Execute()
+void FAvaSequencer::SpawnPlayer_Execute()
 {
 	if (!GEditor)
 	{
@@ -979,11 +946,7 @@ void FAvaSequencer::SpawnPlayers_Execute()
 		return;
 	}
 
-	TArray<UAvaSequence*> Sequences = GetSelectedSequences();
-	if (Sequences.Num() != 1)
-	{
-		return;
-	}
+	TArray<UAvaSequence*, TInlineAllocator<1>> Sequences = GetSelectedSequences();
 
 	FScopedTransaction Transaction(LOCTEXT("SpawnSequencePlayers", "Spawn Sequence Players"));
 
@@ -994,12 +957,12 @@ void FAvaSequencer::SpawnPlayers_Execute()
 	}
 }
 
-bool FAvaSequencer::DeleteSequences_CanExecute() const
+bool FAvaSequencer::DeleteSequence_CanExecute() const
 {
 	return GetProvider().CanEditOrPlaySequences();
 }
 
-void FAvaSequencer::DeleteSequences_Execute()
+void FAvaSequencer::DeleteSequence_Execute()
 {
 	IAvaSequenceProvider* const SequenceProvider = Provider.GetSequenceProvider();
 	if (!SequenceProvider)
@@ -1007,16 +970,12 @@ void FAvaSequencer::DeleteSequences_Execute()
 		return;
 	}
 
-	UObject* const Outer = SequenceProvider->ToUObject();
-	if (!IsValid(Outer))
-	{
-		return;
-	}
-
 	TArray<FAvaSequenceItemPtr> SelectedItems = SequenceTreeView->GetSelectedItems();
 
 	const FScopedTransaction Transaction(LOCTEXT("DeleteSequenceTransaction", "Delete Sequence"));
 
+	UObject* const Outer = SequenceProvider->ToUObject();
+	check(Outer);
 	Outer->Modify();
 
 	TArray<UAvaSequence*> RemovedSequences;
@@ -1064,13 +1023,13 @@ void FAvaSequencer::RelabelSequence_Execute()
 	SelectedItem->RequestRelabel();
 }
 
-bool FAvaSequencer::PlaySequences_CanExecute() const
+bool FAvaSequencer::PlaySelected_CanExecute() const
 {
 	const bool bCanEditSequences = GetProvider().CanEditOrPlaySequences();
 	return bCanEditSequences && SequenceTreeView->GetNumItemsSelected() > 0;
 }
 
-void FAvaSequencer::PlaySequences_Execute()
+void FAvaSequencer::PlaySelected_Execute()
 {
 	const TArray<FAvaSequenceItemPtr> SelectedItems = SequenceTreeView->GetSelectedItems();
 
@@ -1091,13 +1050,13 @@ void FAvaSequencer::PlaySequences_Execute()
 	}
 }
 
-bool FAvaSequencer::ContinueSequences_CanExecute() const
+bool FAvaSequencer::ContinueSelected_CanExecute() const
 {
 	const bool bCanEditSequences = GetProvider().CanEditOrPlaySequences();
 	return bCanEditSequences && SequenceTreeView->GetNumItemsSelected() > 0;
 }
 
-void FAvaSequencer::ContinueSequences_Execute()
+void FAvaSequencer::ContinueSelected_Execute()
 {
 	const TArray<FAvaSequenceItemPtr> SelectedItems = SequenceTreeView->GetSelectedItems();
 	
@@ -1115,12 +1074,12 @@ void FAvaSequencer::ContinueSequences_Execute()
 	}
 }
 
-bool FAvaSequencer::StopSequences_CanExecute() const
+bool FAvaSequencer::StopSelected_CanExecute() const
 {
 	return SequenceTreeView->GetNumItemsSelected() > 0;
 }
 
-void FAvaSequencer::StopSequences_Execute()
+void FAvaSequencer::StopSelected_Execute()
 {
 	const TArray<FAvaSequenceItemPtr> SelectedItems = SequenceTreeView->GetSelectedItems();
 
@@ -1138,302 +1097,6 @@ void FAvaSequencer::StopSequences_Execute()
 	}
 }
 
-bool FAvaSequencer::CanAddSequence() const
-{
-	return GetProvider().CanEditOrPlaySequences();
-}
-
-UAvaSequence* FAvaSequencer::AddSequence(UAvaSequence* const InParentSequence)
-{
-	IAvaSequenceProvider* const SequenceProvider = Provider.GetSequenceProvider();
-	if (!SequenceProvider)
-	{
-		return nullptr;
-	}
-
-	const FScopedTransaction Transaction(LOCTEXT("AddSequenceTransaction", "Add Sequence"));
-
-	UAvaSequence* const NewSequence = CreateSequence();
-
-	if (IsValid(InParentSequence))
-	{
-		InParentSequence->Modify();
-		InParentSequence->AddChild(NewSequence);
-	}
-	else
-	{
-		SequenceProvider->AddSequence(NewSequence);
-	}
-
-	OnSequenceAddedDelegate.Broadcast(NewSequence);
-
-	return NewSequence;
-}
-
-bool FAvaSequencer::CanDuplicateSequences() const
-{
-	return GetProvider().CanEditOrPlaySequences();
-}
-
-void FAvaSequencer::DuplicateSequences()
-{
-	const TArray<UAvaSequence*> SelectedSequences = GetSelectedSequences();
-	if (SelectedSequences.IsEmpty())
-	{
-		return;
-	}
-
-	IAvaSequenceProvider* const SequenceProvider = Provider.GetSequenceProvider();
-	if (!SequenceProvider)
-	{
-		return;
-	}
-
-	UObject* const Outer = SequenceProvider->ToUObject();
-	if (!IsValid(Outer))
-	{
-		return;
-	}
-
-	const FScopedTransaction Transaction(LOCTEXT("DuplicateSequenceTransaction", "Duplicate Sequence"));
-
-	Outer->Modify();
-
-	for (UAvaSequence* const TemplateSequence : SelectedSequences)
-	{
-		check(IsValid(TemplateSequence));
-
-		UAvaSequence* const DupedSequence = DuplicateObject<UAvaSequence>(TemplateSequence, Outer);
-		SequenceProvider->AddSequence(DupedSequence);
-	}
-}
-
-bool FAvaSequencer::CanExportSequences() const
-{
-	return GetSelectedSequenceItems().Num() > 0;
-}
-
-void FAvaSequencer::ExportSequences()
-{
-	const TArray<UAvaSequence*> SelectedSequences = GetSelectedSequences();
-
-	TArray<TWeakObjectPtr<>> WeakSelectedObjects;
-	Algo::Transform(SelectedSequences, WeakSelectedObjects, [](UAvaSequence* const InSequence)
-		{
-			return InSequence;
-		});
-
-	Provider.ExportSequences(SelectedSequences);
-
-	FNotificationInfo Info(LOCTEXT("ExportSuccess", "Sequence Exported Successfully!"));
-	Info.bUseLargeFont = false;
-	Info.bUseSuccessFailIcons = false;
-	Info.ExpireDuration = 5.0f;
-	Info.HyperlinkText = LOCTEXT("ShowNewAssetsInContentBrowser", "Show in content browser");
-	Info.Hyperlink = FSimpleDelegate::CreateLambda([WeakSelectedObjects]()
-		{
-			TArray<UObject*> SyncObjects;
-
-			for (const TWeakObjectPtr<> WeakObject : WeakSelectedObjects)
-			{
-				UObject* const SyncObject = WeakObject.Get();
-				if (IsValid(SyncObject))
-				{
-					SyncObjects.Add(SyncObject);
-				}
-			}
-
-			FContentBrowserModule& ContentBrowserModule = FModuleManager::LoadModuleChecked<FContentBrowserModule>(TEXT("ContentBrowser"));
-			ContentBrowserModule.Get().SyncBrowserToAssets(SyncObjects);
-		});
-
-	FSlateNotificationManager::Get().AddNotification(Info);
-}
-
-bool FAvaSequencer::CanSpawnPlayers() const
-{
-	return GetSelectedSequenceItems().Num() > 0;
-}
-
-void FAvaSequencer::SpawnPlayers()
-{
-	if (!GEditor)
-	{
-		return;
-	}
-
-	UWorld* World = Provider.GetPlaybackContext()->GetWorld();
-	if (!World)
-	{
-		return;
-	}
-
-	UActorFactory* ActorFactory = GEditor->FindActorFactoryForActorClass(AAvaSequenceActor::StaticClass());
-	if (!ensure(ActorFactory))
-	{
-		return;
-	}
-
-	TArray<UAvaSequence*> Sequences = GetSelectedSequences();
-
-	FScopedTransaction Transaction(LOCTEXT("SpawnSequencePlayers", "Spawn Sequence Players"));
-
-	for (UAvaSequence* Sequence : Sequences)
-	{
-		check(Sequence);
-		GEditor->UseActorFactory(ActorFactory, FAssetData(Sequence), &FTransform::Identity);
-	}
-
-	const FText NotifyText = FText::Format(LOCTEXT("SpawnPlayerSuccess", "{0} Sequence Players Spawned Successfully!")
-		, FText::FromString(FString::FromInt(Sequences.Num())));
-
-	FNotificationInfo Info(NotifyText);
-	Info.bUseLargeFont = false;
-	Info.bUseSuccessFailIcons = false;
-	FSlateNotificationManager::Get().AddNotification(Info);
-}
-
-bool FAvaSequencer::CanDeleteSequence() const
-{
-	return GetProvider().CanEditOrPlaySequences();
-}
-
-void FAvaSequencer::DeleteSequence()
-{
-	DeleteSequences(TSet<UAvaSequence*> { GetSelectedSequences() });
-}
-
-void FAvaSequencer::DeleteSequences(const TSet<UAvaSequence*>& InSequences)
-{
-	if (InSequences.IsEmpty())
-	{
-		return;
-	}
-
-	IAvaSequenceProvider* const SequenceProvider = Provider.GetSequenceProvider();
-	if (!SequenceProvider)
-	{
-		return;
-	}
-
-	UObject* const Outer = SequenceProvider->ToUObject();
-	if (!IsValid(Outer))
-	{
-		return;
-	}
-
-	const FScopedTransaction Transaction(LOCTEXT("DeleteSequencesTransaction", "Delete Sequence(s)"));
-
-	Outer->Modify();
-
-	TArray<UAvaSequence*> RemovedSequences;
-	RemovedSequences.Reserve(InSequences.Num());
-
-	// Remove the sequences from the list (not marked as garbage yet)
-	for (UAvaSequence* const Sequence : InSequences)
-	{
-		check(IsValid(Sequence));
-
-		Sequence->Modify();
-		SequenceProvider->RemoveSequence(Sequence);
-
-		RemovedSequences.Add(Sequence);
-
-		OnSequenceRemovedDelegate.Broadcast(Sequence);
-	}
-
-	// Set the Viewed Sequence to the Default one
-	SetViewedSequence(GetDefaultSequence());
-
-	// Once a new viewed sequence is set, the removed sequences can now be marked as garbage 
-	for (UAvaSequence* Sequence : RemovedSequences)
-	{
-		Sequence->OnSequenceRemoved();
-	}
-
-	if (UBlueprint* const Blueprint = Cast<UBlueprint>(Outer))
-	{
-		FBlueprintEditorUtils::MarkBlueprintAsStructurallyModified(Blueprint);
-	}
-}
-
-bool FAvaSequencer::CanRelabelSequence() const
-{
-	const bool bCanEditSequences = GetProvider().CanEditOrPlaySequences();
-	return bCanEditSequences && SequenceTreeView->GetNumItemsSelected() == 1;
-}
-
-void FAvaSequencer::RelabelSequence()
-{
-	TArray<FAvaSequenceItemPtr> SelectedItems = SequenceTreeView->GetSelectedItems();
-	check(SelectedItems.Num() == 1);
-
-	const FAvaSequenceItemPtr& SelectedItem = SelectedItems[0];
-	SelectedItem->RequestRelabel();
-}
-
-bool FAvaSequencer::CanPlaySequences() const
-{
-	const bool bCanEditSequences = GetProvider().CanEditOrPlaySequences();
-	return bCanEditSequences && GetSelectedSequenceItems().Num() > 0;
-}
-
-void FAvaSequencer::PlaySequences()
-{
-	const TArray<UAvaSequence*> SelectedSequences = GetSelectedSequences();
-
-	FAvaSequencePlayParams PlaySettings;
-	PlaySettings.AdvancedSettings.bRestoreState = true;
-
-	for (UAvaSequence* const Sequence : SelectedSequences)
-	{
-		IAvaSequencePlaybackObject* const PlaybackObject = Provider.GetPlaybackObject();
-		if (IsValid(Sequence) && PlaybackObject)
-		{
-			PlaybackObject->PlaySequence(Sequence, PlaySettings);
-		}
-	}
-}
-
-bool FAvaSequencer::CanContinueSequences() const
-{
-	const bool bCanEditSequences = GetProvider().CanEditOrPlaySequences();
-	return bCanEditSequences && GetSelectedSequenceItems().Num() > 0;
-}
-
-void FAvaSequencer::ContinueSequences()
-{
-	const TArray<UAvaSequence*> SelectedSequences = GetSelectedSequences();
-
-	for (UAvaSequence* const Sequence : SelectedSequences)
-	{
-		IAvaSequencePlaybackObject* const PlaybackObject = Provider.GetPlaybackObject();
-		if (IsValid(Sequence) && PlaybackObject)
-		{
-			PlaybackObject->ContinueSequence(Sequence);
-		}
-	}
-}
-
-bool FAvaSequencer::CanStopSequences() const
-{
-	return GetSelectedSequenceItems().Num() > 0;
-}
-
-void FAvaSequencer::StopSequences()
-{
-	const TArray<UAvaSequence*> SelectedSequences = GetSelectedSequences();
-
-	for (UAvaSequence* const Sequence : SelectedSequences)
-	{
-		IAvaSequencePlaybackObject* const PlaybackObject = Provider.GetPlaybackObject();
-		if (IsValid(Sequence) && PlaybackObject)
-		{
-			PlaybackObject->StopSequence(Sequence);
-		}
-	}
-}
-
 UObject* FAvaSequencer::FindObjectToPossess(UObject* InParentObject, const FMovieScenePossessable& InPossessable)
 {
 	if (!IsValid(InParentObject))
@@ -1856,31 +1519,48 @@ TSharedRef<SWidget> FAvaSequencer::CreateSequenceWidget()
 		SequencerWeak.Reset();
 	}
 
-	const TSharedRef<ISequencer> Sequencer = GetSequencer();
+	TSharedRef<ISequencer> Sequencer = GetSequencer();
 
 	UAvaSequencerSettings* const SequencerSettings = GetMutableDefault<UAvaSequencerSettings>();
 	check(IsValid(SequencerSettings));
 	FSidebarState& SidebarState = SequencerSettings->GetSidebarState();
 
-	if (SidebarState.IsVisible())
-	{
-		FSidebarDrawerConfig MotionDesignDrawerConfig;
-		MotionDesignDrawerConfig.UniqueId = SidebarDrawerId;
-		MotionDesignDrawerConfig.ButtonText = LOCTEXT("MotionDesignLabel", "Motion Design");
-		MotionDesignDrawerConfig.ToolTipText = LOCTEXT("SequenceTooltip", "Open the Sequence options panel");
-		MotionDesignDrawerConfig.Icon = FAvaEditorCoreStyle::Get().GetBrush(TEXT("Icons.MotionDesign"));
-		MotionDesignDrawerConfig.InitialState = SidebarState.FindOrAddDrawerState(SidebarDrawerId);
-		Sequencer->RegisterDrawer(MoveTemp(MotionDesignDrawerConfig));
+	// Make sure the sequence tree widget is created
+	GetSequenceTreeWidget();
 
-		const TSharedRef<FAvaSequencer> ThisSequencerRef = SharedThis(this);
-		Sequencer->RegisterDrawerSection(SidebarDrawerId, MakeShared<FAvaSequenceTreeDetails>(ThisSequencerRef));
-		Sequencer->RegisterDrawerSection(SidebarDrawerId, MakeShared<FAvaSequencePlaybackDetails>(ThisSequencerRef));
-		Sequencer->RegisterDrawerSection(SidebarDrawerId, MakeShared<FAvaSequenceSettingsDetails>(ThisSequencerRef));
+	if (!SidebarState.IsVisible())
+	{
+		return Sequencer->GetSequencerWidget();
 	}
 
-	RegisterNavigationTool();
+	TSharedRef<SSidebarContainer> SidebarContainer = SNew(SSidebarContainer);
+
+	TSharedRef<SSidebar> LeftSidebar = SNew(SSidebar, SidebarContainer)
+		.TabLocation(ESidebarTabLocation::Left)
+		.InitialDrawerSize(SidebarState.GetDrawerSize())
+		.OnStateChanged(this, &FAvaSequencer::OnSidebarStateChanged)
+		.OnGetContent(FOnGetContent::CreateLambda([this]()
+		{
+			return GetSequencer()->GetSequencerWidget();
+		}));
+
+	SidebarContainer->RebuildSidebar(LeftSidebar, SidebarState);
+
+	FSidebarDrawerConfig SequenceTreeDrawerConfig;
+	SequenceTreeDrawerConfig.UniqueId = SidebarDrawerId;
+	SequenceTreeDrawerConfig.ButtonText = LOCTEXT("SequenceLabel", "Sequence");
+	SequenceTreeDrawerConfig.ToolTipText = LOCTEXT("SequenceTooltip", "Open the Sequence options panel");
+	SequenceTreeDrawerConfig.Icon = FAppStyle::GetBrush(TEXT("Persona.EditInSequencer"));
+	SequenceTreeDrawerConfig.InitialState = SidebarState.FindOrAddDrawerState(SidebarDrawerId);
 
-	return Sequencer->GetSequencerWidget();
+	LeftSidebar->RegisterDrawer(MoveTemp(SequenceTreeDrawerConfig));
+
+	const TSharedRef<FAvaSequencer> ThisSequencerRef = SharedThis(this);
+	LeftSidebar->RegisterDrawerSection(SidebarDrawerId, MakeShared<FAvaSequenceTreeDetails>(ThisSequencerRef));
+	LeftSidebar->RegisterDrawerSection(SidebarDrawerId, MakeShared<FAvaSequencePlaybackDetails>(ThisSequencerRef));
+	LeftSidebar->RegisterDrawerSection(SidebarDrawerId, MakeShared<FAvaSequenceSettingsDetails>(ThisSequencerRef));
+
+	return SidebarContainer;
 }
 
 void FAvaSequencer::OnActorsCopied(FString& InOutCopiedData, TConstArrayView<AActor*> InCopiedActors)
@@ -1951,13 +1631,6 @@ void FAvaSequencer::OnEditorSelectionChanged(const FAvaEditorSelection& InEditor
 
 void FAvaSequencer::NotifyOnSequenceTreeChanged()
 {
-	using namespace UE::Sequencer::NavigationTool;
-
-	if (const TSharedPtr<INavigationTool> NavigationTool = GetSequencer()->GetNavigationTool())
-	{
-		NavigationTool->RequestRefresh();
-	}
-
 	IAvaSequenceProvider* const SequenceProvider = Provider.GetSequenceProvider();
 	if (!SequenceTreeView.IsValid() || !SequenceProvider)
 	{
@@ -2028,76 +1701,42 @@ void FAvaSequencer::NotifyOnSequenceTreeChanged()
 	SequenceTreeView->RequestTreeRefresh();
 }
 
-void FAvaSequencer::PostUndo(const bool bInSuccess)
+void FAvaSequencer::PostUndo(bool bSuccess)
 {
 	// A just-added sequence might be removed due to this undo, so refresh 
 	NotifyOnSequenceTreeChanged();
 }
 
-void FAvaSequencer::PostRedo(const bool bInSuccess)
+TArray<UAvaSequence*, TInlineAllocator<1>> FAvaSequencer::GetSelectedSequences() const
 {
-	PostUndo(bInSuccess);
-}
-
-TArray<TSharedPtr<UE::Sequencer::NavigationTool::FNavigationToolAvaSequence>> FAvaSequencer::GetSelectedSequenceItems() const
-{
-	using namespace UE::Sequencer::NavigationTool;
-
-	const TSharedPtr<ISequencer> Sequencer = SequencerWeak.Pin();
-	if (!Sequencer.IsValid())
-	{
-		return {};
-	}
-
-	const TSharedPtr<INavigationTool> Tool = Sequencer->GetNavigationTool();
-	if (!Tool.IsValid())
+	if (!SequenceTreeView.IsValid())
 	{
 		return {};
 	}
 
-	const TArray<FNavigationToolItemPtr> SelectedItems = Tool->GetSelectedItems();
+	const TArray<FAvaSequenceItemPtr> SelectedItems = SequenceTreeView->GetSelectedItems();
 	if (SelectedItems.IsEmpty())
 	{
 		return {};
 	}
 
-	TArray<TSharedPtr<FNavigationToolAvaSequence>> SelectedSequenceItems;
-	SelectedSequenceItems.Reserve(SelectedItems.Num());
+	TArray<UAvaSequence*, TInlineAllocator<1>> SelectedSequences;
+	SelectedSequences.Reserve(SelectedItems.Num());
 
-	for (const FNavigationToolItemPtr& Item : SelectedItems)
+	for (const FAvaSequenceItemPtr& Item : SelectedItems)
 	{
-		if (FNavigationToolAvaSequence* const AvaSequenceItem = Item->CastTo<FNavigationToolAvaSequence>())
+		if (!Item.IsValid())
 		{
-			SelectedSequenceItems.Add(SharedThis(AvaSequenceItem));
+			continue;
 		}
-	}
-
-	return SelectedSequenceItems;
-}
-
-TArray<UAvaSequence*> FAvaSequencer::GetSelectedSequences() const
-{
-	using namespace UE::Sequencer::NavigationTool;
-
-	const TArray<TSharedPtr<FNavigationToolAvaSequence>> SelectedItems = GetSelectedSequenceItems();
-	if (SelectedItems.IsEmpty())
-	{
-		return {};
-	}
-
-	TArray<UAvaSequence*> SelectedAvaSequences;
-	SelectedAvaSequences.Reserve(SelectedItems.Num());
 
-	for (const TSharedPtr<FNavigationToolAvaSequence>& SequenceItem : SelectedItems)
-	{
-		UAvaSequence* const AvaSequence = SequenceItem->GetAvaSequence();
-		if (IsValid(AvaSequence))
+		if (UAvaSequence* Sequence = Item->GetSequence())
 		{
-			SelectedAvaSequences.Add(AvaSequence);
+			SelectedSequences.Add(Sequence);
 		}
 	}
 
-	return SelectedAvaSequences;
+	return SelectedSequences;
 }
 
 TSharedPtr<UE::Sequencer::SOutlinerView> FAvaSequencer::GetOutlinerView() const
@@ -2270,64 +1909,4 @@ void FAvaSequencer::OnSequencerClosed(const TSharedRef<ISequencer> InSequencer)
 	}
 }
 
-const TArray<TWeakObjectPtr<UAvaSequence>>& FAvaSequencer::GetRootSequences() const
-{
-	IAvaSequenceProvider* const SequenceProvider = Provider.GetSequenceProvider();
-	check(SequenceProvider);
-	return SequenceProvider->GetRootSequences();
-}
-
-void FAvaSequencer::RegisterNavigationTool()
-{
-	using namespace UE::Sequencer::NavigationTool;
-
-	const TSharedPtr<ISequencer> Sequencer = SequencerWeak.Pin();
-	if (!Sequencer.IsValid())
-	{
-		return;
-	}
-
-	const TSharedPtr<INavigationTool> Tool = Sequencer->GetNavigationTool();
-	if (!Tool.IsValid())
-	{
-		return;
-	}
-
-	const FName ToolId = FNavigationToolExtender::GetToolInstanceId(*Tool);
-
-	const TSharedPtr<FNavigationToolProvider> FoundProvider = FNavigationToolExtender::FindToolProvider(ToolId, UE::AvaSequencer::FAvaNavigationToolProvider::Identifier);
-	if (!FoundProvider.IsValid())
-	{
-		NavigationToolProvider = MakeShared<UE::AvaSequencer::FAvaNavigationToolProvider>(SharedThis(this));
-		FNavigationToolExtender::RegisterToolProvider(ToolId, NavigationToolProvider.ToSharedRef());
-	}
-}
-
-void FAvaSequencer::UnregisterNavigationTool()
-{
-	using namespace UE::Sequencer::NavigationTool;
-
-	if (!NavigationToolProvider.IsValid())
-	{
-		return;
-	}
-
-	const TSharedPtr<ISequencer> Sequencer = SequencerWeak.Pin();
-	if (!Sequencer.IsValid())
-	{
-		return;
-	}
-
-	const TSharedPtr<INavigationTool> Tool = Sequencer->GetNavigationTool();
-	if (!Tool.IsValid())
-	{
-		return;
-	}
-
-	const FName ToolId = FNavigationToolExtender::GetToolInstanceId(*Tool);
-	FNavigationToolExtender::UnregisterToolProvider(ToolId, NavigationToolProvider->GetIdentifier());
-
-	NavigationToolProvider.Reset();
-}
-
 #undef LOCTEXT_NAMESPACE
diff --git a/Engine/Plugins/Experimental/Avalanche/Source/AvalancheSequencer/Private/AvaSequencer.h b/Engine/Plugins/Experimental/Avalanche/Source/AvalancheSequencer/Private/AvaSequencer.h
--- a/Engine/Plugins/Experimental/Avalanche/Source/AvalancheSequencer/Private/AvaSequencer.h
+++ b/Engine/Plugins/Experimental/Avalanche/Source/AvalancheSequencer/Private/AvaSequencer.h
@@ -9,50 +9,49 @@
 #include "IAvaSequencerProvider.h"
 #include "MovieSceneSequenceID.h"
 #include "Widgets/Layout/SSplitter.h"
-#include "Widgets/Views/STreeView.h"
 
 class AActor;
 class FAvaEaseCurveTool;
 class FAvaEditorSelection;
-class FAvaNavigationToolProvider;
 class FAvaSequencerAction;
 class FAvaSequencerCleanView;
+class FAvaSequencerController;
+class FAvaStaggerTool;
 class FExtender;
 class FMenuBuilder;
+class FProperty;
 class FUICommandList;
 class IAvaSequenceColumn;
 class IAvaSequencerController;
+class IPropertyHandle;
 class ISequencer;
+class ISequencerObjectChangeListener;
 class SAvaSequenceTree;
+class SBox;
 class SHeaderRow;
+class SSidebar;
+class SSidebarContainer;
 class SWidget;
 class UAvaSequence;
+class UAvaSequencerSettings;
+class UMovieScene3DTransformTrack;
 class UMovieScene;
+class UMovieSceneTrack;
 class USequencerSettings;
+class UToolMenu;
 enum class EMovieSceneTransformChannel : uint32;
 struct FAvaSequencePreset;
 struct FAvaSequencerArgs;
 struct FMovieSceneBinding;
 struct FMovieScenePossessable;
+struct FSequencerInitParams;
 struct FSidebarState;
+struct FToolMenuContext;
+template<typename ItemType> class STreeView;
 
-namespace UE
+namespace UE::Sequencer
 {
-	namespace Sequencer
-	{
-		class SOutlinerView;
-
-		namespace NavigationTool
-		{
-			class FNavigationToolAvaSequence;
-			class FNavigationToolProvider;
-		}
-	}
-
-	namespace AvaSequencer
-	{
-		class FAvaNavigationToolProvider;
-	}
+	class SOutlinerView;
 }
 
 class FAvaSequencer : public IAvaSequencer, public FEditorUndoClient, public TSharedFromThis<FAvaSequencer>
@@ -122,15 +121,39 @@ public:
 	DECLARE_MULTICAST_DELEGATE_OneParam(FOnViewedSequenceChanged, UAvaSequence*);
 	FOnViewedSequenceChanged& GetOnViewedSequenceChanged() { return OnViewedSequenceChanged; }
 
-	FOnSequenceAdded& OnSequenceAdded() override { return OnSequenceAddedDelegate; }
-	FOnSequenceRemoved& OnSequenceRemoved() override { return OnSequenceRemovedDelegate; }
-
 	void AddTransformKey(EMovieSceneTransformChannel InTransformChannel);
 
 	void ApplyDefaultPresetToSelection(FName InPresetName);
 	void ApplyCustomPresetToSelection(FName InPresetName);
 	void ApplyPresetToSelection(const FAvaSequencePreset& InPreset);
 
+	bool AddSequence_CanExecute() const;
+	void AddSequence_Execute();
+
+	bool DuplicateSequence_CanExecute() const;
+	void DuplicateSequence_Execute();
+
+	bool ExportSequence_CanExecute() const;
+	void ExportSequence_Execute();
+
+	bool SpawnPlayer_CanExecute() const;
+	void SpawnPlayer_Execute();
+
+	bool DeleteSequence_CanExecute() const;
+	void DeleteSequence_Execute();
+
+	bool RelabelSequence_CanExecute() const;
+	void RelabelSequence_Execute();
+
+	bool PlaySelected_CanExecute() const;
+	void PlaySelected_Execute();
+
+	bool ContinueSelected_CanExecute() const;
+	void ContinueSelected_Execute();
+
+	bool StopSelected_CanExecute() const;
+	void StopSelected_Execute();
+
 	/**
  	 * Attempts to find the best matching subobject of the given Parent Object
  	 * based on the information provided by the given possessable
@@ -174,10 +197,9 @@ public:
 
 	TSharedRef<FAvaEaseCurveTool> GetEaseCurveTool() const;
 	
-	TArrayView<TWeakObjectPtr<>> ResolveBoundObjects(const FGuid& InBindingId, UMovieSceneSequence* Sequence) const;
+	TArrayView<TWeakObjectPtr<>> ResolveBoundObjects(const FGuid& InBindingId, class UMovieSceneSequence* Sequence) const;
 
 	//~ Begin IAvaSequencer
-
 	virtual const IAvaSequencerProvider& GetProvider() const override { return Provider; }
 	virtual TSharedRef<ISequencer> GetSequencer() const override;
 	virtual USequencerSettings* GetSequencerSettings() const override;
@@ -191,74 +213,15 @@ public:
 	virtual void OnActorsPasted(FStringView InPastedData, const TMap<FName, AActor*>& InPastedActors) override;
 	virtual void OnEditorSelectionChanged(const FAvaEditorSelection& InEditorSelection) override;
 	virtual void NotifyOnSequenceTreeChanged() override;
-	virtual const TArray<TWeakObjectPtr<UAvaSequence>>& GetRootSequences() const override;
-
 	//~ End IAvaSequencer
 
-	virtual bool CanAddSequence() const;
-	virtual UAvaSequence* AddSequence(UAvaSequence* const InParentSequence = nullptr);
-
-	virtual bool CanDuplicateSequences() const;
-	virtual void DuplicateSequences();
-
-	virtual bool CanExportSequences() const;
-	virtual void ExportSequences();
-
-	virtual bool CanSpawnPlayers() const;
-	virtual void SpawnPlayers();
-
-	virtual bool CanDeleteSequence() const;
-	virtual void DeleteSequence();
-	virtual void DeleteSequences(const TSet<UAvaSequence*>& InSequences);
-
-	virtual bool CanRelabelSequence() const;
-	virtual void RelabelSequence();
-
-	virtual bool CanPlaySequences() const;
-	virtual void PlaySequences();
-
-	virtual bool CanContinueSequences() const;
-	virtual void ContinueSequences();
-
-	virtual bool CanStopSequences() const;
-	virtual void StopSequences();
-
-	bool AddSequence_CanExecute() const;
-	void AddSequence_Execute();
-
-	bool DuplicateSequences_CanExecute() const;
-	void DuplicateSequences_Execute();
-
-	bool ExportSequences_CanExecute() const;
-	void ExportSequences_Execute();
-
-	bool SpawnPlayers_CanExecute() const;
-	void SpawnPlayers_Execute();
-
-	bool DeleteSequences_CanExecute() const;
-	void DeleteSequences_Execute();
-
-	bool RelabelSequence_CanExecute() const;
-	void RelabelSequence_Execute();
-
-	bool PlaySequences_CanExecute() const;
-	void PlaySequences_Execute();
-
-	bool ContinueSequences_CanExecute() const;
-	void ContinueSequences_Execute();
-
-	bool StopSequences_CanExecute() const;
-	void StopSequences_Execute();
-
-	//~ Begin FEditorUndoClient
-	virtual void PostUndo(bool bInSuccess) override;
-	virtual void PostRedo(bool bInSuccess) override;
+	//~ Begin FEditorUndoClient Interface
+	virtual void PostUndo(bool bSuccess) override;
+	virtual void PostRedo(bool bSuccess) override { PostUndo(bSuccess); }
 	// End of FEditorUndoClient
 
 private:
-	TArray<TSharedPtr<UE::Sequencer::NavigationTool::FNavigationToolAvaSequence>> GetSelectedSequenceItems() const;
-
-	TArray<UAvaSequence*> GetSelectedSequences() const;
+	TArray<UAvaSequence*, TInlineAllocator<1>> GetSelectedSequences() const;
 
 	TSharedPtr<UE::Sequencer::SOutlinerView> GetOutlinerView() const;
 
@@ -275,14 +238,8 @@ private:
 
 	void OnSequencerClosed(const TSharedRef<ISequencer> InSequencer);
 
-	void RegisterNavigationTool();
-	void UnregisterNavigationTool();
-
 	FOnViewedSequenceChanged OnViewedSequenceChanged;
 
-	FOnSequenceAdded OnSequenceAddedDelegate;
-	FOnSequenceRemoved OnSequenceRemovedDelegate;
-
 	IAvaSequencerProvider& Provider;
 
 	/** The active sequencer to use. This is set to Weak as FSequencer::Tick does not allow to have more than 1 ref (outside of menus) */
@@ -301,7 +258,6 @@ private:
 
 	FDelegateHandle SequencerAddTrackExtenderHandle;
 
-	// BEGIN REMOVE AFTER NAVIGATION TOOL INTEGRATION
 	TMap<FName, TSharedPtr<IAvaSequenceColumn>> SequenceColumns;
 
 	TArray<FAvaSequenceItemPtr> RootSequenceItems;
@@ -311,7 +267,6 @@ private:
 	TSharedPtr<SAvaSequenceTree> SequenceTree;
 
 	TSharedPtr<STreeView<FAvaSequenceItemPtr>> SequenceTreeView;
-	// END REMOVE AFTER NAVIGATION TOOL INTEGRATION
 
 	TSharedRef<FAvaSequencerCleanView> CleanView;
 
@@ -337,6 +292,4 @@ private:
 	TSharedPtr<FExtender> SidebarExtender;
 
 	FDelegateHandle SidebarSelectionExtenderHandle;
-
-	TSharedPtr<UE::AvaSequencer::FAvaNavigationToolProvider> NavigationToolProvider;
 };
diff --git a/Engine/Plugins/Experimental/Avalanche/Source/AvalancheSequencer/Private/AvaSequencerUtils.cpp b/Engine/Plugins/Experimental/Avalanche/Source/AvalancheSequencer/Private/AvaSequencerUtils.cpp
--- a/Engine/Plugins/Experimental/Avalanche/Source/AvalancheSequencer/Private/AvaSequencerUtils.cpp
+++ b/Engine/Plugins/Experimental/Avalanche/Source/AvalancheSequencer/Private/AvaSequencerUtils.cpp
@@ -1,74 +1,10 @@
 // Copyright Epic Games, Inc. All Rights Reserved.
 
 #include "AvaSequencerUtils.h"
-#include "AvaSceneSubsystem.h"
-#include "AvaSequencerSubsystem.h"
-#include "Engine/World.h"
-#include "IAvaSceneInterface.h"
-#include "IAvaSequencer.h"
 #include "Playback/AvaSequencerController.h"
 #include "Templates/SharedPointer.h"
-#include "Toolkits/IToolkitHost.h"
 
 TSharedRef<IAvaSequencerController> FAvaSequencerUtils::CreateSequencerController()
 {
 	return MakeShared<FAvaSequencerController>();
 }
-
-UWorld* FAvaSequencerUtils::GetSequencerWorld(const TSharedRef<ISequencer>& InSequencer)
-{
-	if (const TSharedPtr<IToolkitHost> ToolkitHost = InSequencer->GetToolkitHost())
-	{
-		return ToolkitHost->GetWorld();
-	}
-	return nullptr;
-}
-
-UAvaSequencerSubsystem* FAvaSequencerUtils::GetSequencerSubsystem(const TSharedRef<ISequencer>& InSequencer)
-{
-	const UWorld* const World = GetSequencerWorld(InSequencer);
-	if (IsValid(World))
-	{
-		return World->GetSubsystem<UAvaSequencerSubsystem>();
-	}
-	return nullptr;
-}
-
-UAvaSceneSubsystem* FAvaSequencerUtils::GetSceneSubsystem(const TSharedRef<ISequencer>& InSequencer)
-{
-	const UWorld* const World = GetSequencerWorld(InSequencer);
-	if (IsValid(World))
-	{
-		return World->GetSubsystem<UAvaSceneSubsystem>();
-	}
-	return nullptr;
-}
-
-IAvaSceneInterface* FAvaSequencerUtils::GetSceneInterface(const TSharedRef<ISequencer>& InSequencer)
-{
-	UAvaSceneSubsystem* const SceneSubsystem = GetSceneSubsystem(InSequencer);
-	if (IsValid(SceneSubsystem))
-	{
-		return SceneSubsystem->GetSceneInterface();
-	}
-	return nullptr;
-}
-
-IAvaSequenceProvider* FAvaSequencerUtils::GetSequenceProvider(const TSharedRef<ISequencer>& InSequencer)
-{
-	if (IAvaSceneInterface* const SceneInterface = GetSceneInterface(InSequencer))
-	{
-		return SceneInterface->GetSequenceProvider();
-	}
-	return nullptr;
-}
-
-TSharedPtr<IAvaSequencer> FAvaSequencerUtils::GetAvaSequencer(const TSharedRef<ISequencer>& InSequencer)
-{
-	UAvaSequencerSubsystem* const SequencerSubsystem = GetSequencerSubsystem(InSequencer);
-	if (IsValid(SequencerSubsystem))
-	{
-		return SequencerSubsystem->GetSequencer();
-	}
-	return nullptr;
-}
diff --git a/Engine/Plugins/Experimental/Avalanche/Source/AvalancheSequencer/Private/Commands/AvaSequencerCommands.cpp b/Engine/Plugins/Experimental/Avalanche/Source/AvalancheSequencer/Private/Commands/AvaSequencerCommands.cpp
--- a/Engine/Plugins/Experimental/Avalanche/Source/AvalancheSequencer/Private/Commands/AvaSequencerCommands.cpp
+++ b/Engine/Plugins/Experimental/Avalanche/Source/AvalancheSequencer/Private/Commands/AvaSequencerCommands.cpp
@@ -6,12 +6,6 @@
 
 void FAvaSequencerCommands::RegisterCommands()
 {
-	UI_COMMAND(AddNew
-		, "Add New"
-		, "Adds a new embedded Sequence to the level"
-		, EUserInterfaceActionType::Button
-		, FInputChord());
-
 	UI_COMMAND(PlaySelected
 		, "Play"
 		, "Plays all selected Sequences"
diff --git a/Engine/Plugins/Experimental/Avalanche/Source/AvalancheSequencer/Private/Commands/AvaSequencerCommands.h b/Engine/Plugins/Experimental/Avalanche/Source/AvalancheSequencer/Private/Commands/AvaSequencerCommands.h
--- a/Engine/Plugins/Experimental/Avalanche/Source/AvalancheSequencer/Private/Commands/AvaSequencerCommands.h
+++ b/Engine/Plugins/Experimental/Avalanche/Source/AvalancheSequencer/Private/Commands/AvaSequencerCommands.h
@@ -16,9 +16,7 @@ public:
 	{
 	}
 
-	virtual void RegisterCommands() override;
-
-	TSharedPtr<FUICommandInfo> AddNew;
+	virtual void RegisterCommands() override;	
 
 	TSharedPtr<FUICommandInfo> PlaySelected;
 
diff --git a/Engine/Plugins/Experimental/Avalanche/Source/AvalancheSequencer/Private/NavigationTool/AvaNavigationToolProvider.cpp b/Engine/Plugins/Experimental/Avalanche/Source/AvalancheSequencer/Private/NavigationTool/AvaNavigationToolProvider.cpp
deleted file mode 100644
diff --git a/Engine/Plugins/Experimental/Avalanche/Source/AvalancheSequencer/Private/NavigationTool/AvaNavigationToolProvider.h b/Engine/Plugins/Experimental/Avalanche/Source/AvalancheSequencer/Private/NavigationTool/AvaNavigationToolProvider.h
deleted file mode 100644
diff --git a/Engine/Plugins/Experimental/Avalanche/Source/AvalancheSequencer/Private/NavigationTool/Columns/AvaNavigationToolStatusColumn.cpp b/Engine/Plugins/Experimental/Avalanche/Source/AvalancheSequencer/Private/NavigationTool/Columns/AvaNavigationToolStatusColumn.cpp
deleted file mode 100644
diff --git a/Engine/Plugins/Experimental/Avalanche/Source/AvalancheSequencer/Private/NavigationTool/Columns/AvaNavigationToolStatusColumn.h b/Engine/Plugins/Experimental/Avalanche/Source/AvalancheSequencer/Private/NavigationTool/Columns/AvaNavigationToolStatusColumn.h
deleted file mode 100644
diff --git a/Engine/Plugins/Experimental/Avalanche/Source/AvalancheSequencer/Private/NavigationTool/Columns/SAvaNavigationToolStatus.cpp b/Engine/Plugins/Experimental/Avalanche/Source/AvalancheSequencer/Private/NavigationTool/Columns/SAvaNavigationToolStatus.cpp
deleted file mode 100644
diff --git a/Engine/Plugins/Experimental/Avalanche/Source/AvalancheSequencer/Private/NavigationTool/Columns/SAvaNavigationToolStatus.h b/Engine/Plugins/Experimental/Avalanche/Source/AvalancheSequencer/Private/NavigationTool/Columns/SAvaNavigationToolStatus.h
deleted file mode 100644
diff --git a/Engine/Plugins/Experimental/Avalanche/Source/AvalancheSequencer/Private/NavigationTool/DragDrop/NavigationToolAvaSequenceDropHandler.cpp b/Engine/Plugins/Experimental/Avalanche/Source/AvalancheSequencer/Private/NavigationTool/DragDrop/NavigationToolAvaSequenceDropHandler.cpp
deleted file mode 100644
diff --git a/Engine/Plugins/Experimental/Avalanche/Source/AvalancheSequencer/Private/NavigationTool/DragDrop/NavigationToolAvaSequenceDropHandler.h b/Engine/Plugins/Experimental/Avalanche/Source/AvalancheSequencer/Private/NavigationTool/DragDrop/NavigationToolAvaSequenceDropHandler.h
deleted file mode 100644
diff --git a/Engine/Plugins/Experimental/Avalanche/Source/AvalancheSequencer/Private/NavigationTool/NavigationToolAvaSequence.cpp b/Engine/Plugins/Experimental/Avalanche/Source/AvalancheSequencer/Private/NavigationTool/NavigationToolAvaSequence.cpp
deleted file mode 100644
diff --git a/Engine/Plugins/Experimental/Avalanche/Source/AvalancheSequencer/Private/NavigationTool/NavigationToolAvaSequence.h b/Engine/Plugins/Experimental/Avalanche/Source/AvalancheSequencer/Private/NavigationTool/NavigationToolAvaSequence.h
deleted file mode 100644
diff --git a/Engine/Plugins/Experimental/Avalanche/Source/AvalancheSequencer/Private/SequenceTree/Widgets/SAvaSequenceTree.cpp b/Engine/Plugins/Experimental/Avalanche/Source/AvalancheSequencer/Private/SequenceTree/Widgets/SAvaSequenceTree.cpp
--- a/Engine/Plugins/Experimental/Avalanche/Source/AvalancheSequencer/Private/SequenceTree/Widgets/SAvaSequenceTree.cpp
+++ b/Engine/Plugins/Experimental/Avalanche/Source/AvalancheSequencer/Private/SequenceTree/Widgets/SAvaSequenceTree.cpp
@@ -196,36 +196,36 @@ void SAvaSequenceTree::BindCommands(const TSharedRef<FAvaSequencer>& InSequencer
 	const FAvaSequencerCommands& AvaSequencerCommands = FAvaSequencerCommands::Get();
 
 	CommandList->MapAction(GenericCommands.Duplicate
-		, FExecuteAction::CreateSP(InSequencer, &FAvaSequencer::DuplicateSequences_Execute)
-		, FCanExecuteAction::CreateSP(InSequencer, &FAvaSequencer::DuplicateSequences_CanExecute));
+		, FExecuteAction::CreateSP(InSequencer, &FAvaSequencer::DuplicateSequence_Execute)
+		, FCanExecuteAction::CreateSP(InSequencer, &FAvaSequencer::DuplicateSequence_CanExecute));
 
 	CommandList->MapAction(GenericCommands.Delete
-		, FExecuteAction::CreateSP(InSequencer, &FAvaSequencer::DeleteSequences_Execute)
-		, FCanExecuteAction::CreateSP(InSequencer, &FAvaSequencer::DeleteSequences_CanExecute));
+		, FExecuteAction::CreateSP(InSequencer, &FAvaSequencer::DeleteSequence_Execute)
+		, FCanExecuteAction::CreateSP(InSequencer, &FAvaSequencer::DeleteSequence_CanExecute));
 
 	CommandList->MapAction(GenericCommands.Rename
 		, FExecuteAction::CreateSP(InSequencer, &FAvaSequencer::RelabelSequence_Execute)
 		, FCanExecuteAction::CreateSP(InSequencer, &FAvaSequencer::RelabelSequence_CanExecute));
 
 	CommandList->MapAction(AvaSequencerCommands.PlaySelected
-		, FExecuteAction::CreateSP(InSequencer, &FAvaSequencer::PlaySequences_Execute)
-		, FCanExecuteAction::CreateSP(InSequencer, &FAvaSequencer::PlaySequences_CanExecute));
+		, FExecuteAction::CreateSP(InSequencer, &FAvaSequencer::PlaySelected_Execute)
+		, FCanExecuteAction::CreateSP(InSequencer, &FAvaSequencer::PlaySelected_CanExecute));
 
 	CommandList->MapAction(AvaSequencerCommands.ContinueSelected
-		, FExecuteAction::CreateSP(InSequencer, &FAvaSequencer::ContinueSequences_Execute)
-		, FCanExecuteAction::CreateSP(InSequencer, &FAvaSequencer::ContinueSequences_CanExecute));
+		, FExecuteAction::CreateSP(InSequencer, &FAvaSequencer::ContinueSelected_Execute)
+		, FCanExecuteAction::CreateSP(InSequencer, &FAvaSequencer::ContinueSelected_CanExecute));
 
 	CommandList->MapAction(AvaSequencerCommands.StopSelected
-		, FExecuteAction::CreateSP(InSequencer, &FAvaSequencer::StopSequences_Execute)
-		, FCanExecuteAction::CreateSP(InSequencer, &FAvaSequencer::StopSequences_CanExecute));
+		, FExecuteAction::CreateSP(InSequencer, &FAvaSequencer::StopSelected_Execute)
+		, FCanExecuteAction::CreateSP(InSequencer, &FAvaSequencer::StopSelected_CanExecute));
 
 	CommandList->MapAction(AvaSequencerCommands.ExportSequence
-		, FExecuteAction::CreateSP(InSequencer, &FAvaSequencer::ExportSequences_Execute)
-		, FCanExecuteAction::CreateSP(InSequencer, &FAvaSequencer::ExportSequences_CanExecute));
+		, FExecuteAction::CreateSP(InSequencer, &FAvaSequencer::ExportSequence_Execute)
+		, FCanExecuteAction::CreateSP(InSequencer, &FAvaSequencer::ExportSequence_CanExecute));
 
 	CommandList->MapAction(AvaSequencerCommands.SpawnSequencePlayer
-		, FExecuteAction::CreateSP(InSequencer, &FAvaSequencer::SpawnPlayers_Execute)
-		, FCanExecuteAction::CreateSP(InSequencer, &FAvaSequencer::SpawnPlayers_CanExecute));
+		, FExecuteAction::CreateSP(InSequencer, &FAvaSequencer::SpawnPlayer_Execute)
+		, FCanExecuteAction::CreateSP(InSequencer, &FAvaSequencer::SpawnPlayer_CanExecute));
 }
 
 TSharedRef<ITableRow> SAvaSequenceTree::OnGenerateRow(FAvaSequenceItemPtr InListItem
@@ -324,7 +324,7 @@ FReply SAvaSequenceTree::OnNewSequenceClicked()
 {
 	if (SequencerWeak.IsValid())
 	{
-		SequencerWeak.Pin()->AddSequence();
+		SequencerWeak.Pin()->AddSequence_Execute();
 		return FReply::Handled();
 	}
 	return FReply::Unhandled();
diff --git a/Engine/Plugins/Experimental/Avalanche/Source/AvalancheSequencer/Public/AvaSequencerUtils.h b/Engine/Plugins/Experimental/Avalanche/Source/AvalancheSequencer/Public/AvaSequencerUtils.h
--- a/Engine/Plugins/Experimental/Avalanche/Source/AvalancheSequencer/Public/AvaSequencerUtils.h
+++ b/Engine/Plugins/Experimental/Avalanche/Source/AvalancheSequencer/Public/AvaSequencerUtils.h
@@ -7,14 +7,7 @@
 #include "Templates/SharedPointerFwd.h"
 #include "Templates/Function.h"
 
-class IAvaSceneInterface;
-class IAvaSequenceProvider;
-class IAvaSequencer;
 class IAvaSequencerController;
-class ISequencer;
-class UAvaSceneSubsystem;
-class UAvaSequencerSubsystem;
-class UWorld;
 enum class EMovieSceneTransformChannel : uint32;
 
 struct FAvaSequencerUtils
@@ -36,22 +29,4 @@ struct FAvaSequencerUtils
 	}
 
 	AVALANCHESEQUENCER_API static TSharedRef<IAvaSequencerController> CreateSequencerController();
-
-	/** @return Motion Design Sequencer World from an ISequencer */
-	static UWorld* GetSequencerWorld(const TSharedRef<ISequencer>& InSequencer);
-
-	/** @return Motion Design Sequencer Subsystem from an ISequencer */
-	static UAvaSequencerSubsystem* GetSequencerSubsystem(const TSharedRef<ISequencer>& InSequencer);
-
-	/** @return Motion Design Scene Subsystem from an ISequencer */
-	static UAvaSceneSubsystem* GetSceneSubsystem(const TSharedRef<ISequencer>& InSequencer);
-
-	/** @return Motion Design Scene Interface from an ISequencer */
-	static IAvaSceneInterface* GetSceneInterface(const TSharedRef<ISequencer>& InSequencer);
-
-	/** @return Motion Design Sequence Provider from an ISequencer */
-	static IAvaSequenceProvider* GetSequenceProvider(const TSharedRef<ISequencer>& InSequencer);
-
-	/** @return Motion Design Sequencer from an ISequencer */
-	static TSharedPtr<IAvaSequencer> GetAvaSequencer(const TSharedRef<ISequencer>& InSequencer);
 };
diff --git a/Engine/Plugins/Experimental/Avalanche/Source/AvalancheSequencer/Public/IAvaSequencer.h b/Engine/Plugins/Experimental/Avalanche/Source/AvalancheSequencer/Public/IAvaSequencer.h
--- a/Engine/Plugins/Experimental/Avalanche/Source/AvalancheSequencer/Public/IAvaSequencer.h
+++ b/Engine/Plugins/Experimental/Avalanche/Source/AvalancheSequencer/Public/IAvaSequencer.h
@@ -2,13 +2,11 @@
 
 #pragma once
 
-#include "AvaSequence.h"
 #include "Containers/ArrayView.h"
 #include "Containers/ContainersFwd.h"
 #include "Containers/StringView.h"
 #include "Templates/SharedPointer.h"
 #include "Templates/UnrealTypeTraits.h"
-#include "UObject/WeakObjectPtrTemplatesFwd.h"
 
 class AActor;
 class FAvaEditorSelection;
@@ -18,6 +16,7 @@ class IAvaSequencerProvider;
 class IPropertyHandle;
 class ISequencer;
 class SWidget;
+class UAvaSequence;
 class UObject;
 class USequencerSettings;
 
@@ -60,13 +59,4 @@ public:
 
 	/** Should be called when the UAvaSequence tree has changed, so that FAvaSequencer can trigger a UI refresh */
 	virtual void NotifyOnSequenceTreeChanged() = 0;
-
-	/** Gets the list of embedded level sequences */
-	virtual const TArray<TWeakObjectPtr<UAvaSequence>>& GetRootSequences() const = 0;
-
-	DECLARE_MULTICAST_DELEGATE_OneParam(FOnSequenceAdded, UAvaSequence*);
-	virtual FOnSequenceAdded& OnSequenceAdded() = 0;
-
-	DECLARE_MULTICAST_DELEGATE_OneParam(FOnSequenceRemoved, UAvaSequence*);
-	virtual FOnSequenceRemoved& OnSequenceRemoved() = 0;
 };
