[Motion Design] Rundown and Playback: Support for scene state events through RC controller behaviors

Refactor to allow the UpdateRC command to carry flags for more specific purposes.
In particular, for any RC update that are not part of a transition, it will now execute the controller behaviors.

diff --git a/Engine/Plugins/Experimental/Avalanche/Source/AvalancheMedia/Private/Playable/AvaPlayable.cpp b/Engine/Plugins/Experimental/Avalanche/Source/AvalancheMedia/Private/Playable/AvaPlayable.cpp
--- a/Engine/Plugins/Experimental/Avalanche/Source/AvalancheMedia/Private/Playable/AvaPlayable.cpp
+++ b/Engine/Plugins/Experimental/Avalanche/Source/AvalancheMedia/Private/Playable/AvaPlayable.cpp
@@ -249,7 +249,7 @@ EAvaPlayableCommandResult UAvaPlayable::ExecuteAnimationCommand(EAvaPlaybackAnim
 	return EAvaPlayableCommandResult::Executed;
 }
 
-EAvaPlayableCommandResult UAvaPlayable::UpdateRemoteControlCommand(const TSharedRef<FAvaPlayableRemoteControlValues>& InRemoteControlValues)
+EAvaPlayableCommandResult UAvaPlayable::UpdateRemoteControlCommand(const TSharedRef<FAvaPlayableRemoteControlValues>& InRemoteControlValues, EAvaPlayableRCUpdateFlags InFlags)
 {
 	const EAvaPlayableStatus PlayableStatus = GetPlayableStatus();
 
@@ -291,37 +291,57 @@ EAvaPlayableCommandResult UAvaPlayable::UpdateRemoteControlCommand(const TShared
 		return EAvaPlayableCommandResult::ErrorDiscard;
 	}
 
-	LatestRemoteControlValues = InRemoteControlValues;
-
 	using namespace UE::AvaPlayable::Private;
 	UE_LOG(LogAvaPlayable, Verbose, TEXT("%s Playable {%s} -> Executing RC Update."), *GetBriefFrameInfo(), *GetPrettyPlayableInfo(this));
 
 	TSet<FGuid> AppliedEntities;
 	
-	// Apply special controllers that don't work just with entity values.
-	const bool bExecuteSpecialControllerBehavior = CVarExecuteSpecialRCControllerBehavior.GetValueOnGameThread();
-	const bool bUpdateControllerValues = CVarUpdateRCControllerValues.GetValueOnGameThread();
-	if (bExecuteSpecialControllerBehavior || bUpdateControllerValues)
+	if (CVarUpdateRCControllerValues.GetValueOnGameThread())
 	{
+		// Apply special controllers that don't work just with entity values.
+		const bool bExecuteSpecialControllerBehavior = CVarExecuteSpecialRCControllerBehavior.GetValueOnGameThread();
+		const bool bExecuteAllControllerBehaviors = EnumHasAnyFlags(InFlags, EAvaPlayableRCUpdateFlags::ExecuteControllerBehaviors);
+
+		TSet<FGuid> ModifiedControllers;
+
+		// Behaviors are executed only on the controller values that changed. 
+		if (bExecuteAllControllerBehaviors)
+		{
+			for (const TPair<FGuid, FAvaPlayableRemoteControlValue>& Controller : InRemoteControlValues->ControllerValues)
+			{
+				if (LatestRemoteControlValues)
+				{
+					const FAvaPlayableRemoteControlValue* FoundValue = LatestRemoteControlValues->ControllerValues.Find(Controller.Key);
+					if (!FoundValue || !FoundValue->IsSameValueAs(Controller.Value))
+					{
+						ModifiedControllers.Add(Controller.Key);
+					}
+				}
+				else
+				{
+					ModifiedControllers.Add(Controller.Key);
+				}
+			}
+		}
+		
 		TArray<URCVirtualPropertyBase*> Controllers = RemoteControlPreset->GetControllers();
 		for (URCVirtualPropertyBase* Controller : Controllers)
 		{
-			using namespace UE::AvaPlayableRemoteControl;
-			using namespace UE::RCCustomControllers;
-
-			// Texture Controller Require the Bind Behavior to be executed to setup the texture in external mode.
-			if (bExecuteSpecialControllerBehavior &&
-				GetCustomControllerTypeName(Controller) == CustomTextureControllerName)
+			if (const FAvaPlayableRemoteControlValue* ControllerValue = InRemoteControlValues->GetControllerValue(Controller->Id))
 			{
-				if (const FAvaPlayableRemoteControlValue* ControllerValue = InRemoteControlValues->GetControllerValue(Controller->Id))
+				using namespace UE::AvaPlayableRemoteControl;
+				using namespace UE::RCCustomControllers;
+
+				// Texture Controller Require the Bind Behavior to be executed to setup the texture in external mode.
+				if ((bExecuteSpecialControllerBehavior && GetCustomControllerTypeName(Controller) == CustomTextureControllerName)
+					|| (bExecuteAllControllerBehaviors && (ModifiedControllers.Contains(Controller->Id))))
 				{
 					SetValueOfController(Controller, ControllerValue->Value);
 					GetEntitiesControlledByController(RemoteControlPreset, Controller, AppliedEntities);
+					UE_LOG(LogAvaPlayable, Verbose, TEXT("%s Playable {%s} -> Updating Controller %s (with behaviors)."),
+						*GetBriefFrameInfo(), *GetPrettyPlayableInfo(this), *Controller->DisplayName.ToString());
 				}
-			}
-			else if (bUpdateControllerValues)
-			{
-				if (const FAvaPlayableRemoteControlValue* ControllerValue = InRemoteControlValues->GetControllerValue(Controller->Id))
+				else
 				{
 					FScopedPushControllerBehavioursEnable PushBehaviorsEnable(Controller, /*bInBehaviorsEnabled*/ false);
 					SetValueOfController(Controller, ControllerValue->Value);
@@ -365,6 +385,7 @@ EAvaPlayableCommandResult UAvaPlayable::UpdateRemoteControlCommand(const TShared
 		}
 	}
 
+	LatestRemoteControlValues = InRemoteControlValues;
 	OnRemoteControlValuesApplied();
 
 	return EAvaPlayableCommandResult::Executed;
diff --git a/Engine/Plugins/Experimental/Avalanche/Source/AvalancheMedia/Private/Playable/AvaPlayableLibrary.cpp b/Engine/Plugins/Experimental/Avalanche/Source/AvalancheMedia/Private/Playable/AvaPlayableLibrary.cpp
--- a/Engine/Plugins/Experimental/Avalanche/Source/AvalancheMedia/Private/Playable/AvaPlayableLibrary.cpp
+++ b/Engine/Plugins/Experimental/Avalanche/Source/AvalancheMedia/Private/Playable/AvaPlayableLibrary.cpp
@@ -93,7 +93,7 @@ bool UAvaPlayableLibrary::UpdatePlayableRemoteControlValues(const UObject* InWor
 
 			if (const TSharedPtr<FAvaPlayableRemoteControlValues> RemoteControlValues = Transition->GetValuesForPlayable(Playable, bIsEnterPlayable))
 			{
-				if (Playable->UpdateRemoteControlCommand(RemoteControlValues.ToSharedRef()) == EAvaPlayableCommandResult::Executed)
+				if (Playable->UpdateRemoteControlCommand(RemoteControlValues.ToSharedRef(), EAvaPlayableRCUpdateFlags::ExecuteControllerBehaviors) == EAvaPlayableCommandResult::Executed)
 				{
 					return true;
 				}
diff --git a/Engine/Plugins/Experimental/Avalanche/Source/AvalancheMedia/Private/Playable/Playables/AvaPlayableRemoteProxy.cpp b/Engine/Plugins/Experimental/Avalanche/Source/AvalancheMedia/Private/Playable/Playables/AvaPlayableRemoteProxy.cpp
--- a/Engine/Plugins/Experimental/Avalanche/Source/AvalancheMedia/Private/Playable/Playables/AvaPlayableRemoteProxy.cpp
+++ b/Engine/Plugins/Experimental/Avalanche/Source/AvalancheMedia/Private/Playable/Playables/AvaPlayableRemoteProxy.cpp
@@ -221,12 +221,12 @@ EAvaPlayableCommandResult UAvaPlayableRemoteProxy::ExecuteAnimationCommand(EAvaP
 	return EAvaPlayableCommandResult::Executed;
 }
 
-EAvaPlayableCommandResult UAvaPlayableRemoteProxy::UpdateRemoteControlCommand(const TSharedRef<FAvaPlayableRemoteControlValues>& InRemoteControlValues)
+EAvaPlayableCommandResult UAvaPlayableRemoteProxy::UpdateRemoteControlCommand(const TSharedRef<FAvaPlayableRemoteControlValues>& InRemoteControlValues, EAvaPlayableRCUpdateFlags InFlags)
 {
 	IAvaPlaybackClient& PlaybackClient = IAvaMediaModule::Get().GetPlaybackClient();
 	if (PlaybackClient.HasAnyServerOnlineForChannel(PlayingChannelFName))
 	{
-		PlaybackClient.RequestRemoteControlUpdate(InstanceId, SourceAssetPath, PlayingChannelName, *InRemoteControlValues);
+		PlaybackClient.RequestRemoteControlUpdate(InstanceId, SourceAssetPath, PlayingChannelName, *InRemoteControlValues, InFlags);
 	}
 	return EAvaPlayableCommandResult::Executed;
 }
diff --git a/Engine/Plugins/Experimental/Avalanche/Source/AvalancheMedia/Private/Playable/Playables/AvaPlayableRemoteProxy.h b/Engine/Plugins/Experimental/Avalanche/Source/AvalancheMedia/Private/Playable/Playables/AvaPlayableRemoteProxy.h
--- a/Engine/Plugins/Experimental/Avalanche/Source/AvalancheMedia/Private/Playable/Playables/AvaPlayableRemoteProxy.h
+++ b/Engine/Plugins/Experimental/Avalanche/Source/AvalancheMedia/Private/Playable/Playables/AvaPlayableRemoteProxy.h
@@ -28,7 +28,7 @@ public:
 	virtual EAvaPlayableStatus GetPlayableStatus() const override;
 	virtual IAvaSceneInterface* GetSceneInterface() const override;
 	virtual EAvaPlayableCommandResult ExecuteAnimationCommand(EAvaPlaybackAnimAction InAnimAction, const FAvaPlaybackAnimPlaySettings& InAnimPlaySettings) override;
-	virtual EAvaPlayableCommandResult UpdateRemoteControlCommand(const TSharedRef<FAvaPlayableRemoteControlValues>& InRemoteControlValues) override;
+	virtual EAvaPlayableCommandResult UpdateRemoteControlCommand(const TSharedRef<FAvaPlayableRemoteControlValues>& InRemoteControlValues, EAvaPlayableRCUpdateFlags InFlags) override;
 	virtual bool IsRemoteProxy() const override { return true; }
 	virtual void SetUserData(const FString& InUserData) override;
 
diff --git a/Engine/Plugins/Experimental/Avalanche/Source/AvalancheMedia/Private/Playable/Transition/AvaPlayableLocalTransition.cpp b/Engine/Plugins/Experimental/Avalanche/Source/AvalancheMedia/Private/Playable/Transition/AvaPlayableLocalTransition.cpp
--- a/Engine/Plugins/Experimental/Avalanche/Source/AvalancheMedia/Private/Playable/Transition/AvaPlayableLocalTransition.cpp
+++ b/Engine/Plugins/Experimental/Avalanche/Source/AvalancheMedia/Private/Playable/Transition/AvaPlayableLocalTransition.cpp
@@ -332,7 +332,8 @@ bool UAvaPlayableLocalTransition::Start()
 	{
 		if (EnterPlayableValues.IsValidIndex(ArrayIndex) && EnterPlayableValues[ArrayIndex].IsValid() && !IsPlayingPlayable(Playable))
 		{
-			Playable->UpdateRemoteControlCommand(EnterPlayableValues[ArrayIndex].ToSharedRef());
+			// Note: we are still not running the controller behaviors when updating RC from a transition.
+			Playable->UpdateRemoteControlCommand(EnterPlayableValues[ArrayIndex].ToSharedRef(), EAvaPlayableRCUpdateFlags::None);
 		}
 
 		if (bIsPreviewFrameTransition || !Helper.AddTransitionBehaviorInstance(Playable, EAvaPlayableTransitionEntryRole::Enter))
diff --git a/Engine/Plugins/Experimental/Avalanche/Source/AvalancheMedia/Private/Playback/AvaPlaybackClient.cpp b/Engine/Plugins/Experimental/Avalanche/Source/AvalancheMedia/Private/Playback/AvaPlaybackClient.cpp
--- a/Engine/Plugins/Experimental/Avalanche/Source/AvalancheMedia/Private/Playback/AvaPlaybackClient.cpp
+++ b/Engine/Plugins/Experimental/Avalanche/Source/AvalancheMedia/Private/Playback/AvaPlaybackClient.cpp
@@ -399,13 +399,15 @@ void FAvaPlaybackClient::RequestAnimAction(const FGuid& InInstanceId, const FSof
 
 void FAvaPlaybackClient::RequestRemoteControlUpdate(const FGuid& InInstanceId, const FSoftObjectPath& InAssetPath,
 														 const FString& InChannelName,
-														 const FAvaPlayableRemoteControlValues& InRemoteControlValues)
+														 const FAvaPlayableRemoteControlValues& InRemoteControlValues,
+														 EAvaPlayableRCUpdateFlags InFlags)
 {
 	FAvaPlaybackRemoteControlUpdateRequest* Request = FMessageEndpoint::MakeMessage<FAvaPlaybackRemoteControlUpdateRequest>();
 	Request->InstanceId = InInstanceId;
 	Request->AssetPath = InAssetPath;
 	Request->ChannelName = InChannelName;
 	Request->RemoteControlValues = InRemoteControlValues;
+	Request->UpdateFlags = InFlags;
 
 	if (!InAssetPath.IsValid())
 	{
diff --git a/Engine/Plugins/Experimental/Avalanche/Source/AvalancheMedia/Private/Playback/AvaPlaybackClient.h b/Engine/Plugins/Experimental/Avalanche/Source/AvalancheMedia/Private/Playback/AvaPlaybackClient.h
--- a/Engine/Plugins/Experimental/Avalanche/Source/AvalancheMedia/Private/Playback/AvaPlaybackClient.h
+++ b/Engine/Plugins/Experimental/Avalanche/Source/AvalancheMedia/Private/Playback/AvaPlaybackClient.h
@@ -36,7 +36,7 @@ public:
 	virtual void RequestPlayback(const FGuid& InInstanceId, const FSoftObjectPath& InAssetPath, const FString& InChannelName, EAvaPlaybackAction InAction, const FString& InArguments = FString()) override;
 	virtual void RequestAnimPlayback(const FGuid& InInstanceId, const FSoftObjectPath& InAssetPath, const FString& InChannelName, const FAvaPlaybackAnimPlaySettings& InAnimSettings) override;
 	virtual void RequestAnimAction(const FGuid& InInstanceId, const FSoftObjectPath& InAssetPath, const FString& InChannelName, const FString& InAnimationName, EAvaPlaybackAnimAction InAction) override;
-	virtual void RequestRemoteControlUpdate(const FGuid& InInstanceId, const FSoftObjectPath& InAssetPath, const FString& InChannelName, const FAvaPlayableRemoteControlValues& InRemoteControlValues) override;
+	virtual void RequestRemoteControlUpdate(const FGuid& InInstanceId, const FSoftObjectPath& InAssetPath, const FString& InChannelName, const FAvaPlayableRemoteControlValues& InRemoteControlValues, EAvaPlayableRCUpdateFlags InFlags) override;
 	virtual void RequestPlayableTransitionStart(const FGuid& InTransitionId, TArray<FGuid>&& InEnterInstanceIds, TArray<FGuid>&& InPlayingInstanceIds, TArray<FGuid>&& InExitInstanceIds, TArray<FAvaPlayableRemoteControlValues>&& InEnterValues, const FName& InChannelName, EAvaPlayableTransitionFlags InTransitionFlags) override;
 	virtual void RequestPlayableTransitionStop(const FGuid& InTransitionId, const FName& InChannelName) override;
 	virtual void RequestBroadcast(const FString& InProfile, const FName& InChannel, const TArray<UMediaOutput*>& InRemoteMediaOutputs, EAvaBroadcastAction InAction, const FString& InServerName = FString()) override;
diff --git a/Engine/Plugins/Experimental/Avalanche/Source/AvalancheMedia/Private/Playback/AvaPlaybackClientDummy.h b/Engine/Plugins/Experimental/Avalanche/Source/AvalancheMedia/Private/Playback/AvaPlaybackClientDummy.h
--- a/Engine/Plugins/Experimental/Avalanche/Source/AvalancheMedia/Private/Playback/AvaPlaybackClientDummy.h
+++ b/Engine/Plugins/Experimental/Avalanche/Source/AvalancheMedia/Private/Playback/AvaPlaybackClientDummy.h
@@ -30,7 +30,7 @@ public:
 	virtual void RequestPlayback(const FGuid& InInstanceId, const FSoftObjectPath& InAssetPath, const FString& InChannelName, EAvaPlaybackAction InAction, const FString& InArguments) override {}
 	virtual void RequestAnimPlayback(const FGuid& InInstanceId, const FSoftObjectPath& InAssetPath, const FString& InChannelName, const FAvaPlaybackAnimPlaySettings& InAnimSettings) override {};
 	virtual void RequestAnimAction(const FGuid& InInstanceId, const FSoftObjectPath& InAssetPath, const FString& InChannelName, const FString& InAnimationName, EAvaPlaybackAnimAction InAction) override {};
-	virtual void RequestRemoteControlUpdate(const FGuid& InInstanceId, const FSoftObjectPath& InAssetPath, const FString& InChannelName, const FAvaPlayableRemoteControlValues& InRemoteControlValues) override {};
+	virtual void RequestRemoteControlUpdate(const FGuid& InInstanceId, const FSoftObjectPath& InAssetPath, const FString& InChannelName, const FAvaPlayableRemoteControlValues& InRemoteControlValues, EAvaPlayableRCUpdateFlags InFlags) override {};
 	virtual void RequestPlayableTransitionStart(const FGuid& InTransitionId, TArray<FGuid>&& InEnterInstanceIds, TArray<FGuid>&& InPlayingInstanceIds, TArray<FGuid>&& InExitInstanceIds, TArray<FAvaPlayableRemoteControlValues>&& InEnterValues, const FName& InChannelName, EAvaPlayableTransitionFlags InTransitionFlags) override {}
 	virtual void RequestPlayableTransitionStop(const FGuid& InTransitionId, const FName& InChannelName) override {}
 	virtual void RequestBroadcast(const FString& InProfile, const FName& InChannel, const TArray<UMediaOutput*>& InRemoteMediaOutputs, EAvaBroadcastAction InAction, const FString& InServerName) override {}
diff --git a/Engine/Plugins/Experimental/Avalanche/Source/AvalancheMedia/Private/Playback/AvaPlaybackGraph.cpp b/Engine/Plugins/Experimental/Avalanche/Source/AvalancheMedia/Private/Playback/AvaPlaybackGraph.cpp
--- a/Engine/Plugins/Experimental/Avalanche/Source/AvalancheMedia/Private/Playback/AvaPlaybackGraph.cpp
+++ b/Engine/Plugins/Experimental/Avalanche/Source/AvalancheMedia/Private/Playback/AvaPlaybackGraph.cpp
@@ -541,10 +541,10 @@ void UAvaPlaybackGraph::PushAnimationCommand(const FSoftObjectPath& InSourceAsse
 	PendingAnimationCommands.Add({{Timeout, InSourceAssetPath, InChannelName, FName(InChannelName)}, InAnimAction, InAnimSettings});
 }
 
-void UAvaPlaybackGraph::PushRemoteControlValues(const FSoftObjectPath& InSourceAssetPath, const FString& InChannelName, const TSharedRef<FAvaPlayableRemoteControlValues>& InRemoteControlValues)
+void UAvaPlaybackGraph::PushRemoteControlValues(const FSoftObjectPath& InSourceAssetPath, const FString& InChannelName, const TSharedRef<FAvaPlayableRemoteControlValues>& InRemoteControlValues, EAvaPlayableRCUpdateFlags InFlags)
 {
 	const FDateTime Timeout = FDateTime::UtcNow() + FTimespan::FromSeconds(10);
-	PendingRemoteControlCommands.Add({{Timeout, InSourceAssetPath, InChannelName, FName(InChannelName)}, InRemoteControlValues});
+	PendingRemoteControlCommands.Add({{Timeout, InSourceAssetPath, InChannelName, FName(InChannelName)}, InRemoteControlValues, InFlags});
 }
 
 void UAvaPlaybackGraph::SetupPlaybackNode(UAvaPlaybackNode* InPlaybackNode, bool bSelectNewNode)
@@ -714,7 +714,7 @@ void UAvaPlaybackGraph::ExecutePendingRemoteControlCommands()
 			continue;
 		}
 
-		const EAvaPlayableCommandResult Result = Playable->UpdateRemoteControlCommand(RemoteControlCommand.Values);
+		const EAvaPlayableCommandResult Result = Playable->UpdateRemoteControlCommand(RemoteControlCommand.Values, RemoteControlCommand.Flags);
 		if (Result == EAvaPlayableCommandResult::KeepPending)
 		{
 			StillPendingRemoteControlCommands.Add(RemoteControlCommand);
diff --git a/Engine/Plugins/Experimental/Avalanche/Source/AvalancheMedia/Private/Playback/AvaPlaybackManager.cpp b/Engine/Plugins/Experimental/Avalanche/Source/AvalancheMedia/Private/Playback/AvaPlaybackManager.cpp
--- a/Engine/Plugins/Experimental/Avalanche/Source/AvalancheMedia/Private/Playback/AvaPlaybackManager.cpp
+++ b/Engine/Plugins/Experimental/Avalanche/Source/AvalancheMedia/Private/Playback/AvaPlaybackManager.cpp
@@ -486,7 +486,7 @@ bool FAvaPlaybackManager::PushAnimationCommand(const FGuid& InInstanceId, const
 	return bAnimationCommandPushed;
 }
 
-bool FAvaPlaybackManager::PushRemoteControlCommand(const FGuid& InInstanceId, const FSoftObjectPath& InSourcePath, const FString& InChannelName, const TSharedRef<FAvaPlayableRemoteControlValues>& InRemoteControlValues)
+bool FAvaPlaybackManager::PushRemoteControlCommand(const FGuid& InInstanceId, const FSoftObjectPath& InSourcePath, const FString& InChannelName, const TSharedRef<FAvaPlayableRemoteControlValues>& InRemoteControlValues, EAvaPlayableRCUpdateFlags InFlags)
 {
 	bool bRemoteControlValuesPushed = false;
 	if (InInstanceId.IsValid())
@@ -494,18 +494,18 @@ bool FAvaPlaybackManager::PushRemoteControlCommand(const FGuid& InInstanceId, co
 		// Check if we have the entry directly in the map. Rundown path.
 		if (const TSharedPtr<FAvaPlaybackInstance> PlaybackInstance = FindPlaybackInstance(InInstanceId, InSourcePath, InChannelName))
 		{
-			PlaybackInstance->GetPlayback()->PushRemoteControlValues(InSourcePath, InChannelName, InRemoteControlValues);
+			PlaybackInstance->GetPlayback()->PushRemoteControlValues(InSourcePath, InChannelName, InRemoteControlValues, InFlags);
 			bRemoteControlValuesPushed = true;
 		}
 	}
 	else
 	{
 		// If we had a playback asset (not from the rundown path), check if it is present in one of the playback instances.
-		ForAllPlaybackInstances([InSourcePath, InChannelName, InRemoteControlValues, &bRemoteControlValuesPushed](FAvaPlaybackInstance& InPlaybackInstance)
+		ForAllPlaybackInstances([InSourcePath, InChannelName, InRemoteControlValues, &bRemoteControlValuesPushed, InFlags](FAvaPlaybackInstance& InPlaybackInstance)
 		{
 			if (InPlaybackInstance.GetPlayback()->HasPlayerNodeForSourceAsset(InSourcePath))
 			{
-				InPlaybackInstance.GetPlayback()->PushRemoteControlValues(InSourcePath, InChannelName, InRemoteControlValues);
+				InPlaybackInstance.GetPlayback()->PushRemoteControlValues(InSourcePath, InChannelName, InRemoteControlValues, InFlags);
 				bRemoteControlValuesPushed = true;
 			}
 		});
@@ -514,7 +514,7 @@ bool FAvaPlaybackManager::PushRemoteControlCommand(const FGuid& InInstanceId, co
 	if (!bRemoteControlValuesPushed && bEnablePlaybackCommandsBuffering)
 	{
 		FString CommandBufferKey = MakeCommandBufferKey(InInstanceId, InSourcePath, InChannelName);
-		GetOrCreatePlaybackCommandBuffers(MoveTemp(CommandBufferKey)).RemoteControlCommands.Add({InRemoteControlValues});
+		GetOrCreatePlaybackCommandBuffers(MoveTemp(CommandBufferKey)).RemoteControlCommands.Add({InRemoteControlValues, InFlags});
 	}
 	return bRemoteControlValuesPushed;
 }
@@ -542,7 +542,7 @@ void FAvaPlaybackManager::ApplyPendingCommands(UAvaPlaybackGraph* InPlaybackObje
 		}
 		for (const FPlaybackObjectCommandBuffers::FRemoteControlCommand& RCCommand : Commands->RemoteControlCommands)
 		{
-			InPlaybackObject->PushRemoteControlValues(InSourcePath, InChannelName, RCCommand.Values);
+			InPlaybackObject->PushRemoteControlValues(InSourcePath, InChannelName, RCCommand.Values, RCCommand.UpdateFlags);
 		}
 		PlaybackObjectCommandBuffers.Remove(CommandBufferKey);
 	}
diff --git a/Engine/Plugins/Experimental/Avalanche/Source/AvalancheMedia/Private/Playback/AvaPlaybackMessages.h b/Engine/Plugins/Experimental/Avalanche/Source/AvalancheMedia/Private/Playback/AvaPlaybackMessages.h
--- a/Engine/Plugins/Experimental/Avalanche/Source/AvalancheMedia/Private/Playback/AvaPlaybackMessages.h
+++ b/Engine/Plugins/Experimental/Avalanche/Source/AvalancheMedia/Private/Playback/AvaPlaybackMessages.h
@@ -457,6 +457,9 @@ struct FAvaPlaybackRemoteControlUpdateRequest : public FAvaPlaybackClientMessage
 	
 	UPROPERTY()
 	FAvaPlayableRemoteControlValues RemoteControlValues;
+	
+	UPROPERTY()
+	EAvaPlayableRCUpdateFlags UpdateFlags = EAvaPlayableRCUpdateFlags::None;
 };
 
 USTRUCT()
diff --git a/Engine/Plugins/Experimental/Avalanche/Source/AvalancheMedia/Private/Playback/AvaPlaybackServer.cpp b/Engine/Plugins/Experimental/Avalanche/Source/AvalancheMedia/Private/Playback/AvaPlaybackServer.cpp
--- a/Engine/Plugins/Experimental/Avalanche/Source/AvalancheMedia/Private/Playback/AvaPlaybackServer.cpp
+++ b/Engine/Plugins/Experimental/Avalanche/Source/AvalancheMedia/Private/Playback/AvaPlaybackServer.cpp
@@ -672,7 +672,8 @@ void FAvaPlaybackServer::HandleAnimPlaybackRequest(const FAvaPlaybackAnimPlaybac
 
 void FAvaPlaybackServer::HandleRemoteControlUpdateRequest(const FAvaPlaybackRemoteControlUpdateRequest& InMessage, const TSharedRef<IMessageContext, ESPMode::ThreadSafe>& InContext)
 {
-	Manager->PushRemoteControlCommand(InMessage.InstanceId, InMessage.AssetPath, InMessage.ChannelName, MakeShared<FAvaPlayableRemoteControlValues>(InMessage.RemoteControlValues));
+	Manager->PushRemoteControlCommand(InMessage.InstanceId, InMessage.AssetPath, InMessage.ChannelName,
+		MakeShared<FAvaPlayableRemoteControlValues>(InMessage.RemoteControlValues), InMessage.UpdateFlags);
 }
 
 void FAvaPlaybackServer::HandlePlayableTransitionStartRequest(const FAvaPlaybackTransitionStartRequest& InMessage, const TSharedRef<IMessageContext, ESPMode::ThreadSafe>& InContext)
diff --git a/Engine/Plugins/Experimental/Avalanche/Source/AvalancheMedia/Private/Rundown/AvaRundown.cpp b/Engine/Plugins/Experimental/Avalanche/Source/AvalancheMedia/Private/Rundown/AvaRundown.cpp
--- a/Engine/Plugins/Experimental/Avalanche/Source/AvalancheMedia/Private/Rundown/AvaRundown.cpp
+++ b/Engine/Plugins/Experimental/Avalanche/Source/AvalancheMedia/Private/Rundown/AvaRundown.cpp
@@ -1965,7 +1965,9 @@ bool UAvaRundown::PushRuntimeRemoteControlValues(int32 InPageId, bool bInIsPrevi
 		{
 			if (UAvaPlaybackGraph* Playback = PagePlayer->GetPlayback(InstanceIndex))
 			{
-				Playback->PushRemoteControlValues(PagePlayer->GetSourceAssetPath(InstanceIndex), PagePlayer->ChannelName, SharedRCValues);
+				// Remote Control Value (live) updates, when not part of a transition, will execute the controller behaviors on the runtime instances.
+				constexpr EAvaPlayableRCUpdateFlags UpdateFlags = EAvaPlayableRCUpdateFlags::ExecuteControllerBehaviors;
+				Playback->PushRemoteControlValues(PagePlayer->GetSourceAssetPath(InstanceIndex), PagePlayer->ChannelName, SharedRCValues, UpdateFlags);
 			}
 		}
 		bValuesPushed = true;
diff --git a/Engine/Plugins/Experimental/Avalanche/Source/AvalancheMedia/Public/AvaMediaDefines.h b/Engine/Plugins/Experimental/Avalanche/Source/AvalancheMedia/Public/AvaMediaDefines.h
--- a/Engine/Plugins/Experimental/Avalanche/Source/AvalancheMedia/Public/AvaMediaDefines.h
+++ b/Engine/Plugins/Experimental/Avalanche/Source/AvalancheMedia/Public/AvaMediaDefines.h
@@ -172,6 +172,16 @@ enum class EAvaPlayableTransitionEventFlags : uint8
 };
 ENUM_CLASS_FLAGS(EAvaPlayableTransitionEventFlags);
 
+/** Flags for the remote control values update command. */
+UENUM(meta=(Bitflags, UseEnumValuesAsMaskValuesInEditor = "true"))
+enum class EAvaPlayableRCUpdateFlags : uint8
+{
+	None = 0,
+	/** Execute Controller Behaviors */
+	ExecuteControllerBehaviors = 1 << 0
+};
+ENUM_CLASS_FLAGS(EAvaPlayableRCUpdateFlags);
+
 /*
  * Parameters used to tell a Channel its Playback Settings
  */
diff --git a/Engine/Plugins/Experimental/Avalanche/Source/AvalancheMedia/Public/Playable/AvaPlayable.h b/Engine/Plugins/Experimental/Avalanche/Source/AvalancheMedia/Public/Playable/AvaPlayable.h
--- a/Engine/Plugins/Experimental/Avalanche/Source/AvalancheMedia/Public/Playable/AvaPlayable.h
+++ b/Engine/Plugins/Experimental/Avalanche/Source/AvalancheMedia/Public/Playable/AvaPlayable.h
@@ -115,7 +115,7 @@ public:
 	virtual EAvaPlayableStatus GetPlayableStatus() const { return EAvaPlayableStatus::Unknown;}
 	virtual IAvaSceneInterface* GetSceneInterface() const {return nullptr;}
 	virtual EAvaPlayableCommandResult ExecuteAnimationCommand(EAvaPlaybackAnimAction InAnimAction, const FAvaPlaybackAnimPlaySettings& InAnimPlaySettings);
-	virtual EAvaPlayableCommandResult UpdateRemoteControlCommand(const TSharedRef<FAvaPlayableRemoteControlValues>& InRemoteControlValues);
+	virtual EAvaPlayableCommandResult UpdateRemoteControlCommand(const TSharedRef<FAvaPlayableRemoteControlValues>& InRemoteControlValues, EAvaPlayableRCUpdateFlags InFlags = EAvaPlayableRCUpdateFlags::None);
 	virtual bool IsRemoteProxy() const { return false; }
 	virtual bool GetShouldBeVisible() const { return true; }
 	virtual void SetShouldBeVisible(bool bInShouldBeVisible) {}
diff --git a/Engine/Plugins/Experimental/Avalanche/Source/AvalancheMedia/Public/Playback/AvaPlaybackGraph.h b/Engine/Plugins/Experimental/Avalanche/Source/AvalancheMedia/Public/Playback/AvaPlaybackGraph.h
--- a/Engine/Plugins/Experimental/Avalanche/Source/AvalancheMedia/Public/Playback/AvaPlaybackGraph.h
+++ b/Engine/Plugins/Experimental/Avalanche/Source/AvalancheMedia/Public/Playback/AvaPlaybackGraph.h
@@ -272,7 +272,7 @@ public:
 	virtual void Tick(float DeltaTime);
 	
 	void PushAnimationCommand(const FSoftObjectPath& InSourceAssetPath, const FString& InChannelName, EAvaPlaybackAnimAction InAnimAction, const FAvaPlaybackAnimPlaySettings& InAnimSettings);
-	void PushRemoteControlValues(const FSoftObjectPath& InSourceAssetPath, const FString& InChannelName, const TSharedRef<FAvaPlayableRemoteControlValues>& InRemoteControlValues);
+	void PushRemoteControlValues(const FSoftObjectPath& InSourceAssetPath, const FString& InChannelName, const TSharedRef<FAvaPlayableRemoteControlValues>& InRemoteControlValues, EAvaPlayableRCUpdateFlags InFlags = EAvaPlayableRCUpdateFlags::None);
 
 	void SetupPlaybackNode(UAvaPlaybackNode* InPlaybackNode, bool bSelectNewNode = true);
 
@@ -440,6 +440,7 @@ protected:
 	struct FRemoteControlCommand : public FPlaybackCommand
 	{
 		TSharedRef<FAvaPlayableRemoteControlValues> Values;
+		EAvaPlayableRCUpdateFlags Flags;
 	};
 	
 	TArray<FAnimationCommand> PendingAnimationCommands;
diff --git a/Engine/Plugins/Experimental/Avalanche/Source/AvalancheMedia/Public/Playback/AvaPlaybackManager.h b/Engine/Plugins/Experimental/Avalanche/Source/AvalancheMedia/Public/Playback/AvaPlaybackManager.h
--- a/Engine/Plugins/Experimental/Avalanche/Source/AvalancheMedia/Public/Playback/AvaPlaybackManager.h
+++ b/Engine/Plugins/Experimental/Avalanche/Source/AvalancheMedia/Public/Playback/AvaPlaybackManager.h
@@ -225,7 +225,7 @@ public:
 	TArray<FSoftObjectPath> StopAllPlaybacks(bool bInUnload);
 
 	bool PushAnimationCommand(const FGuid& InInstanceId, const FSoftObjectPath& InSourcePath, const FString& InChannelName, EAvaPlaybackAnimAction InAction, const FAvaPlaybackAnimPlaySettings& InAnimSettings);
-	bool PushRemoteControlCommand(const FGuid& InInstanceId, const FSoftObjectPath& InSourcePath, const FString& InChannelName, const TSharedRef<FAvaPlayableRemoteControlValues>& InRemoteControlValues);
+	bool PushRemoteControlCommand(const FGuid& InInstanceId, const FSoftObjectPath& InSourcePath, const FString& InChannelName, const TSharedRef<FAvaPlayableRemoteControlValues>& InRemoteControlValues, EAvaPlayableRCUpdateFlags InFlags = EAvaPlayableRCUpdateFlags::None);
 	bool PushPlaybackTransitionStartCommand(UAvaPlaybackTransition* InTransitionToStart);
 
 	/** This is used on the playback server to apply any pending commands to a playback instance. */
@@ -348,6 +348,7 @@ private:
 		struct FRemoteControlCommand
 		{
 			TSharedRef<FAvaPlayableRemoteControlValues> Values;
+			EAvaPlayableRCUpdateFlags UpdateFlags;
 		};
 	
 		TArray<FAnimationCommand> AnimationCommands;
diff --git a/Engine/Plugins/Experimental/Avalanche/Source/AvalancheMedia/Public/Playback/IAvaPlaybackClient.h b/Engine/Plugins/Experimental/Avalanche/Source/AvalancheMedia/Public/Playback/IAvaPlaybackClient.h
--- a/Engine/Plugins/Experimental/Avalanche/Source/AvalancheMedia/Public/Playback/IAvaPlaybackClient.h
+++ b/Engine/Plugins/Experimental/Avalanche/Source/AvalancheMedia/Public/Playback/IAvaPlaybackClient.h
@@ -49,7 +49,7 @@ public:
 	virtual void RequestPlayback(const FGuid& InInstanceId, const FSoftObjectPath& InAssetPath, const FString& InChannelName, EAvaPlaybackAction InAction, const FString& InArguments = FString()) = 0;
 	virtual void RequestAnimPlayback(const FGuid& InInstanceId, const FSoftObjectPath& InAssetPath, const FString& InChannelName, const FAvaPlaybackAnimPlaySettings& InAnimSettings) = 0;
 	virtual void RequestAnimAction(const FGuid& InInstanceId, const FSoftObjectPath& InAssetPath, const FString& InChannelName, const FString& InAnimationName, EAvaPlaybackAnimAction InAction) = 0;
-	virtual void RequestRemoteControlUpdate(const FGuid& InInstanceId, const FSoftObjectPath& InAssetPath, const FString& InChannelName, const FAvaPlayableRemoteControlValues& InRemoteControlValues) = 0;
+	virtual void RequestRemoteControlUpdate(const FGuid& InInstanceId, const FSoftObjectPath& InAssetPath, const FString& InChannelName, const FAvaPlayableRemoteControlValues& InRemoteControlValues, EAvaPlayableRCUpdateFlags InFlags = EAvaPlayableRCUpdateFlags::None) = 0;
 	virtual void RequestPlayableTransitionStart(const FGuid& InTransitionId, TArray<FGuid>&& InEnterInstanceIds, TArray<FGuid>&& InPlayingInstanceIds, TArray<FGuid>&& InExitInstanceIds, TArray<FAvaPlayableRemoteControlValues>&& InEnterValues, const FName& InChannelName, EAvaPlayableTransitionFlags InTransitionFlags) = 0;
 	virtual void RequestPlayableTransitionStop(const FGuid& InTransitionId, const FName& InChannelName) = 0;
 	virtual void RequestBroadcast(const FString& InProfile, const FName& InChannel, const TArray<UMediaOutput*>& InRemoteMediaOutputs, EAvaBroadcastAction InAction, const FString& InServerName = FString()) = 0;
