[MotionDesign] Fix undefined symbol error

diff --git a/Engine/Plugins/Experimental/Avalanche/Source/AvalancheSequencer/Private/NavigationTool/AvaNavigationToolProvider.cpp b/Engine/Plugins/Experimental/Avalanche/Source/AvalancheSequencer/Private/NavigationTool/AvaNavigationToolProvider.cpp
--- a/Engine/Plugins/Experimental/Avalanche/Source/AvalancheSequencer/Private/NavigationTool/AvaNavigationToolProvider.cpp
+++ b/Engine/Plugins/Experimental/Avalanche/Source/AvalancheSequencer/Private/NavigationTool/AvaNavigationToolProvider.cpp
@@ -148,8 +148,8 @@ void FAvaNavigationToolProvider::OnActivate()
 	ExtendToolToolBar();
 	ExtendToolItemContextMenu();
 
-	DragDropInitializedDelegate = UE::Sequencer::NavigationTool::FNavigationToolExtender::OnItemDragDropOpInitialized()
-		.AddSP(this, &FAvaNavigationToolProvider::ExtendToolItemDragDropOp);
+	//DragDropInitializedDelegate = UE::Sequencer::NavigationTool::FNavigationToolExtender::OnItemDragDropOpInitialized()
+	//	.AddSP(this, &FAvaNavigationToolProvider::ExtendToolItemDragDropOp);
 }
 
 void FAvaNavigationToolProvider::OnDeactivate()
@@ -157,8 +157,8 @@ void FAvaNavigationToolProvider::OnDeactivate()
 	RemoveToolToolBarExtension();
 	RemoveToolItemContextMenuExtension();
 
-	UE::Sequencer::NavigationTool::FNavigationToolExtender::OnItemDragDropOpInitialized()
-		.Remove(DragDropInitializedDelegate);
+	//UE::Sequencer::NavigationTool::FNavigationToolExtender::OnItemDragDropOpInitialized()
+	//	.Remove(DragDropInitializedDelegate);
 }
 
 void FAvaNavigationToolProvider::OnExtendColumns(UE::Sequencer::NavigationTool::FNavigationToolColumnExtender& OutExtender)
