[Motion Design] No easy way to adjust transforms of MD assets when used in a Persistent Level

- Adding load options for playback instances and playables allowing to implement different methods of positining the levels on load.
- Adding a SetTransform page command to directly specify a level transform.
- Adding a SetSpawnPoint page command to be able to spawn the graphics at specific tagged null actor positions in the persistent level.

diff --git a/Engine/Plugins/Experimental/Avalanche/Source/AvalancheMedia/Private/Playable/Playables/AvaPlayableLevelStreaming.cpp b/Engine/Plugins/Experimental/Avalanche/Source/AvalancheMedia/Private/Playable/Playables/AvaPlayableLevelStreaming.cpp
--- a/Engine/Plugins/Experimental/Avalanche/Source/AvalancheMedia/Private/Playable/Playables/AvaPlayableLevelStreaming.cpp
+++ b/Engine/Plugins/Experimental/Avalanche/Source/AvalancheMedia/Private/Playable/Playables/AvaPlayableLevelStreaming.cpp
@@ -5,13 +5,13 @@
 #include "AvaRemoteControlRebind.h"
 #include "AvaRemoteControlUtils.h"
 #include "AvaScene.h"
-#include "Camera/CameraActor.h"
 #include "Components/PrimitiveComponent.h"
 #include "Engine/Level.h"
 #include "Engine/LevelStreamingDynamic.h"
 #include "Engine/LocalPlayer.h"
-#include "Framework/AvaGameInstance.h"
+#include "EngineUtils.h"
 #include "Framework/AvaInstanceSettings.h"
+#include "Framework/AvaNullActor.h"
 #include "Framework/AvaSoftAssetPtr.h"
 #include "IAvaMediaModule.h"
 #include "IAvaRemoteControlInterface.h"
@@ -51,6 +51,121 @@ namespace UE::AvaMedia::LevelStreamingPlayable::Private
 		return AvaScene;
 	}
 
+	/**
+	 * Utility function to extract the level transform from load options.
+	 * returns true if the transform was specified, false otherwise.
+	 */
+	bool GetTransformFromOptions(FTransform& OutTransform, const FString& InOptions, const UAvaPlayableGroup* InPlayableGroup, const FAvaSoftAssetPtr& InSourceAsset)
+	{
+		// Read the transform from load options if available
+		FString TransformString;
+		if (FParse::Value(*InOptions,TEXT("Transform="), TransformString))
+		{
+			if (OutTransform.InitFromString(TransformString))
+			{
+				return true;
+			}
+		}
+
+		FString SpawnPointString;
+		if (FParse::Value(*InOptions,TEXT("SpawnPointTag="), SpawnPointString))
+		{
+			if (InPlayableGroup)
+			{
+				FName SpawnPointName = FName(*SpawnPointString);
+				if (UWorld* World = InPlayableGroup->GetPlayWorld())
+				{
+					for (AAvaNullActor* Actor : TActorRange<AAvaNullActor>(World))
+					{
+						if (Actor && Actor->ActorHasTag(SpawnPointName))
+						{
+							OutTransform = Actor->GetTransform();
+							return true;
+						}
+					}
+					// To help diagnose problems
+					UE_LOG(LogAvaPlayable, Verbose, TEXT("Loading Level [%s]: NullActor with tag \"%s\" was not found."),
+						*InSourceAsset.ToSoftObjectPath().ToString(), *SpawnPointString);
+				}
+			}
+			else
+			{
+				UE_LOG(LogAvaPlayable, Error,
+					TEXT("Loading Level [%s]: \"SpawnPoint\" option was specified but the playable doesn't have a valid playable group."),
+					*InSourceAsset.ToSoftObjectPath().ToString());
+			}
+		}
+
+		OutTransform = FTransform::Identity;
+		return false;
+	}
+
+	/**
+	 * Version of FLevelUtils::ApplyLevelTransform specialized to work with animated objects with sequencer.
+	 * All root actors get added to a root null and the root is moved.
+	 * Only actors that can't be attached get moved.
+	 * todo: We could try to exclude cameras from the transform.
+	 * In the standard setup, the camera under DefaultScene, along with other things we can to transform. 
+	 */
+	AActor* ApplyLevelTransform(ULevel* InLevel, const FTransform& InTransform, AActor* InExistingPivotActor)
+	{
+		// Create a root actor and put everything under it.
+		UWorld* const World = InLevel->GetWorld();
+		if (!IsValid(World))
+		{
+			return nullptr;
+		}
+
+		AActor* PivotActor = InExistingPivotActor;
+		if (!PivotActor)
+		{
+			FActorSpawnParameters SpawnParams;
+			SpawnParams.OverrideLevel = InLevel;
+			SpawnParams.SpawnCollisionHandlingOverride = ESpawnActorCollisionHandlingMethod::AlwaysSpawn;
+			SpawnParams.bNoFail = true;
+			PivotActor = World->SpawnActor<AAvaNullActor>(SpawnParams);
+
+			if (!PivotActor)
+			{
+				UE_LOG(LogAvaPlayable, Error, TEXT("Failed to create a pivot actor to transform the level."));
+				return nullptr;
+			}
+		}
+
+		for (AActor* Actor : InLevel->Actors)
+		{
+			// Exclude the cameras from the transformation
+			if (Actor && Actor != PivotActor)
+			{
+				// Only attach root actors.
+				if (Actor->GetAttachParentActor())
+				{
+					continue;
+				}
+
+				// TODO: Have a generic way to tag actors we want to exclude from the level transform (like camera).
+
+				if (!Actor->AttachToActor(PivotActor, FAttachmentTransformRules::KeepRelativeTransform))
+				{
+					// In case we fail to attach, apply the transformation (as in FLevelUtils::ApplyLevelTransform)
+					if (USceneComponent* RootComponent = Actor->GetRootComponent())
+					{
+						RootComponent->SetRelativeLocation_Direct(InTransform.TransformPosition(RootComponent->GetRelativeLocation()));
+						RootComponent->SetRelativeRotation_Direct(InTransform.TransformRotation(RootComponent->GetRelativeRotation().Quaternion()).Rotator());
+						RootComponent->SetRelativeScale3D_Direct(InTransform.GetScale3D() * RootComponent->GetRelativeScale3D());
+
+						// Any components which have cached their bounds will not be accurate after a level transform is applied. Force them to recompute the bounds once more.
+						Actor->MarkNeedsRecomputeBoundsOnceForGame();
+					}
+				}
+			}
+		}
+
+		PivotActor->SetActorTransform(InTransform);
+
+		return PivotActor;
+	}
+
 	EAvaPlayableStatus GetPlayableStatusFromLevelStreamingState(ELevelStreamingState InState, bool bInShouldBeLoaded)
 	{
 		switch (InState)
@@ -155,7 +270,7 @@ namespace UE::AvaMedia::LevelStreamingPlayable::Private
 	}
 }
 
-bool UAvaPlayableLevelStreaming::LoadAsset(const FAvaSoftAssetPtr& InSourceAsset, bool bInInitiallyVisible)
+bool UAvaPlayableLevelStreaming::LoadAsset(const FAvaSoftAssetPtr& InSourceAsset, bool bInInitiallyVisible, const FString& InLoadOptions)
 {
 	if (!PlayableGroup)
 	{
@@ -165,11 +280,15 @@ bool UAvaPlayableLevelStreaming::LoadAsset(const FAvaSoftAssetPtr& InSourceAsset
 	// Ensure world is created. Does nothing if already created.
 	PlayableGroup->ConditionalCreateWorld();
 
+	// Setup the level transform from the load options.
+	bHasTransform = UE::AvaMedia::LevelStreamingPlayable::Private::GetTransformFromOptions(LevelTransform, InLoadOptions, PlayableGroup, InSourceAsset);
+
 	const FAvaInstanceSettings& PlaybackInstanceSettings = IAvaMediaModule::Get().GetAvaInstanceSettings();
 	bLoadSubPlayables = PlaybackInstanceSettings.bEnableLoadSubPlayables;
 	
 	check(InSourceAsset.GetAssetType() == EMotionDesignAssetType::World);
-	const bool bAssetLoading = LoadLevel(TSoftObjectPtr<UWorld>(InSourceAsset.ToSoftObjectPath()), bInInitiallyVisible); 
+	// Remark: We are not using the level streaming transform because it doesn't work with animated objects.
+	const bool bAssetLoading = LoadLevel(TSoftObjectPtr<UWorld>(InSourceAsset.ToSoftObjectPath()), FTransform::Identity, bInInitiallyVisible);
 	if (bAssetLoading)
 	{
 		// Refresh status immediately (assumes it won't be loaded nor visible immediately, so shouldn't affect TL)
@@ -317,7 +436,7 @@ void UAvaPlayableLevelStreaming::SetupView(FSceneViewFamily& InViewFamily, FScen
 	}
 }
 
-bool UAvaPlayableLevelStreaming::LoadLevel(const TSoftObjectPtr<UWorld>& InSourceLevel, bool bInInitiallyVisible)
+bool UAvaPlayableLevelStreaming::LoadLevel(const TSoftObjectPtr<UWorld>& InSourceLevel, const FTransform& InTransform, bool bInInitiallyVisible)
 {
 	if (SourceLevel == InSourceLevel)
 	{
@@ -331,11 +450,8 @@ bool UAvaPlayableLevelStreaming::LoadLevel(const TSoftObjectPtr<UWorld>& InSourc
 	}
 	
 	bool bSuccess = false;
-	const FVector Origin(0,0,0);
 
-	ULevelStreamingDynamic::FLoadLevelInstanceParams Params(PlayableGroup->GetPlayWorld()
-			, InSourceLevel.GetLongPackageName()
-			, FTransform(FRotator::ZeroRotator, Origin));
+	ULevelStreamingDynamic::FLoadLevelInstanceParams Params(PlayableGroup->GetPlayWorld(), InSourceLevel.GetLongPackageName(), InTransform);
 	Params.bLoadAsTempPackage = true;
 	Params.bInitiallyVisible  = bInInitiallyVisible;
 
@@ -423,7 +539,7 @@ void UAvaPlayableLevelStreaming::OnLevelStreamingStateChanged_Synchronized(ELeve
 	}
 	else if (InNewState == ELevelStreamingState::LoadedNotVisible || InNewState == ELevelStreamingState::LoadedVisible)
 	{
-		const ULevel* const Level = LevelStreaming->GetLoadedLevel();
+		ULevel* Level = LevelStreaming->GetLoadedLevel();
 		if (!IsValid(Level))
 		{
 			return;
@@ -446,6 +562,12 @@ void UAvaPlayableLevelStreaming::OnLevelStreamingStateChanged_Synchronized(ELeve
 		{
 			FAvaPlaybackUtils::FlushPackageLoading(Level->GetPackage());
 		}
+
+		if (bHasTransform && !bTransformApplied)
+		{
+			PivotActorForTransform = UE::AvaMedia::LevelStreamingPlayable::Private::ApplyLevelTransform(Level, LevelTransform, PivotActorForTransform.Get());						
+			bTransformApplied = true;
+		}
 		
 		// Resolve the ava scene for the other operations.
 		ResolveScene(Level);
@@ -723,7 +845,7 @@ void UAvaPlayableLevelStreaming::GetOrLoadSubPlayable(const ULevelStreaming* InL
 	{
 		// TODO: Propagate more stuff from InLevelStreaming. Needs to reach LoadLevel.
 		const FAvaSoftAssetPtr AssetPtr = { UWorld::StaticClass(), TSoftObjectPtr<UObject>(SourceAssetPath)};
-		if (NewPlayable->LoadAsset(AssetPtr, GetShouldBeVisible()))
+		if (NewPlayable->LoadAsset(AssetPtr, GetShouldBeVisible(), FString()))
 		{
 			AddSubPlayable(NewPlayable);
 		}
diff --git a/Engine/Plugins/Experimental/Avalanche/Source/AvalancheMedia/Private/Playable/Playables/AvaPlayableLevelStreaming.h b/Engine/Plugins/Experimental/Avalanche/Source/AvalancheMedia/Private/Playable/Playables/AvaPlayableLevelStreaming.h
--- a/Engine/Plugins/Experimental/Avalanche/Source/AvalancheMedia/Private/Playable/Playables/AvaPlayableLevelStreaming.h
+++ b/Engine/Plugins/Experimental/Avalanche/Source/AvalancheMedia/Private/Playable/Playables/AvaPlayableLevelStreaming.h
@@ -19,7 +19,7 @@ class UAvaPlayableLevelStreaming : public UAvaPlayable
 	GENERATED_BODY()
 public:
 	//~ Begin UAvaPlayable
-	virtual bool LoadAsset(const FAvaSoftAssetPtr& InSourceAsset, bool bInInitiallyVisible) override;
+	virtual bool LoadAsset(const FAvaSoftAssetPtr& InSourceAsset, bool bInInitiallyVisible, const FString& InLoadOptions) override;
 	virtual bool UnloadAsset() override;
 	virtual const FSoftObjectPath& GetSourceAssetPath() const override { return SourceLevel.ToSoftObjectPath(); }
 	virtual EAvaPlayableStatus GetPlayableStatus() const override { return PlayableStatus; }
@@ -46,7 +46,7 @@ public:
 	bool GetShouldBeHidden() const { return bShouldBeHidden; }
 
 protected:
-	bool LoadLevel(const TSoftObjectPtr<UWorld>& InSourceLevel, bool bInInitiallyVisible);
+	bool LoadLevel(const TSoftObjectPtr<UWorld>& InSourceLevel, const FTransform& InTransform, bool bInInitiallyVisible);
 
 	void HandleTransitionEvent(UAvaPlayable* InPlayable, UAvaPlayableTransition* InTransition, EAvaPlayableTransitionEventFlags InTransitionFlags);
 
@@ -123,4 +123,16 @@ protected:
 
 	/** If true, all primitives from the playable will hidden. */
 	bool bShouldBeHidden = false;
+
+	/** True if the load command specified a transform. */
+	bool bHasTransform = false;
+
+	/** Keep track if the transform was applied to avoid doing it multiple times during level streaming. */
+	bool bTransformApplied = false;
+
+	/** Keep a pointer to the pivot actor used for transform in case we want to change the transform later on (for reuse). */
+	TWeakObjectPtr<AActor> PivotActorForTransform;
+
+	/** Transform to apply to the actors at the end of level streaming. */
+	FTransform LevelTransform;
 };
diff --git a/Engine/Plugins/Experimental/Avalanche/Source/AvalancheMedia/Private/Playable/Playables/AvaPlayableRemoteProxy.cpp b/Engine/Plugins/Experimental/Avalanche/Source/AvalancheMedia/Private/Playable/Playables/AvaPlayableRemoteProxy.cpp
--- a/Engine/Plugins/Experimental/Avalanche/Source/AvalancheMedia/Private/Playable/Playables/AvaPlayableRemoteProxy.cpp
+++ b/Engine/Plugins/Experimental/Avalanche/Source/AvalancheMedia/Private/Playable/Playables/AvaPlayableRemoteProxy.cpp
@@ -12,7 +12,7 @@
 
 #define LOCTEXT_NAMESPACE "AvaPlayableRemoteProxy"
 
-bool UAvaPlayableRemoteProxy::LoadAsset(const FAvaSoftAssetPtr& InSourceAsset, bool bInInitiallyVisible)
+bool UAvaPlayableRemoteProxy::LoadAsset(const FAvaSoftAssetPtr& InSourceAsset, bool bInInitiallyVisible, const FString& InLoadOptions)
 {
 	if (!PlayableGroup)
 	{
@@ -57,7 +57,7 @@ bool UAvaPlayableRemoteProxy::LoadAsset(const FAvaSoftAssetPtr& InSourceAsset, b
 	{
 		// Todo(opt): Combine actions: LoadWithUserData
 		// Server-side requirement: a load request on an already loaded/playing asset will not do anything.
-		PlaybackClient.RequestPlayback(InstanceId, SourceAssetPath, PlayingChannelName, EAvaPlaybackAction::Load);
+		PlaybackClient.RequestPlayback(InstanceId, SourceAssetPath, PlayingChannelName, EAvaPlaybackAction::Load, InLoadOptions);
 		PlaybackClient.RequestPlayback(InstanceId, SourceAssetPath, PlayingChannelName, EAvaPlaybackAction::SetUserData, UserData);
 	}
 
diff --git a/Engine/Plugins/Experimental/Avalanche/Source/AvalancheMedia/Private/Playable/Playables/AvaPlayableRemoteProxy.h b/Engine/Plugins/Experimental/Avalanche/Source/AvalancheMedia/Private/Playable/Playables/AvaPlayableRemoteProxy.h
--- a/Engine/Plugins/Experimental/Avalanche/Source/AvalancheMedia/Private/Playable/Playables/AvaPlayableRemoteProxy.h
+++ b/Engine/Plugins/Experimental/Avalanche/Source/AvalancheMedia/Private/Playable/Playables/AvaPlayableRemoteProxy.h
@@ -22,7 +22,7 @@ public:
 	FName GetPlayingChannelFName() const { return PlayingChannelFName; }
 	
 	//~ Begin UAvaPlayable
-	virtual bool LoadAsset(const FAvaSoftAssetPtr& InSourceAsset, bool bInInitiallyVisible) override;
+	virtual bool LoadAsset(const FAvaSoftAssetPtr& InSourceAsset, bool bInInitiallyVisible, const FString& InLoadOptions) override;
 	virtual bool UnloadAsset() override;
 	virtual const FSoftObjectPath& GetSourceAssetPath() const override { return SourceAssetPath; }
 	virtual EAvaPlayableStatus GetPlayableStatus() const override;
diff --git a/Engine/Plugins/Experimental/Avalanche/Source/AvalancheMedia/Private/Playback/AvaPlaybackGraph.cpp b/Engine/Plugins/Experimental/Avalanche/Source/AvalancheMedia/Private/Playback/AvaPlaybackGraph.cpp
--- a/Engine/Plugins/Experimental/Avalanche/Source/AvalancheMedia/Private/Playback/AvaPlaybackGraph.cpp
+++ b/Engine/Plugins/Experimental/Avalanche/Source/AvalancheMedia/Private/Playback/AvaPlaybackGraph.cpp
@@ -25,6 +25,19 @@
 
 DEFINE_LOG_CATEGORY(LogAvaPlayback);
 
+namespace UE::AvaPlaybackGraph::Private
+{
+	const FString& GetLoadOptions(UAvaPlaybackNode* InPlaybackNode)
+	{
+		if (UAvaPlaybackNodePlayer* PlayerNode = Cast<UAvaPlaybackNodePlayer>(InPlaybackNode))
+		{
+			return PlayerNode->GetLoadOptions();
+		}
+		static FString Empty;
+		return Empty;
+	}
+}
+
 UAvaPlaybackGraph::~UAvaPlaybackGraph()
 {
 	FAvaBroadcastOutputChannel::GetOnChannelChanged().RemoveAll(this);
@@ -51,9 +64,9 @@ void UAvaPlaybackGraph::LoadInstances()
 	for (TPair<FName, FAvaPlaybackChannelParameters>& Playback : PlaybackSettings)
 	{
 		// Load the assets
-		for (const FAvaSoftAssetPtr& Asset : Playback.Value.Assets)
+		for (const FAvaPlaybackPlayerParameters& Player : Playback.Value.Players)
 		{
-			FindOrLoadPlayable(Asset, ResolveChannelName(Playback.Key));
+			FindOrLoadPlayable(Player.Asset, ResolveChannelName(Playback.Key), UE::AvaPlaybackGraph::Private::GetLoadOptions(Player.Node.Get()));
 		}
 	}
 }
@@ -181,7 +194,7 @@ UAvaPlayable* UAvaPlaybackGraph::FindPlayable(const FSoftObjectPath& InSourceAss
 	return ChannelPlayableGroup ? ChannelPlayableGroup->FindPlayable(InSourceAssetPath) : nullptr;
 }
 
-UAvaPlayable* UAvaPlaybackGraph::FindOrLoadPlayable(const FAvaSoftAssetPtr& InSourceAsset, const FName& InChannelName)
+UAvaPlayable* UAvaPlaybackGraph::FindOrLoadPlayable(const FAvaSoftAssetPtr& InSourceAsset, const FName& InChannelName, const FString& InLoadOptions)
 {
 	FAvaPlaybackPlayableGroup* ChannelPlayableGroup = ChannelPlayableGroups.Find(InChannelName);
 	if (!ChannelPlayableGroup)
@@ -200,7 +213,7 @@ UAvaPlayable* UAvaPlaybackGraph::FindOrLoadPlayable(const FAvaSoftAssetPtr& InSo
 	{
 		ChannelPlayableGroup->Playables.Add(InSourceAsset.ToSoftObjectPath(), NewPlayable);
 		OnPlayableCreated.Broadcast(this, NewPlayable);
-		NewPlayable->LoadAsset(InSourceAsset, false);
+		NewPlayable->LoadAsset(InSourceAsset, false, InLoadOptions);
 	}
 	return NewPlayable;
 }
@@ -242,9 +255,9 @@ TArray<UAvaGameInstance*> UAvaPlaybackGraph::GetActiveGameInstances() const
 		for (const TPair<FName, FAvaPlaybackChannelParameters>& Playback : PlaybackSettings)
 		{
 			FName ChannelName = ResolveChannelName(Playback.Key);
-			for (const FAvaSoftAssetPtr& Asset : Playback.Value.Assets)
+			for (const FAvaPlaybackPlayerParameters& Player : Playback.Value.Players)
 			{
-				const UAvaPlayable* Playable = FindPlayable(Asset.ToSoftObjectPath(), ChannelName);
+				const UAvaPlayable* Playable = FindPlayable(Player.Asset.ToSoftObjectPath(), ChannelName);
 				if (Playable && Playable->GetPlayableGroup())
 				{
 					if (UAvaGameInstance* AvaGameInstance = Cast<UAvaGameInstance>(Playable->GetPlayableGroup()->GetGameInstance()))
@@ -447,9 +460,9 @@ void UAvaPlaybackGraph::StopPlayable(const FSoftObjectPath& InSourceAssetPath, c
 	}
 }
 
-void UAvaPlaybackGraph::LoadAsset(const FAvaSoftAssetPtr& InSourceAsset, const FName& InChannelName)
+void UAvaPlaybackGraph::LoadAsset(const FAvaSoftAssetPtr& InSourceAsset, const FName& InChannelName, const FString& InLoadOptions)
 {
-	FindOrLoadPlayable(InSourceAsset, InChannelName);
+	FindOrLoadPlayable(InSourceAsset, InChannelName, InLoadOptions);
 }
 
 void UAvaPlaybackGraph::Tick(float DeltaTime)
@@ -471,7 +484,7 @@ void UAvaPlaybackGraph::Tick(float DeltaTime)
 	//Pass 1: Tick from Root (Channels Node) up to the Player Nodes.
 	for (TPair<FName, FAvaPlaybackChannelParameters>& Playback : PlaybackSettings)
 	{
-		Playback.Value.Assets.Reset();
+		Playback.Value.Players.Reset();
 	}
 	RootNode->TickRoot(DeltaTime, PlaybackSettings);
 
@@ -481,26 +494,32 @@ void UAvaPlaybackGraph::Tick(float DeltaTime)
 		const FName ChannelName = ResolveChannelName(Playback.Key);
 		FAvaBroadcastOutputChannel& Channel = UAvaBroadcast::Get().GetCurrentProfile().GetChannelMutable(ChannelName);
 
-		for (const FAvaSoftAssetPtr& Asset : Playback.Value.Assets)
+		for (const FAvaPlaybackPlayerParameters& Player : Playback.Value.Players)
 		{
 			if (IsPreviewOnly())
 			{
-				RefreshPreview(Asset, Channel, ChannelName);
+				RefreshPreview(Player.Asset, UE::AvaPlaybackGraph::Private::GetLoadOptions(Player.Node.Get()), Channel, ChannelName);
 			}
 			else if (Channel.IsValidChannel())
 			{
-				RefreshPlayback(Asset, Channel, ChannelName);
+				RefreshPlayback(Player.Asset, UE::AvaPlaybackGraph::Private::GetLoadOptions(Player.Node.Get()), Channel, ChannelName);
 			}
 		}
 
 		// Check if the asset playing has changed (this may be the result of a switch/combiner node).
 		if (const FAvaPlaybackChannelParameters* PreviousParameters = PreviousPlaybackSettings.Find(Playback.Key))
 		{
-			for (const FAvaSoftAssetPtr& PreviousAsset : PreviousParameters->Assets)
+			for (const FAvaPlaybackPlayerParameters& PreviousPlayer : PreviousParameters->Players)
 			{
-				if (Playback.Value.Assets.Find(PreviousAsset) == INDEX_NONE)
+				// We still don't support multi-instance within a playback graph. The asset remains the key.
+				auto IsSameAsset = [&PreviousPlayer](const FAvaPlaybackPlayerParameters& InPlayer)
+				{
+					return PreviousPlayer.Asset == InPlayer.Asset;
+				};
+				
+				if (!Playback.Value.Players.FindByPredicate(IsSameAsset))
 				{
-					StopPlayable(PreviousAsset.ToSoftObjectPath(), ChannelName);
+					StopPlayable(PreviousPlayer.Asset.ToSoftObjectPath(), ChannelName);
 				}
 			}
 		}
@@ -773,9 +792,9 @@ void UAvaPlaybackGraph::OnChannelBroadcastStateChanged(const FAvaBroadcastOutput
 			if (ChannelName == ChangedChannelName)
 			{
 				FAvaBroadcastOutputChannel& Channel = Profile.GetChannelMutable(ChannelName);
-				for (const FAvaSoftAssetPtr& Asset : Playback.Value.Assets)
+				for (const FAvaPlaybackPlayerParameters& Player : Playback.Value.Players)
 				{
-					RefreshPlayback(Asset, Channel, ChannelName);
+					RefreshPlayback(Player.Asset, UE::AvaPlaybackGraph::Private::GetLoadOptions(Player.Node.Get()), Channel, ChannelName);
 				}
 			}
 		}
@@ -818,7 +837,7 @@ bool UAvaPlaybackGraph::CanRefreshPlayback(const FAvaSoftAssetPtr& InSourceAsset
 	return IsPlaying() && InChannel.IsValidChannel() && !InSourceAsset.IsNull();
 }
 
-bool UAvaPlaybackGraph::RefreshPlayback(const FAvaSoftAssetPtr& InSourceAsset, FAvaBroadcastOutputChannel& InChannel, const FName& InChannelName)
+bool UAvaPlaybackGraph::RefreshPlayback(const FAvaSoftAssetPtr& InSourceAsset, const FString& InLoadOptions, FAvaBroadcastOutputChannel& InChannel, const FName& InChannelName)
 {
 	if (!CanRefreshPlayback(InSourceAsset, InChannel))
 	{
@@ -827,7 +846,7 @@ bool UAvaPlaybackGraph::RefreshPlayback(const FAvaSoftAssetPtr& InSourceAsset, F
 		return false;
 	}
 	
-	if (UAvaPlayable* const Playable = FindOrLoadPlayable(InSourceAsset, InChannelName))
+	if (UAvaPlayable* const Playable = FindOrLoadPlayable(InSourceAsset, InChannelName, InLoadOptions))
 	{
 		UTextureRenderTarget2D* const RenderTarget = UpdatePlaybackRenderTarget(Playable, InChannel);
 
@@ -857,12 +876,12 @@ bool UAvaPlaybackGraph::CanRefreshPreview(const FAvaSoftAssetPtr& InSourceAsset)
 	return IsPlaying() && !InSourceAsset.IsNull();
 }
 
-bool UAvaPlaybackGraph::RefreshPreview(const FAvaSoftAssetPtr& InSourceAsset, FAvaBroadcastOutputChannel& InChannel, const FName& InChannelName)
+bool UAvaPlaybackGraph::RefreshPreview(const FAvaSoftAssetPtr& InSourceAsset, const FString& InLoadOptions, FAvaBroadcastOutputChannel& InChannel, const FName& InChannelName)
 {
 	// If the preview channel has been setup, go through the normal channel refresh.
 	if (InChannel.IsValidChannel())
 	{
-		return RefreshPlayback(InSourceAsset, InChannel, InChannelName);
+		return RefreshPlayback(InSourceAsset, InLoadOptions, InChannel, InChannelName);
 	}
 
 	if (!CanRefreshPreview(InSourceAsset))
@@ -870,7 +889,7 @@ bool UAvaPlaybackGraph::RefreshPreview(const FAvaSoftAssetPtr& InSourceAsset, FA
 		return false;
 	}
 
-	UAvaPlayable* const Playable = FindOrLoadPlayable(InSourceAsset, InChannelName);
+	UAvaPlayable* const Playable = FindOrLoadPlayable(InSourceAsset, InChannelName, InLoadOptions);
 
 	if (!Playable || !Playable->GetPlayableGroup())
 	{
diff --git a/Engine/Plugins/Experimental/Avalanche/Source/AvalancheMedia/Private/Playback/AvaPlaybackManager.cpp b/Engine/Plugins/Experimental/Avalanche/Source/AvalancheMedia/Private/Playback/AvaPlaybackManager.cpp
--- a/Engine/Plugins/Experimental/Avalanche/Source/AvalancheMedia/Private/Playback/AvaPlaybackManager.cpp
+++ b/Engine/Plugins/Experimental/Avalanche/Source/AvalancheMedia/Private/Playback/AvaPlaybackManager.cpp
@@ -232,13 +232,11 @@ TSharedPtr<FAvaPlaybackInstance> FAvaPlaybackManager::AcquirePlaybackInstance(co
 	return TSharedPtr<FAvaPlaybackInstance>();
 }
 
-// Package name "/Game/AvaPlayback"
-// Asset name "AvaPlayback"
-TSharedPtr<FAvaPlaybackInstance> FAvaPlaybackManager::LoadPlaybackInstance(const FSoftObjectPath& InAssetPath, const FString& InChannelName)
+TSharedPtr<FAvaPlaybackInstance> FAvaPlaybackManager::LoadPlaybackInstance(const FSoftObjectPath& InAssetPath, const FString& InChannelName, const FString& InLoadOptions)
 {
 	TSharedPtr<FAvaPlaybackInstance> PlaybackInstance;
 	
-	if (UAvaPlaybackGraph* Playback = LoadPlaybackObject(InAssetPath, InChannelName))
+	if (UAvaPlaybackGraph* Playback = LoadPlaybackObject(InAssetPath, InChannelName, InLoadOptions))
 	{
 		FGuid NewInstanceId = FGuid::NewGuid();
 		
@@ -253,12 +251,12 @@ TSharedPtr<FAvaPlaybackInstance> FAvaPlaybackManager::LoadPlaybackInstance(const
 	return PlaybackInstance;
 }
 
-TSharedPtr<FAvaPlaybackInstance> FAvaPlaybackManager::AcquireOrLoadPlaybackInstance(const FSoftObjectPath& InAssetPath, const FString& InChannelName)
+TSharedPtr<FAvaPlaybackInstance> FAvaPlaybackManager::AcquireOrLoadPlaybackInstance(const FSoftObjectPath& InAssetPath, const FString& InChannelName, const FString& InLoadOptions)
 {
 	TSharedPtr<FAvaPlaybackInstance> PlaybackInstance = AcquirePlaybackInstance(InAssetPath, InChannelName);
 	if (!PlaybackInstance)
 	{
-		PlaybackInstance = LoadPlaybackInstance(InAssetPath, InChannelName);
+		PlaybackInstance = LoadPlaybackInstance(InAssetPath, InChannelName, InLoadOptions);
 	}
 	return PlaybackInstance;
 }
@@ -356,7 +354,7 @@ void FAvaPlaybackManager::InvalidatePlaybackAssetEntry(const FSoftObjectPath& In
 	PlaybackAssetEntries.Remove(InAssetPath);
 }
 
-UAvaPlaybackGraph* FAvaPlaybackManager::LoadPlaybackObject(const FSoftObjectPath& InAssetPath, const FString& InChannelName) const
+UAvaPlaybackGraph* FAvaPlaybackManager::LoadPlaybackObject(const FSoftObjectPath& InAssetPath, const FString& InChannelName, const FString& InLoadOptions) const
 {
 	TRACE_CPUPROFILER_EVENT_SCOPE(FAvaPlaybackManager::LoadPlaybackObject);
 	
@@ -366,7 +364,7 @@ UAvaPlaybackGraph* FAvaPlaybackManager::LoadPlaybackObject(const FSoftObjectPath
 	// The generated playback graph will load it async using a level streaming playable.
 	if (FAvaPlaybackUtils::IsMapAsset(PackageName))
 	{
-		UAvaPlaybackGraph* const AvaPlayback = BuildPlaybackFromWorld(TSoftObjectPtr<UWorld>(InAssetPath), InChannelName);
+		UAvaPlaybackGraph* const AvaPlayback = BuildPlaybackFromWorld(TSoftObjectPtr<UWorld>(InAssetPath), InChannelName, InLoadOptions);
 		check(AvaPlayback);
 		return AvaPlayback;
 	}
@@ -383,6 +381,7 @@ UAvaPlaybackGraph* FAvaPlaybackManager::LoadPlaybackObject(const FSoftObjectPath
 	}
 	
 	// When the asset is a Playback Graph, it is used directly.
+	// Warning: this is not an instance, but the asset itself. Will not work for multi-instancing (TL).
 	if (UAvaPlaybackGraph* const AvaPlayback = Cast<UAvaPlaybackGraph>(LoadedAsset))
 	{
 		return AvaPlayback;
@@ -390,7 +389,7 @@ UAvaPlaybackGraph* FAvaPlaybackManager::LoadPlaybackObject(const FSoftObjectPath
 
 	if (UWorld* const World = Cast<UWorld>(LoadedAsset))
 	{
-		UAvaPlaybackGraph* const AvaPlayback = BuildPlaybackFromWorld(World, InChannelName);
+		UAvaPlaybackGraph* const AvaPlayback = BuildPlaybackFromWorld(World, InChannelName, InLoadOptions);
 		check(AvaPlayback);
 		return AvaPlayback;
 	}
@@ -402,7 +401,7 @@ UAvaPlaybackGraph* FAvaPlaybackManager::LoadPlaybackObject(const FSoftObjectPath
 	return nullptr;
 }
 
-UAvaPlaybackGraph* FAvaPlaybackManager::BuildPlaybackFromWorld(const TSoftObjectPtr<UWorld>& InWorld, const FString& InChannelName) const
+UAvaPlaybackGraph* FAvaPlaybackManager::BuildPlaybackFromWorld(const TSoftObjectPtr<UWorld>& InWorld, const FString& InChannelName, const FString& InLoadOptions) const
 {
 	using namespace UE::AvaPlaybackManager::Private;
 	FAvaPlaybackGraphBuilder GraphBuilder(GetPlayableGroupManager());
@@ -410,6 +409,7 @@ UAvaPlaybackGraph* FAvaPlaybackManager::BuildPlaybackFromWorld(const TSoftObject
 	// Construct Player Node and assign the World that we passed in
 	UAvaPlaybackNodeLevelPlayer* const PlayerNode = GraphBuilder.ConstructPlaybackNode<UAvaPlaybackNodeLevelPlayer>();
 	PlayerNode->SetAsset(InWorld);
+	PlayerNode->SetLoadOptions(InLoadOptions);
 	
 	GraphBuilder.ConnectToRoot(InChannelName, PlayerNode);
 	return GraphBuilder.FinishBuilding();
diff --git a/Engine/Plugins/Experimental/Avalanche/Source/AvalancheMedia/Private/Playback/AvaPlaybackServer.cpp b/Engine/Plugins/Experimental/Avalanche/Source/AvalancheMedia/Private/Playback/AvaPlaybackServer.cpp
--- a/Engine/Plugins/Experimental/Avalanche/Source/AvalancheMedia/Private/Playback/AvaPlaybackServer.cpp
+++ b/Engine/Plugins/Experimental/Avalanche/Source/AvalancheMedia/Private/Playback/AvaPlaybackServer.cpp
@@ -962,7 +962,7 @@ void FAvaPlaybackServer::StartPlaybackCommand(const TArray<FString>& InArgs)
 		const FSoftObjectPath AssetPath(InArgs[0]);
 
 		// Uncached load.
-		if (UAvaPlaybackGraph* PlaybackObject = Manager->LoadPlaybackObject(AssetPath, ChannelName))
+		if (UAvaPlaybackGraph* PlaybackObject = Manager->LoadPlaybackObject(AssetPath, ChannelName, FString()))
 		{
 			PlaybackObject->Play();
 			StartBroadcast();
@@ -1311,11 +1311,11 @@ void FAvaPlaybackServer::ExecutePendingPlaybackCommands(const FDateTime& InUtcNo
 			break;
 			
 		case EAvaPlaybackAction::Load:
-			LoadPlayback(ReplyTo, Command.InstanceId, Command.ChannelName, Command.AssetPath);
+			LoadPlayback(ReplyTo, Command.InstanceId, Command.ChannelName, Command.AssetPath, Command.Arguments);
 			break;
 			
 		case EAvaPlaybackAction::Start:
-			StartPlayback(ReplyTo, Command.InstanceId, Command.ChannelName, Command.AssetPath);
+			StartPlayback(ReplyTo, Command.InstanceId, Command.ChannelName, Command.AssetPath, Command.Arguments);
 			break;
 			
 		case EAvaPlaybackAction::Stop:
@@ -1367,7 +1367,7 @@ void FAvaPlaybackServer::ExecutePendingPlaybackCommands(const FDateTime& InUtcNo
 	}
 }
 
-TSharedPtr<FAvaPlaybackInstance> FAvaPlaybackServer::GetOrLoadPlaybackInstance(const FGuid& InInstanceId, const FString& InChannelName, const FSoftObjectPath& InAssetPath)
+TSharedPtr<FAvaPlaybackInstance> FAvaPlaybackServer::GetOrLoadPlaybackInstance(const FGuid& InInstanceId, const FString& InChannelName, const FSoftObjectPath& InAssetPath, const FString& InLoadOptions)
 {
 	// Check if loaded locally by the server under the same Id.
 	TSharedPtr<FAvaPlaybackInstance> PlaybackInstance = FindActivePlaybackInstance(InInstanceId);
@@ -1400,7 +1400,7 @@ TSharedPtr<FAvaPlaybackInstance> FAvaPlaybackServer::GetOrLoadPlaybackInstance(c
 	}
 
 	// Load it or acquire a cached recycled asset.
-	PlaybackInstance = Manager->AcquireOrLoadPlaybackInstance(InAssetPath, InChannelName);
+	PlaybackInstance = Manager->AcquireOrLoadPlaybackInstance(InAssetPath, InChannelName, InLoadOptions);
 	
 	if (PlaybackInstance.IsValid())
 	{
@@ -1413,7 +1413,7 @@ TSharedPtr<FAvaPlaybackInstance> FAvaPlaybackServer::GetOrLoadPlaybackInstance(c
 	return PlaybackInstance;	
 }
 
-void FAvaPlaybackServer::LoadPlayback(const FMessageAddress& InReplyToAddress, const FGuid& InInstanceId, const FString& InChannelName, const FSoftObjectPath& InAssetPath)
+void FAvaPlaybackServer::LoadPlayback(const FMessageAddress& InReplyToAddress, const FGuid& InInstanceId, const FString& InChannelName, const FSoftObjectPath& InAssetPath, const FString& InLoadOptions)
 {
 	if (!InAssetPath.IsValid())
 	{
@@ -1422,7 +1422,7 @@ void FAvaPlaybackServer::LoadPlayback(const FMessageAddress& InReplyToAddress, c
 	}
 	else
 	{
-		if (const TSharedPtr<FAvaPlaybackInstance> PlaybackInstance = GetOrLoadPlaybackInstance(InInstanceId, InChannelName, InAssetPath))
+		if (const TSharedPtr<FAvaPlaybackInstance> PlaybackInstance = GetOrLoadPlaybackInstance(InInstanceId, InChannelName, InAssetPath, InLoadOptions))
 		{
 			if (PlaybackInstance->GetPlayback())
 			{
@@ -1439,7 +1439,7 @@ void FAvaPlaybackServer::LoadPlayback(const FMessageAddress& InReplyToAddress, c
 	}
 }
 
-void FAvaPlaybackServer::StartPlayback(const FMessageAddress& InReplyToAddress, const FGuid& InInstanceId, const FString& InChannelName, const FSoftObjectPath& InAssetPath)
+void FAvaPlaybackServer::StartPlayback(const FMessageAddress& InReplyToAddress, const FGuid& InInstanceId, const FString& InChannelName, const FSoftObjectPath& InAssetPath, const FString& InLoadOptions)
 {
 	if (!InAssetPath.IsValid())
 	{
@@ -1449,7 +1449,7 @@ void FAvaPlaybackServer::StartPlayback(const FMessageAddress& InReplyToAddress,
 	}
 	else
 	{
-		if (const TSharedPtr<FAvaPlaybackInstance> PlaybackInstance = GetOrLoadPlaybackInstance(InInstanceId, InChannelName, InAssetPath))
+		if (const TSharedPtr<FAvaPlaybackInstance> PlaybackInstance = GetOrLoadPlaybackInstance(InInstanceId, InChannelName, InAssetPath, InLoadOptions))
 		{
 			PlaybackInstance->GetPlayback()->Play();
 			PlaybackInstance->SetStatus(EAvaPlaybackStatus::Starting);
diff --git a/Engine/Plugins/Experimental/Avalanche/Source/AvalancheMedia/Private/Playback/AvaPlaybackServer.h b/Engine/Plugins/Experimental/Avalanche/Source/AvalancheMedia/Private/Playback/AvaPlaybackServer.h
--- a/Engine/Plugins/Experimental/Avalanche/Source/AvalancheMedia/Private/Playback/AvaPlaybackServer.h
+++ b/Engine/Plugins/Experimental/Avalanche/Source/AvalancheMedia/Private/Playback/AvaPlaybackServer.h
@@ -185,10 +185,10 @@ protected:
 	// Playback Commands
 	void ExecutePendingPlaybackCommands(const FDateTime& InUtcNow);
 
-	TSharedPtr<FAvaPlaybackInstance> GetOrLoadPlaybackInstance(const FGuid& InInstanceId, const FString& InChannelName, const FSoftObjectPath& InAssetPath);
-	void LoadPlayback(const FMessageAddress& InReplyToAddress, const FGuid& InInstanceId, const FString& InChannelName, const FSoftObjectPath& InAssetPath);
-	void StartPlayback(const FMessageAddress& InReplyToAddress, const FGuid& InInstanceId, const FString& InChannelName, const FSoftObjectPath& InAssetPath);
-	void StopPlayback(const FMessageAddress& InReplyToAddress, const FGuid& InInstanceId, const FString& InChannelName, const FSoftObjectPath& InAssetPath);	
+	TSharedPtr<FAvaPlaybackInstance> GetOrLoadPlaybackInstance(const FGuid& InInstanceId, const FString& InChannelName, const FSoftObjectPath& InAssetPath, const FString& InLoadOptions);
+	void LoadPlayback(const FMessageAddress& InReplyToAddress, const FGuid& InInstanceId, const FString& InChannelName, const FSoftObjectPath& InAssetPath, const FString& InLoadOptions);
+	void StartPlayback(const FMessageAddress& InReplyToAddress, const FGuid& InInstanceId, const FString& InChannelName, const FSoftObjectPath& InAssetPath, const FString& InLoadOptions);
+	void StopPlayback(const FMessageAddress& InReplyToAddress, const FGuid& InInstanceId, const FString& InChannelName, const FSoftObjectPath& InAssetPath);
 	void UnloadPlayback(const FMessageAddress& InReplyToAddress, const FGuid& InInstanceId, const FString& InChannelName, const FSoftObjectPath& InAssetPath);
 	bool SetPlaybackUserData(const FMessageAddress& InReplyToAddress, const FGuid& InInstanceId, const FString& InUserData);
 	bool SendPlaybackUserData(const FMessageAddress& InReplyToAddress, const FGuid& InInstanceId);
diff --git a/Engine/Plugins/Experimental/Avalanche/Source/AvalancheMedia/Private/Playback/Nodes/AvaPlaybackNodePlayer.cpp b/Engine/Plugins/Experimental/Avalanche/Source/AvalancheMedia/Private/Playback/Nodes/AvaPlaybackNodePlayer.cpp
--- a/Engine/Plugins/Experimental/Avalanche/Source/AvalancheMedia/Private/Playback/Nodes/AvaPlaybackNodePlayer.cpp
+++ b/Engine/Plugins/Experimental/Avalanche/Source/AvalancheMedia/Private/Playback/Nodes/AvaPlaybackNodePlayer.cpp
@@ -24,7 +24,19 @@ void UAvaPlaybackNodePlayer::PostAllocateNode()
 
 void UAvaPlaybackNodePlayer::Tick(float DeltaTime, FAvaPlaybackChannelParameters& ChannelParameters)
 {
-	ChannelParameters.Assets.AddUnique(GetAssetPtr());
+	FAvaSoftAssetPtr Asset = GetAssetPtr();
+	
+	// We still don't support multi-instance within a playback graph. The asset remains the unique key.
+	auto IsSameAsset = [&Asset](const FAvaPlaybackPlayerParameters& InPlayer)
+	{
+		return InPlayer.Asset == Asset;
+	};
+
+	if (!ChannelParameters.Players.FindByPredicate(IsSameAsset))
+	{
+		ChannelParameters.Players.Add({this, Asset});	
+	}
+	
 	ChannelIndices.Add(ChannelParameters.ChannelIndex);
 	//No Children Ticks. The Player Node is the Dead End for Ticking.
 	//Events are handled separately
@@ -126,4 +138,9 @@ const FSoftObjectPath& UAvaPlaybackNodePlayer::GetAssetPath() const
 	return EmptyPath;
 }
 
+void UAvaPlaybackNodePlayer::SetLoadOptions(const FString& InLoadOptions)
+{
+	LoadOptions = InLoadOptions;
+}
+
 #undef LOCTEXT_NAMESPACE
diff --git a/Engine/Plugins/Experimental/Avalanche/Source/AvalancheMedia/Private/Playback/Nodes/Events/Actions/AvaPlaybackNode_PreloadPlayer.cpp b/Engine/Plugins/Experimental/Avalanche/Source/AvalancheMedia/Private/Playback/Nodes/Events/Actions/AvaPlaybackNode_PreloadPlayer.cpp
--- a/Engine/Plugins/Experimental/Avalanche/Source/AvalancheMedia/Private/Playback/Nodes/Events/Actions/AvaPlaybackNode_PreloadPlayer.cpp
+++ b/Engine/Plugins/Experimental/Avalanche/Source/AvalancheMedia/Private/Playback/Nodes/Events/Actions/AvaPlaybackNode_PreloadPlayer.cpp
@@ -28,7 +28,7 @@ void UAvaPlaybackNode_PreloadPlayer::OnEventTriggered(const FAvaPlaybackEventPar
 		TArray<FName> ChannelNames = Playback->GetChannelNamesForIndices(InEventParameters.ChannelIndices);
 		for (const FName& ChannelName : ChannelNames)
 		{
-			Playback->LoadAsset(InEventParameters.Asset, ChannelName);
+			Playback->LoadAsset(InEventParameters.Asset, ChannelName, FString());
 		}
 	}
 }
diff --git a/Engine/Plugins/Experimental/Avalanche/Source/AvalancheMedia/Private/Rundown/AvaRundown.cpp b/Engine/Plugins/Experimental/Avalanche/Source/AvalancheMedia/Private/Rundown/AvaRundown.cpp
--- a/Engine/Plugins/Experimental/Avalanche/Source/AvalancheMedia/Private/Rundown/AvaRundown.cpp
+++ b/Engine/Plugins/Experimental/Avalanche/Source/AvalancheMedia/Private/Rundown/AvaRundown.cpp
@@ -1167,6 +1167,18 @@ TArray<UAvaRundown::FLoadedInstanceInfo> UAvaRundown::LoadPage(int32 InPageId,
 	}
 
 	const FName ChannelName = bInPreview ? (InPreviewChannelName.IsNone() ? GetDefaultPreviewChannelName() : InPreviewChannelName) : Page.GetChannelName();
+
+	// Get the Load Options from page command, if any.
+	FString LoadOptions;
+	FAvaRundownPageCommandContext PageCommandContext = {*this, Page, ChannelName};
+
+	Page.ForEachInstancedCommands(
+		[&PageCommandContext, &LoadOptions](const FAvaRundownPageCommand& InCommand, const FAvaRundownPage& InPage)
+		{
+			InCommand.ExecuteOnLoad(PageCommandContext, LoadOptions);
+		},
+		this, /*bInDirectOnly*/ false); // Traverse templates.
+
 	const TArray<FSoftObjectPath> AssetPaths = Page.GetAssetPaths(this);
 
 	TArray<FLoadedInstanceInfo> LoadedInstances;
@@ -1174,7 +1186,7 @@ TArray<UAvaRundown::FLoadedInstanceInfo> UAvaRundown::LoadPage(int32 InPageId,
 
 	for (const FSoftObjectPath& AssetPath : AssetPaths)
 	{
-		const TSharedPtr<FAvaPlaybackInstance> PlaybackInstance = GetPlaybackManager().AcquireOrLoadPlaybackInstance(AssetPath, ChannelName.ToString());
+		const TSharedPtr<FAvaPlaybackInstance> PlaybackInstance = GetPlaybackManager().AcquireOrLoadPlaybackInstance(AssetPath, ChannelName.ToString(), LoadOptions);
 		if (!PlaybackInstance || !PlaybackInstance->GetPlayback())
 		{
 			continue;
@@ -1354,10 +1366,12 @@ bool UAvaRundown::CanPlayPage(int32 InPageId, bool bInPreview, const FName& InPr
 		// Verify if the commands can be executed.
 		bool bCanRunCommand = false;
 		FName ChannelName = bInPreview ? InPreviewChannelName : SelectedPage.GetChannelName();
-		SelectedPage.ForEachInstancedCommands([&bCanRunCommand, this, &SelectedPage, ChannelName, OutFailureReason](const FAvaRundownPageCommand& InCommand, const FAvaRundownPage& InPage)
+		FAvaRundownPageCommandContext PageCommandContext = { *this, SelectedPage, ChannelName };
+
+		SelectedPage.ForEachInstancedCommands([&bCanRunCommand, &PageCommandContext, OutFailureReason](const FAvaRundownPageCommand& InCommand, const FAvaRundownPage& InPage)
 		{
 			FString FailureReason;
-			const bool bCanExec = InCommand.CanExecute(*this, SelectedPage, ChannelName, OutFailureReason ? &FailureReason : nullptr);
+			const bool bCanExec = InCommand.CanExecuteOnPlay(PageCommandContext, OutFailureReason ? &FailureReason : nullptr);
 			bCanRunCommand |= bCanExec;
 			if (!bCanExec && OutFailureReason)
 			{
@@ -2265,15 +2279,14 @@ bool UAvaRundown::PlayPageWithTransition(FAvaRundownPageTransitionBuilder& InBui
 	if (InPage.HasCommands(this))
 	{
 		const FName ChannelName = bInIsPreview ? InPreviewChannelName : InPage.GetChannelName();
-
-		constexpr bool bDirectOnly = false;	// Traverse templates.
+		FAvaRundownPageCommandContext PageCommandContext = {*this, InPage, ChannelName};
 
 		InPage.ForEachInstancedCommands(
-			[&InBuilder, this, &ChannelName](const FAvaRundownPageCommand& InCommand, const FAvaRundownPage& InPage)
+			[&InBuilder, &PageCommandContext](const FAvaRundownPageCommand& InCommand, const FAvaRundownPage& InPage)
 			{
-				InCommand.Execute(InBuilder, *this, InPage, ChannelName);
+				InCommand.ExecuteOnPlay(InBuilder, PageCommandContext);
 			},
-			this, bDirectOnly);
+			this, /*bInDirectOnly*/ false); // Traverse templates
 
 		// If this is a page with no assets, stop here to avoid making a page player.
 		if (!InPage.HasAssets(*this))
diff --git a/Engine/Plugins/Experimental/Avalanche/Source/AvalancheMedia/Private/Rundown/AvaRundownPage.cpp b/Engine/Plugins/Experimental/Avalanche/Source/AvalancheMedia/Private/Rundown/AvaRundownPage.cpp
--- a/Engine/Plugins/Experimental/Avalanche/Source/AvalancheMedia/Private/Rundown/AvaRundownPage.cpp
+++ b/Engine/Plugins/Experimental/Avalanche/Source/AvalancheMedia/Private/Rundown/AvaRundownPage.cpp
@@ -183,9 +183,9 @@ bool FAvaRundownPage::UpdatePageSummary(const UAvaRundown* InRundown, const TArr
 	{
 		ForEachInstancedCommands([&OrderedDescription, &DisplayIndexOffset](const FAvaRundownPageCommand& InCommand, const FAvaRundownPage&)
 		{
-			OrderedDescription.FindOrAdd(DisplayIndexOffset) = FText::FromString(InCommand.GetDescription()); 
-			++DisplayIndexOffset;			
-		}, InRundown, /*bDirectOnly*/ false);
+			OrderedDescription.FindOrAdd(DisplayIndexOffset) = InCommand.GetDescription(); 
+			++DisplayIndexOffset;
+		}, InRundown, /*bInDirectOnly*/ false); // Traverse templates
 	}
 	
 	if (!OrderedDescription.IsEmpty())
@@ -998,7 +998,7 @@ void FAvaRundownPage::SetInstancedCommands(const TArray<FInstancedStruct>& InIns
 	InstancedCommands = InInstancedCommands;
 }
 
-void FAvaRundownPage::ForEachInstancedCommands(TFunctionRef<void(const FAvaRundownPageCommand&, const FAvaRundownPage&)> InFunction, const UAvaRundown* InRundown, bool bDirectOnly) const
+void FAvaRundownPage::ForEachInstancedCommands(TFunctionRef<void(const FAvaRundownPageCommand&, const FAvaRundownPage&)> InFunction, const UAvaRundown* InRundown, bool bInDirectOnly) const
 {
 	auto CallOnCommands = [&InFunction, this](const TArray<FInstancedStruct>& InCommands)
 	{
@@ -1015,7 +1015,7 @@ void FAvaRundownPage::ForEachInstancedCommands(TFunctionRef<void(const FAvaRundo
 	CallOnCommands(InstancedCommands);
 
 	// Note: if no rundown is provided to resolve templates, we have to stop here.
-	if (bDirectOnly || !InRundown)
+	if (bInDirectOnly || !InRundown)
 	{
 		return;
 	}
diff --git a/Engine/Plugins/Experimental/Avalanche/Source/AvalancheMedia/Private/Rundown/AvaRundownPagePlayer.cpp b/Engine/Plugins/Experimental/Avalanche/Source/AvalancheMedia/Private/Rundown/AvaRundownPagePlayer.cpp
--- a/Engine/Plugins/Experimental/Avalanche/Source/AvalancheMedia/Private/Rundown/AvaRundownPagePlayer.cpp
+++ b/Engine/Plugins/Experimental/Avalanche/Source/AvalancheMedia/Private/Rundown/AvaRundownPagePlayer.cpp
@@ -7,6 +7,7 @@
 #include "Playback/AvaPlaybackManager.h"
 #include "Playback/AvaPlaybackUtils.h"
 #include "Rundown/AvaRundown.h"
+#include "Rundown/AvaRundownPageCommand.h"
 
 namespace UE::AvaMedia::Rundown::PagePlayer::Private
 {
@@ -25,9 +26,20 @@ bool UAvaRundownPlaybackInstancePlayer::Load(const UAvaRundownPagePlayer& InPage
 {
 	SourceAssetPath = InPage.GetAssetPath(InRundown, InSubPageIndex);
 	TransitionLayer = InPage.GetTransitionLayer(InRundown, InSubPageIndex);
+
+	// Get the Load Options from page command, if any.
+	FString LoadOptions;
+	FAvaRundownPageCommandContext PageCommandContext = {*InRundown, InPage, InPagePlayer.ChannelFName};
 	
+	InPage.ForEachInstancedCommands(
+		[&PageCommandContext, &LoadOptions](const FAvaRundownPageCommand& InCommand, const FAvaRundownPage& InPage)
+		{
+			InCommand.ExecuteOnLoad(PageCommandContext, LoadOptions);
+		},
+		InRundown, /*bInDirectOnly*/ false); // Traverse templates.
+
 	FAvaPlaybackManager& Manager = InRundown->GetPlaybackManager();
-	PlaybackInstance = Manager.AcquireOrLoadPlaybackInstance(SourceAssetPath, InPagePlayer.ChannelName);
+	PlaybackInstance = Manager.AcquireOrLoadPlaybackInstance(SourceAssetPath, InPagePlayer.ChannelName, LoadOptions);
 	Playback = PlaybackInstance ? PlaybackInstance->GetPlayback() : nullptr;
 
 	// If restoring from a remote instance.
diff --git a/Engine/Plugins/Experimental/Avalanche/Source/AvalancheMedia/Private/Rundown/Commands/AvaRundownPageCommandSetSpawnPoint.cpp b/Engine/Plugins/Experimental/Avalanche/Source/AvalancheMedia/Private/Rundown/Commands/AvaRundownPageCommandSetSpawnPoint.cpp
new file mode 100644
--- /dev/null
+++ b/Engine/Plugins/Experimental/Avalanche/Source/AvalancheMedia/Private/Rundown/Commands/AvaRundownPageCommandSetSpawnPoint.cpp
@@ -0,0 +1,33 @@
+// Copyright Epic Games, Inc. All Rights Reserved.
+
+#include "AvaRundownPageCommandSetSpawnPoint.h"
+
+#define LOCTEXT_NAMESPACE "AvaRundownPageCommandSetSpawnPoint"
+
+FText FAvaRundownPageCommandSetSpawnPoint::GetDescription() const
+{
+	return LOCTEXT("Command_Description", "Set Spawn Point");
+}
+
+bool FAvaRundownPageCommandSetSpawnPoint::HasTransitionLogic() const
+{
+	return false;
+}
+
+FString FAvaRundownPageCommandSetSpawnPoint::GetTransitionLayerString(const FString& InSeparator) const
+{
+	return TEXT("");
+}
+
+bool FAvaRundownPageCommandSetSpawnPoint::CanExecuteOnLoad(FAvaRundownPageCommandContext& InContext, FString* OutFailureReason) const
+{
+	return true;
+}
+
+bool FAvaRundownPageCommandSetSpawnPoint::ExecuteOnLoad(FAvaRundownPageCommandContext& InContext, FString& OutLoadOptions) const
+{
+	OutLoadOptions += FString::Printf(TEXT(" -SpawnPointTag=\"%s\""), *SpawnPointTag.ToString());
+	return true;
+}
+
+#undef LOCTEXT_NAMESPACE
\ No newline at end of file
diff --git a/Engine/Plugins/Experimental/Avalanche/Source/AvalancheMedia/Private/Rundown/Commands/AvaRundownPageCommandSetSpawnPoint.h b/Engine/Plugins/Experimental/Avalanche/Source/AvalancheMedia/Private/Rundown/Commands/AvaRundownPageCommandSetSpawnPoint.h
new file mode 100644
--- /dev/null
+++ b/Engine/Plugins/Experimental/Avalanche/Source/AvalancheMedia/Private/Rundown/Commands/AvaRundownPageCommandSetSpawnPoint.h
@@ -0,0 +1,30 @@
+// Copyright Epic Games, Inc. All Rights Reserved.
+
+#pragma once
+
+#include "Rundown/AvaRundownPageCommand.h"
+
+#include "AvaRundownPageCommandSetSpawnPoint.generated.h"
+
+/**
+ * Specify a reference null actor to use the position to load the graphic at.
+ * This can be used to position a null actor with the given tag in the base level and have the graphics
+ * load at those positions.
+ */
+USTRUCT(BlueprintType, meta = (DisplayName = "Set Spawn Point"))
+struct FAvaRundownPageCommandSetSpawnPoint : public FAvaRundownPageCommand
+{
+	GENERATED_BODY()
+
+	//~ Begin FAvaRundownPageCommand
+	virtual FText GetDescription() const override;
+	virtual bool HasTransitionLogic() const override;
+	virtual FString GetTransitionLayerString(const FString& InSeparator) const override;
+	virtual bool CanExecuteOnLoad(FAvaRundownPageCommandContext& InContext, FString* OutFailureReason) const override;
+	virtual bool ExecuteOnLoad(FAvaRundownPageCommandContext& InContext, FString& OutLoadOptions) const override;
+	//~ End FAvaRundownPageCommand
+	
+	/** Indicate which spawn point, by tag, to search for. */
+	UPROPERTY(EditAnywhere, BlueprintReadWrite, Category = "Page Data")
+	FName SpawnPointTag;
+};
diff --git a/Engine/Plugins/Experimental/Avalanche/Source/AvalancheMedia/Private/Rundown/Commands/AvaRundownPageCommandSetTransform.cpp b/Engine/Plugins/Experimental/Avalanche/Source/AvalancheMedia/Private/Rundown/Commands/AvaRundownPageCommandSetTransform.cpp
new file mode 100644
--- /dev/null
+++ b/Engine/Plugins/Experimental/Avalanche/Source/AvalancheMedia/Private/Rundown/Commands/AvaRundownPageCommandSetTransform.cpp
@@ -0,0 +1,33 @@
+// Copyright Epic Games, Inc. All Rights Reserved.
+
+#include "AvaRundownPageCommandSetTransform.h"
+
+#define LOCTEXT_NAMESPACE "AvaRundownPageCommandSetTransform"
+
+FText FAvaRundownPageCommandSetTransform::GetDescription() const
+{
+	return LOCTEXT("Command_Description", "Set Transform");
+}
+
+bool FAvaRundownPageCommandSetTransform::HasTransitionLogic() const
+{
+	return false;
+}
+
+FString FAvaRundownPageCommandSetTransform::GetTransitionLayerString(const FString& InSeparator) const
+{
+	return TEXT("");
+}
+
+bool FAvaRundownPageCommandSetTransform::CanExecuteOnLoad(FAvaRundownPageCommandContext& InContext, FString* OutFailureReason) const
+{
+	return true;
+}
+
+bool FAvaRundownPageCommandSetTransform::ExecuteOnLoad(FAvaRundownPageCommandContext& InContext, FString& OutLoadOptions) const
+{
+	OutLoadOptions += FString::Printf(TEXT(" -Transform=\"%s\""), *Transform.ToString());
+	return true;
+}
+
+#undef LOCTEXT_NAMESPACE
\ No newline at end of file
diff --git a/Engine/Plugins/Experimental/Avalanche/Source/AvalancheMedia/Private/Rundown/Commands/AvaRundownPageCommandSetTransform.h b/Engine/Plugins/Experimental/Avalanche/Source/AvalancheMedia/Private/Rundown/Commands/AvaRundownPageCommandSetTransform.h
new file mode 100644
--- /dev/null
+++ b/Engine/Plugins/Experimental/Avalanche/Source/AvalancheMedia/Private/Rundown/Commands/AvaRundownPageCommandSetTransform.h
@@ -0,0 +1,30 @@
+// Copyright Epic Games, Inc. All Rights Reserved.
+
+#pragma once
+
+#include "Rundown/AvaRundownPageCommand.h"
+
+#include "AvaRundownPageCommandSetTransform.generated.h"
+
+
+/**
+ * Page Command to specify the transform when loading the graphic.
+ * For streaming levels, this is applied when the level is loaded.
+ */
+USTRUCT(BlueprintType, meta = (DisplayName = "Set Transform"))
+struct FAvaRundownPageCommandSetTransform : public FAvaRundownPageCommand
+{
+	GENERATED_BODY()
+
+	//~ Begin FAvaRundownPageCommand
+	virtual FText GetDescription() const override;
+	virtual bool HasTransitionLogic() const override;
+	virtual FString GetTransitionLayerString(const FString& InSeparator) const override;
+	virtual bool CanExecuteOnLoad(FAvaRundownPageCommandContext& InContext, FString* OutFailureReason) const override;
+	virtual bool ExecuteOnLoad(FAvaRundownPageCommandContext& InContext, FString& OutLoadOptions) const override;
+	//~ End FAvaRundownPageCommand
+	
+	/** Specify the root transform to apply at load time. */
+	UPROPERTY(EditAnywhere, BlueprintReadWrite, Category = "Page Data")
+	FTransform Transform;
+};
diff --git a/Engine/Plugins/Experimental/Avalanche/Source/AvalancheMedia/Private/Rundown/Commands/AvaRundownPageCommandStopLayers.cpp b/Engine/Plugins/Experimental/Avalanche/Source/AvalancheMedia/Private/Rundown/Commands/AvaRundownPageCommandStopLayers.cpp
--- a/Engine/Plugins/Experimental/Avalanche/Source/AvalancheMedia/Private/Rundown/Commands/AvaRundownPageCommandStopLayers.cpp
+++ b/Engine/Plugins/Experimental/Avalanche/Source/AvalancheMedia/Private/Rundown/Commands/AvaRundownPageCommandStopLayers.cpp
@@ -4,12 +4,13 @@
 
 #include "AvaTagCollection.h"
 #include "Rundown/AvaRundown.h"
-#include "Rundown/AvaRundownPage.h"
 #include "Rundown/AvaRundownPagePlayer.h"
 #include "Rundown/AvaRundownPlaybackUtils.h"
 #include "Rundown/Transition/AvaRundownPageTransition.h"
 #include "Rundown/Transition/AvaRundownPageTransitionBuilder.h"
 
+#define LOCTEXT_NAMESPACE "AvaRundownPageCommandStopLayers"
+
 namespace UE::AvaMedia::RundownPageCommandStopLayers::Private
 {
 	// todo: might be worth having this in FAvaTagHandleContainer. 
@@ -47,10 +48,10 @@ namespace UE::AvaMedia::RundownPageCommandStopLayers::Private
 	}
 }
 
-FString FAvaRundownPageCommandStopLayers::GetDescription() const
+FText FAvaRundownPageCommandStopLayers::GetDescription() const
 {
 	using namespace UE::AvaMedia::RundownPageCommandStopLayers::Private;
-	return FString::Printf(TEXT("Stop Layers: %s"), *GetLayersString(Layers, TEXT(""), TEXT(", ")));
+	return FText::Format(LOCTEXT("Command_Description", "Stop Layers: {0}"), FText::FromString(GetLayersString(Layers, TEXT(""), TEXT(", "))));
 }
 
 bool FAvaRundownPageCommandStopLayers::HasTransitionLogic() const
@@ -66,7 +67,7 @@ FString FAvaRundownPageCommandStopLayers::GetTransitionLayerString(const FString
 	return GetLayersString(Layers, TEXT("-"), InSeparator);
 }
 
-bool FAvaRundownPageCommandStopLayers::CanExecute(const UAvaRundown& InRundown, const FAvaRundownPage& InPage, FName InChannelName, FString* OutFailureReason) const
+bool FAvaRundownPageCommandStopLayers::CanExecuteOnPlay(FAvaRundownPageCommandContext& InContext, FString* OutFailureReason) const
 {
 	using namespace UE::AvaMedia::RundownPageCommandStopLayers::Private;
 	if (IsEmpty(Layers))
@@ -79,9 +80,9 @@ bool FAvaRundownPageCommandStopLayers::CanExecute(const UAvaRundown& InRundown,
 	}
 
 	// Check if there are currently playing pages that have layers than can be stopped.
-	for (const UAvaRundownPagePlayer* PagePlayer : InRundown.GetPagePlayers())
+	for (const UAvaRundownPagePlayer* PagePlayer : InContext.Rundown.GetPagePlayers())
 	{
-		if (!PagePlayer || PagePlayer->ChannelFName != InChannelName)	// Not this channel
+		if (!PagePlayer || PagePlayer->ChannelFName != InContext.ChannelName)	// Not this channel
 		{
 			continue;
 		}
@@ -107,7 +108,7 @@ bool FAvaRundownPageCommandStopLayers::CanExecute(const UAvaRundown& InRundown,
 	return false;
 }
 
-bool FAvaRundownPageCommandStopLayers::Execute(FAvaRundownPageTransitionBuilder& InTransitionBuilder, const UAvaRundown& InRundown, const FAvaRundownPage& InPage, FName InChannelName) const
+bool FAvaRundownPageCommandStopLayers::ExecuteOnPlay(FAvaRundownPageTransitionBuilder& InTransitionBuilder, FAvaRundownPageCommandContext& InContext) const
 {
 	using namespace UE::AvaMedia::RundownPageCommandStopLayers::Private;
 	if (IsEmpty(Layers))
@@ -115,7 +116,7 @@ bool FAvaRundownPageCommandStopLayers::Execute(FAvaRundownPageTransitionBuilder&
 		return false;
 	}
 
-	if (UAvaRundownPageTransition* Transition = InTransitionBuilder.FindOrAddTransition(InChannelName))
+	if (UAvaRundownPageTransition* Transition = InTransitionBuilder.FindOrAddTransition(InContext.ChannelName))
 	{
 		for (const FAvaTagId& TagId : Layers.GetTagIds())
 		{
@@ -125,3 +126,5 @@ bool FAvaRundownPageCommandStopLayers::Execute(FAvaRundownPageTransitionBuilder&
 	}
 	return false;
 }
+
+#undef LOCTEXT_NAMESPACE
\ No newline at end of file
diff --git a/Engine/Plugins/Experimental/Avalanche/Source/AvalancheMedia/Private/Rundown/Commands/AvaRundownPageCommandStopLayers.h b/Engine/Plugins/Experimental/Avalanche/Source/AvalancheMedia/Private/Rundown/Commands/AvaRundownPageCommandStopLayers.h
--- a/Engine/Plugins/Experimental/Avalanche/Source/AvalancheMedia/Private/Rundown/Commands/AvaRundownPageCommandStopLayers.h
+++ b/Engine/Plugins/Experimental/Avalanche/Source/AvalancheMedia/Private/Rundown/Commands/AvaRundownPageCommandStopLayers.h
@@ -13,11 +13,11 @@ struct FAvaRundownPageCommandStopLayers : public FAvaRundownPageCommand
 	GENERATED_BODY()
 
 	//~ Begin FAvaRundownPageCommand
-	virtual FString GetDescription() const override;
+	virtual FText GetDescription() const override;
 	virtual bool HasTransitionLogic() const override;
 	virtual FString GetTransitionLayerString(const FString& InSeparator) const override;
-	virtual bool CanExecute(const UAvaRundown& InRundown, const FAvaRundownPage& InPage, FName InChannelName, FString* OutFailureReason) const override;
-	virtual bool Execute(FAvaRundownPageTransitionBuilder& InTransitionBuilder, const UAvaRundown& InRundown, const FAvaRundownPage& InPage, FName InChannelName) const override;
+	virtual bool CanExecuteOnPlay(FAvaRundownPageCommandContext& InContext, FString* OutFailureReason) const override;
+	virtual bool ExecuteOnPlay(FAvaRundownPageTransitionBuilder& InTransitionBuilder, FAvaRundownPageCommandContext& InContext) const override;
 	//~ End FAvaRundownPageCommand
 	
 	/** Layers to stop. */
diff --git a/Engine/Plugins/Experimental/Avalanche/Source/AvalancheMedia/Public/AvaMediaDefines.h b/Engine/Plugins/Experimental/Avalanche/Source/AvalancheMedia/Public/AvaMediaDefines.h
--- a/Engine/Plugins/Experimental/Avalanche/Source/AvalancheMedia/Public/AvaMediaDefines.h
+++ b/Engine/Plugins/Experimental/Avalanche/Source/AvalancheMedia/Public/AvaMediaDefines.h
@@ -182,6 +182,15 @@ enum class EAvaPlayableRCUpdateFlags : uint8
 };
 ENUM_CLASS_FLAGS(EAvaPlayableRCUpdateFlags);
 
+/**
+ * Parameters of a player entry for active player nodes in playback graph. 
+ */
+struct FAvaPlaybackPlayerParameters
+{
+	TWeakObjectPtr<class UAvaPlaybackNode> Node;
+	FAvaSoftAssetPtr Asset;
+};
+
 /*
  * Parameters used to tell a Channel its Playback Settings
  */
@@ -192,13 +201,13 @@ struct FAvaPlaybackChannelParameters
 
 	/**
 	 * If a playback node is reached during tick traversal,
-	 * it is going to set this value.
+	 * it is going to add an entry.
 	 */
-	TArray<FAvaSoftAssetPtr> Assets;
+	TArray<FAvaPlaybackPlayerParameters> Players;
 
 	bool HasAssets() const
 	{
-		return !Assets.IsEmpty();
+		return !Players.IsEmpty();
 	}
 };
 
diff --git a/Engine/Plugins/Experimental/Avalanche/Source/AvalancheMedia/Public/Playable/AvaPlayable.h b/Engine/Plugins/Experimental/Avalanche/Source/AvalancheMedia/Public/Playable/AvaPlayable.h
--- a/Engine/Plugins/Experimental/Avalanche/Source/AvalancheMedia/Public/Playable/AvaPlayable.h
+++ b/Engine/Plugins/Experimental/Avalanche/Source/AvalancheMedia/Public/Playable/AvaPlayable.h
@@ -58,17 +58,14 @@ enum class EAvaPlayableCommandResult : uint8
  * Design goal:
  *
  * The design goal is to abstract the implementation of a playable.
- * So far we have 1 implementation:
+ * So far we have one concrete implementation:
  * - Level Streaming that can be streamed with other levels in the same game instance.
  *
  * To support multiple playable in the same channel/output, there are 2 ways:
  * - rendering in the same world which can only be done with the level streaming playables.
  * - compositing different renders. This is not yet supported, but the playable abstraction should help.
  *
- * Future Work:
- * There is also another abstraction desired which is the "remote" playable. The AvaPlaybackGraph
- * has a lot of logic to manage remote replication of playable commands. It would probably
- * be cleaner and more reusable to have a RemotePlayable class that implements all those instead. 
+ * Distributed rendering is implemented using playable and playable group remote proxies.
  */
 UCLASS(Abstract, NotBlueprintable, BlueprintType, ClassGroup = "Motion Design Playable",
 	meta = (DisplayName = "Motion Design Playable"))	// Todo: review class specifier
@@ -108,8 +105,8 @@ public:
 
 	DECLARE_MULTICAST_DELEGATE_OneParam(FOnPlayableStatusChanged, UAvaPlayable*);
 	FOnPlayableStatusChanged& OnPlayableStatusChanged() { return OnPlayableStatusChangedDelegate; }
-	
-	virtual bool LoadAsset(const FAvaSoftAssetPtr& InSourceAsset, bool bInInitiallyVisible) { return false;}
+
+	virtual bool LoadAsset(const FAvaSoftAssetPtr& InSourceAsset, bool bInInitiallyVisible, const FString& InLoadOptions = FString()) { return false;}
 	virtual bool UnloadAsset() {return false;}
 	virtual const FSoftObjectPath& GetSourceAssetPath() const;
 	virtual EAvaPlayableStatus GetPlayableStatus() const { return EAvaPlayableStatus::Unknown;}
diff --git a/Engine/Plugins/Experimental/Avalanche/Source/AvalancheMedia/Public/Playback/AvaPlaybackGraph.h b/Engine/Plugins/Experimental/Avalanche/Source/AvalancheMedia/Public/Playback/AvaPlaybackGraph.h
--- a/Engine/Plugins/Experimental/Avalanche/Source/AvalancheMedia/Public/Playback/AvaPlaybackGraph.h
+++ b/Engine/Plugins/Experimental/Avalanche/Source/AvalancheMedia/Public/Playback/AvaPlaybackGraph.h
@@ -187,7 +187,7 @@ public:
 	bool HasPlayable(const UAvaPlayable* InPlayable) const;
 	
 	UAvaPlayable* FindPlayable(const FSoftObjectPath& InSourceAssetPath, const FName& InChannelName) const;
-	UAvaPlayable* FindOrLoadPlayable(const FAvaSoftAssetPtr& InSourceAsset, const FName& InChannelName);
+	UAvaPlayable* FindOrLoadPlayable(const FAvaSoftAssetPtr& InSourceAsset, const FName& InChannelName, const FString& InLoadOptions = FString());
 	void RemovePlayable(const FSoftObjectPath& InSourceAssetPath, const FName& InChannelName);
 	bool UnloadAndRemovePlayable(UAvaPlayable* InPlayable, const FSoftObjectPath& InSourceAssetPath, const FName& InChannelName, bool bInForceImmediate);
 	bool UnloadAndRemovePlayable(const FSoftObjectPath& InSourceAssetPath, const FName& InChannelName, bool bInForceImmediate);
@@ -267,7 +267,7 @@ public:
 	/**
 	 * Start streaming the asset in a game instance.
 	 */
-	void LoadAsset(const FAvaSoftAssetPtr& InSourceAsset, const FName& InChannelName);
+	void LoadAsset(const FAvaSoftAssetPtr& InSourceAsset, const FName& InChannelName, const FString& InLoadOptions);
 	
 	virtual void Tick(float DeltaTime);
 	
@@ -346,11 +346,11 @@ protected:
 	
 	bool CanRefreshPlayback(const FAvaSoftAssetPtr& InSourceAsset, const FAvaBroadcastOutputChannel& InChannel) const;
 
-	bool RefreshPlayback(const FAvaSoftAssetPtr& InSourceAsset, FAvaBroadcastOutputChannel& InChannel, const FName& InChannelName);
+	bool RefreshPlayback(const FAvaSoftAssetPtr& InSourceAsset, const FString& InLoadOptions, FAvaBroadcastOutputChannel& InChannel, const FName& InChannelName);
 	
 	bool CanRefreshPreview(const FAvaSoftAssetPtr& InSourceAsset) const;
 	
-	bool RefreshPreview(const FAvaSoftAssetPtr& InSourceAsset, FAvaBroadcastOutputChannel& InChannel, const FName& InChannelName);
+	bool RefreshPreview(const FAvaSoftAssetPtr& InSourceAsset, const FString& InLoadOptions, FAvaBroadcastOutputChannel& InChannel, const FName& InChannelName);
 
 	static UTextureRenderTarget2D* UpdatePlaybackRenderTarget(const UAvaPlayable* InPlayable, const FAvaBroadcastOutputChannel& InChannel);
 	
diff --git a/Engine/Plugins/Experimental/Avalanche/Source/AvalancheMedia/Public/Playback/AvaPlaybackManager.h b/Engine/Plugins/Experimental/Avalanche/Source/AvalancheMedia/Public/Playback/AvaPlaybackManager.h
--- a/Engine/Plugins/Experimental/Avalanche/Source/AvalancheMedia/Public/Playback/AvaPlaybackManager.h
+++ b/Engine/Plugins/Experimental/Avalanche/Source/AvalancheMedia/Public/Playback/AvaPlaybackManager.h
@@ -151,14 +151,31 @@ public:
 
 	UAvaPlayableGroupManager* GetPlayableGroupManager() const { return PlayableGroupManager.Get(); }
 
-	/** Acquire a cached playback instance. Will return null if none available in the cache. */
+	/**
+	 * Acquire (recycle) a cached playback instance. Will return null if none available in the cache.
+	 * @param InAssetPath Soft object path to the asset to load. This can either be a playback graph itself or a playable asset.
+	 * @param InChannelName For a playable asset, specifying the channel name will indicate where to connect the player node in the playback graph
+	 */
 	TSharedPtr<FAvaPlaybackInstance> AcquirePlaybackInstance(const FSoftObjectPath& InAssetPath, const FString& InChannelName) const;
-	
-	/** Load a new playback instance. */
-	TSharedPtr<FAvaPlaybackInstance> LoadPlaybackInstance(const FSoftObjectPath& InAssetPath, const FString& InChannelName);
 
-	/** Attempts to acquire existing instance. If none available, will load a new one. */  
-	TSharedPtr<FAvaPlaybackInstance> AcquireOrLoadPlaybackInstance(const FSoftObjectPath& InAssetPath, const FString& InChannelName);
+	/**
+	 * Load a new (recyclable) playback instance for the given asset in the given channel.
+	 * @remark This function only loads (or creates) the playback graph. The playable assets should be loaded only once the graph is playing.
+	 * @param InAssetPath Soft object path to the asset to load. This can either be a playback graph itself or a playable asset.
+	 * @param InChannelName For a playable asset, specifying the channel name will indicate where to connect the player node in the playback graph
+	 * @param InLoadOptions Optional parameters to be used when loading the playable. 
+	 */
+	TSharedPtr<FAvaPlaybackInstance> LoadPlaybackInstance(const FSoftObjectPath& InAssetPath, const FString& InChannelName, const FString& InLoadOptions = FString());
+
+	/**
+	 * Attempts to acquire (recycle) an existing playback instance. If none available, will load a new one.
+	 * @param InAssetPath Soft object path to the asset to load. This can either be a playback graph itself or a playable asset.
+	 * @param InChannelName For a playable asset, specifying the channel name will indicate where to connect the player node in the playback graph
+	 * @param InLoadOptions Optional parameters to be used when loading the playable.
+ 	 * @remark Load options will not be applied for recycled instances. (todo)
+	 * @see LoadPlaybackInstance, AcquirePlaybackInstance
+	 */  
+	TSharedPtr<FAvaPlaybackInstance> AcquireOrLoadPlaybackInstance(const FSoftObjectPath& InAssetPath, const FString& InChannelName, const FString& InLoadOptions = FString());
 
 	/**
 	 * Finds an existing (either available or used) playback instance.
@@ -213,9 +230,24 @@ public:
 	 *	Invalidate the asset entry. All cached instances will be invalidated along with it.
 	 */
 	void InvalidatePlaybackAssetEntry(const FSoftObjectPath& InAssetPath);
-	
-	UAvaPlaybackGraph* LoadPlaybackObject(const FSoftObjectPath& InAssetPath, const FString& InChannelName) const;
-	UAvaPlaybackGraph* BuildPlaybackFromWorld(const TSoftObjectPtr<UWorld>& InWorld, const FString& InChannelName) const;
+
+	/**
+	 * Loads a new playback graph for the given asset with the given channel as context.
+	 * @remark The playback graph is either loaded or created. The referenced assets (playables) are not loaded yet.
+	 * @param InAssetPath Soft object path to the asset to load. This can either be a playback graph itself or a playable asset.
+	 * @param InChannelName For a playable asset, specifying the channel name will indicate where to connect the player node in the playback graph
+	 * @param InLoadOptions Optional parameters to be used when loading the playable.
+	 * @remark Direct loading of playback graphs is not fully supported yet. Experimental use only.
+	 */
+	UAvaPlaybackGraph* LoadPlaybackObject(const FSoftObjectPath& InAssetPath, const FString& InChannelName, const FString& InLoadOptions = FString()) const;
+
+	/**
+	 * Utility function to create a new playback graph instance for the given world asset with the given channel as context.
+	 * @param InWorld Soft object path to the world asset to load.
+	 * @param InChannelName For a playable asset, specifying the channel name will indicate where to connect the player node in the playback graph
+	 * @param InLoadOptions Optional parameters to be used when loading the playable.
+	 */
+	UAvaPlaybackGraph* BuildPlaybackFromWorld(const TSoftObjectPtr<UWorld>& InWorld, const FString& InChannelName, const FString& InLoadOptions = FString()) const;
 	
 	/**
 	 * Stops all currently playing playback objects.
diff --git a/Engine/Plugins/Experimental/Avalanche/Source/AvalancheMedia/Public/Playback/Nodes/AvaPlaybackNodePlayer.h b/Engine/Plugins/Experimental/Avalanche/Source/AvalancheMedia/Public/Playback/Nodes/AvaPlaybackNodePlayer.h
--- a/Engine/Plugins/Experimental/Avalanche/Source/AvalancheMedia/Public/Playback/Nodes/AvaPlaybackNodePlayer.h
+++ b/Engine/Plugins/Experimental/Avalanche/Source/AvalancheMedia/Public/Playback/Nodes/AvaPlaybackNodePlayer.h
@@ -46,10 +46,19 @@ public:
 
 	virtual FAvaSoftAssetPtr GetAssetPtr() const { return FAvaSoftAssetPtr(); }
 
+	virtual const FString& GetLoadOptions() const { return LoadOptions;}
+
 	/** Get the channel indices from the last time this node was traversed. */
 	const TArray<int32>& GetLastTickChannelIndices() const { return LastTickChannelIndices; }
 
+	/** Set current load options for the asset. */
+	virtual void SetLoadOptions(const FString& InLoadOptions);
+
 protected:
+	/** Additional information necessary to load an asset. */
+	UPROPERTY(EditAnywhere, Category = "Motion Design")
+	FString LoadOptions;
+	
 	FText DisplayNameText;
 
 	/**
diff --git a/Engine/Plugins/Experimental/Avalanche/Source/AvalancheMedia/Public/Rundown/AvaRundownPage.h b/Engine/Plugins/Experimental/Avalanche/Source/AvalancheMedia/Public/Rundown/AvaRundownPage.h
--- a/Engine/Plugins/Experimental/Avalanche/Source/AvalancheMedia/Public/Rundown/AvaRundownPage.h
+++ b/Engine/Plugins/Experimental/Avalanche/Source/AvalancheMedia/Public/Rundown/AvaRundownPage.h
@@ -307,9 +307,9 @@ public:
 	 * @brief Calls the given function for each commands in the page.
 	 * @param InFunction Function to call for each command. The Page parameter is the direct page the command is on, i.e. may be the template.
 	 * @param InRundown Owning rundown (used to resolve templates)
-	 * @param bDirectOnly If true, only the current page's command is traversed, otherwise, the templates commands are included as well.
+	 * @param bInDirectOnly If true, only the current page's command is traversed, otherwise, the templates commands are included as well.
 	 */
-	void ForEachInstancedCommands(TFunctionRef<void(const FAvaRundownPageCommand&, const FAvaRundownPage&)> InFunction, const UAvaRundown* InRundown, bool bDirectOnly) const;
+	void ForEachInstancedCommands(TFunctionRef<void(const FAvaRundownPageCommand&, const FAvaRundownPage&)> InFunction, const UAvaRundown* InRundown, bool bInDirectOnly) const;
 
 protected:
 	UPROPERTY(VisibleAnywhere, BlueprintReadOnly, Category = "Motion Design")
diff --git a/Engine/Plugins/Experimental/Avalanche/Source/AvalancheMedia/Public/Rundown/AvaRundownPageCommand.h b/Engine/Plugins/Experimental/Avalanche/Source/AvalancheMedia/Public/Rundown/AvaRundownPageCommand.h
--- a/Engine/Plugins/Experimental/Avalanche/Source/AvalancheMedia/Public/Rundown/AvaRundownPageCommand.h
+++ b/Engine/Plugins/Experimental/Avalanche/Source/AvalancheMedia/Public/Rundown/AvaRundownPageCommand.h
@@ -10,7 +10,20 @@ class FAvaRundownPageTransitionBuilder;
 class UAvaRundown;
 struct FAvaRundownPage;
 
-/** Base class for page commands */
+/**
+ * Base context carrying parameters for page commands execution. 
+ */
+struct FAvaRundownPageCommandContext
+{
+	const UAvaRundown& Rundown;
+	const FAvaRundownPage& Page;
+	const FName ChannelName;
+};
+
+/**
+ * Base class for page commands.
+ * @remark this API is experimental.
+ */
 USTRUCT(BlueprintType)
 struct FAvaRundownPageCommand
 {
@@ -19,7 +32,7 @@ struct FAvaRundownPageCommand
 	virtual ~FAvaRundownPageCommand() = default;
 
 	/** Returns this command's description. Used for page summary. */
-	virtual FString GetDescription() const { return FString();}
+	virtual FText GetDescription() const { return FText();}
 	
 	/** Returns true if the command will act as transition logic. */
 	virtual bool HasTransitionLogic() const { return false; }
@@ -28,8 +41,14 @@ struct FAvaRundownPageCommand
 	virtual FString GetTransitionLayerString(const FString& InSeparator) const { return FString(); }
 	
 	/** Returns true if the command can be executed. Used by UI for button status. */
-	virtual bool CanExecute(const UAvaRundown& InRundown, const FAvaRundownPage& InPage, FName InChannelName, FString* OutFailureReason) const { return false; }
+	virtual bool CanExecuteOnPlay(FAvaRundownPageCommandContext& InContext, FString* OutFailureReason) const { return false; }
+
+	/** Execute the command when a page is played. */
+	virtual bool ExecuteOnPlay(FAvaRundownPageTransitionBuilder& InTransitionBuilder, FAvaRundownPageCommandContext& InContext) const { return false; }
+
+	/** Returns true if the command can be executed. Used by UI for button status. */
+	virtual bool CanExecuteOnLoad(FAvaRundownPageCommandContext& InContext, FString* OutFailureReason) const { return false; }
 
-	/** Execute the command. */
-	virtual bool Execute(FAvaRundownPageTransitionBuilder& InTransitionBuilder, const UAvaRundown& InRundown, const FAvaRundownPage& InPage, FName InChannelName) const { return false; }
+	/** Execute the command when a page is loaded. */
+	virtual bool ExecuteOnLoad(FAvaRundownPageCommandContext& InContext, FString& OutLoadOptions) const { return false; }
 };
diff --git a/Engine/Plugins/Experimental/Avalanche/Source/AvalancheMediaEditor/Private/Rundown/Pages/PageViews/AvaRundownPageViewImpl.cpp b/Engine/Plugins/Experimental/Avalanche/Source/AvalancheMediaEditor/Private/Rundown/Pages/PageViews/AvaRundownPageViewImpl.cpp
--- a/Engine/Plugins/Experimental/Avalanche/Source/AvalancheMediaEditor/Private/Rundown/Pages/PageViews/AvaRundownPageViewImpl.cpp
+++ b/Engine/Plugins/Experimental/Avalanche/Source/AvalancheMediaEditor/Private/Rundown/Pages/PageViews/AvaRundownPageViewImpl.cpp
@@ -93,7 +93,7 @@ FText FAvaRundownPageViewImpl::GetPageTransitionLayerNameText() const
 				{
 					AppendWithSeparator(TransitionLayers, CommandLayersString);
 				}
-			}, Rundown, /*bDirectOnly*/ false);
+			}, Rundown, /*bInDirectOnly*/ false); // Traverse templates
 
 			return !TransitionLayers.IsEmpty() ? FText::FromString(TransitionLayers) : LOCTEXT("PageTransitionLayerText_NA", "N/A");
 		}
