Motion Design: renaming actors now propagates to all the sequences inside the motion design scene and any sub-sequence used within these

diff --git a/Engine/Plugins/Experimental/Avalanche/Source/AvalancheSequencer/Private/AvaSequencer.cpp b/Engine/Plugins/Experimental/Avalanche/Source/AvalancheSequencer/Private/AvaSequencer.cpp
--- a/Engine/Plugins/Experimental/Avalanche/Source/AvalancheSequencer/Private/AvaSequencer.cpp
+++ b/Engine/Plugins/Experimental/Avalanche/Source/AvalancheSequencer/Private/AvaSequencer.cpp
@@ -12,6 +12,7 @@
 #include "Clipboard/AvaSequenceImporter.h"
 #include "Commands/AvaSequencerAction.h"
 #include "Commands/AvaSequencerCommands.h"
+#include "Compilation/MovieSceneCompiledDataManager.h"
 #include "ContentBrowserModule.h"
 #include "CoreGlobals.h"
 #include "DetailsView/SAvaMarkDetails.h"
@@ -46,14 +47,14 @@
 #include "PropertyHandle.h"
 #include "ScopedTransaction.h"
 #include "Selection/AvaEditorSelection.h"
-#include "SequencerCommands.h"
-#include "SequencerSettings.h"
-#include "SequencerUtilities.h"
 #include "SequenceTree/AvaSequenceItem.h"
-#include "SequenceTree/IAvaSequenceItem.h"
 #include "SequenceTree/Columns/AvaSequenceNameColumn.h"
 #include "SequenceTree/Columns/AvaSequenceStatusColumn.h"
+#include "SequenceTree/IAvaSequenceItem.h"
 #include "SequenceTree/Widgets/SAvaSequenceTree.h"
+#include "SequencerCommands.h"
+#include "SequencerSettings.h"
+#include "SequencerUtilities.h"
 #include "Settings/AvaSequencerSettings.h"
 #include "Sidebar/SidebarDrawerConfig.h"
 #include "StaggerTool/AvaStaggerTool.h"
@@ -118,7 +119,75 @@ namespace UE::AvaSequencer::Private
 		/** Cached selected objects of Target */
 		TArray<UObject*> OriginalTargetSelectedObjects;
 	};
-}
+
+	struct FRenameBindingParams
+	{
+		/** The shared playback state */
+		TSharedRef<MovieScene::FSharedPlaybackState> PlaybackState;
+		/** The state to assign sequences and find object ids */
+		FMovieSceneEvaluationState* EvaluationState;
+		/** Describes the hierarchy of the sequence (how sub-sequences are within a root sequence) */
+		const FMovieSceneSequenceHierarchy* SequenceHierarchy;
+	};
+	struct FRenameBindingSequenceParams
+	{
+		/** currently processed sequence. Starts with root */
+		UMovieSceneSequence* Sequence;
+		/** id of the currently processed sequence. Starts with root id */
+		FMovieSceneSequenceIDRef SequenceId;
+	};
+	/** Renames all bindings of the given actor to its updated actor label */
+	void RenameBindingRecursive(AActor* InActor, const FRenameBindingParams& InParams, const FRenameBindingSequenceParams& InSequenceParams)
+	{
+		// check parameters that should be guaranteed valid by caller
+		check(InActor && InParams.SequenceHierarchy && InParams.EvaluationState);
+
+		if (!InSequenceParams.Sequence)
+		{
+			return;
+		}
+
+		UMovieScene* const MovieScene = InSequenceParams.Sequence->GetMovieScene();
+		if (!MovieScene)
+		{
+			return;
+		}
+
+		InParams.EvaluationState->AssignSequence(InSequenceParams.SequenceId, *InSequenceParams.Sequence, InParams.PlaybackState);
+
+		const FGuid ObjectId = InParams.EvaluationState->FindObjectId(*InActor, InSequenceParams.SequenceId, InParams.PlaybackState);
+		if (ObjectId.IsValid())
+		{
+			if (FMovieScenePossessable* Possessable = MovieScene->FindPossessable(ObjectId))
+			{
+				InSequenceParams.Sequence->Modify();
+				Possessable->SetName(InActor->GetActorLabel());
+			}
+			else if (FMovieSceneSpawnable* Spawnable = MovieScene->FindSpawnable(ObjectId))
+			{
+				InSequenceParams.Sequence->Modify();
+				Spawnable->SetName(InActor->GetActorLabel());
+			}
+		}
+
+		// Recurse into child nodes
+		if (const FMovieSceneSequenceHierarchyNode* Node = InParams.SequenceHierarchy->FindNode(InSequenceParams.SequenceId))
+		{
+			for (FMovieSceneSequenceIDRef ChildId : Node->Children)
+			{
+				if (const FMovieSceneSubSequenceData* SubData = InParams.SequenceHierarchy->FindSubData(ChildId))
+				{
+					const FRenameBindingSequenceParams ChildSequenceParams
+						{
+							.Sequence = SubData->GetSequence(),
+							.SequenceId = ChildId
+						};
+					RenameBindingRecursive(InActor, InParams, ChildSequenceParams);
+				}
+			}
+		}
+	}
+} // UE::AvaSequencer::Private
 
 const FName FAvaSequencer::SidebarDrawerId = TEXT("MotionDesign");
 
@@ -144,6 +213,8 @@ FAvaSequencer::FAvaSequencer(IAvaSequencerProvider& InProvider, FAvaSequencerArg
 	OnSequenceFinishedHandle = UAvaSequencePlayer::OnSequenceFinished().AddLambda(
 		[this](UAvaSequencePlayer*, UAvaSequence*){ return NotifyOnSequenceStopped(); });
 
+	OnActorLabelChangedHandle = FCoreDelegates::OnActorLabelChanged.AddRaw(this, &FAvaSequencer::OnActorLabelChanged);
+
 	// Register sequencer menu extenders.
 	ISequencerModule& SequencerModule = FAvaSequencerUtils::GetSequencerModule();
 	{
@@ -183,6 +254,9 @@ FAvaSequencer::~FAvaSequencer()
 	OnSequenceStartedHandle.Reset();
 	OnSequenceFinishedHandle.Reset();
 
+	FCoreDelegates::OnActorLabelChanged.Remove(OnActorLabelChangedHandle);
+	OnActorLabelChangedHandle.Reset();
+
 	if (GEditor)
 	{
 		GEditor->UnregisterForUndo(this);
@@ -733,6 +807,61 @@ void FAvaSequencer::NotifyOnSequenceStopped()
 	CleanView->Restore();
 }
 
+void FAvaSequencer::OnActorLabelChanged(AActor* InActor)
+{
+	using namespace UE::AvaSequencer;
+
+	// Skip actors that are invalid or preview actors
+	if (!InActor || InActor->bIsEditorPreviewActor)
+	{
+		return;
+	}
+
+	// Ignore processing label changed for PIE
+	const UWorld* const World = InActor->GetWorld();
+	if (World && World->IsPlayInEditor())
+	{
+		return;
+	}
+
+	IAvaSequenceProvider* const SequenceProvider = Provider.GetSequenceProvider();
+	if (!SequenceProvider)
+	{
+		return;
+	}
+
+	UE::MovieScene::FSharedPlaybackStateCreateParams CreateParams;
+
+	for (UAvaSequence* Sequence : SequenceProvider->GetSequences())
+	{
+		if (!Sequence || !Sequence->GetMovieScene())
+		{
+			continue;
+		}
+
+		FMovieSceneSequenceHierarchy SequenceHierarchy;
+		UMovieSceneCompiledDataManager::CompileHierarchy(Sequence, &SequenceHierarchy, EMovieSceneServerClientMask::All);
+
+		const TSharedRef<UE::MovieScene::FSharedPlaybackState> PlaybackState = MakeShared<UE::MovieScene::FSharedPlaybackState>(*Sequence, CreateParams);
+		FMovieSceneEvaluationState EvaluationState;
+		PlaybackState->AddCapabilityRaw(&EvaluationState);
+
+		const Private::FRenameBindingParams RenameParams
+			{
+				.PlaybackState = PlaybackState,
+				.EvaluationState = &EvaluationState,
+				.SequenceHierarchy = &SequenceHierarchy,
+			};
+
+		const Private::FRenameBindingSequenceParams SequenceParams
+			{
+				.Sequence = Sequence,
+				.SequenceId = MovieSceneSequenceID::Root,
+			};
+		Private::RenameBindingRecursive(InActor, RenameParams, SequenceParams);
+	}
+}
+
 void FAvaSequencer::OnMovieSceneBindingsPasted(const TArray<FMovieSceneBinding>& InBindings)
 {
 	UAvaSequence* const Sequence = GetViewedSequence();
diff --git a/Engine/Plugins/Experimental/Avalanche/Source/AvalancheSequencer/Private/AvaSequencer.h b/Engine/Plugins/Experimental/Avalanche/Source/AvalancheSequencer/Private/AvaSequencer.h
--- a/Engine/Plugins/Experimental/Avalanche/Source/AvalancheSequencer/Private/AvaSequencer.h
+++ b/Engine/Plugins/Experimental/Avalanche/Source/AvalancheSequencer/Private/AvaSequencer.h
@@ -115,6 +115,9 @@ public:
 	/** Called when an Sequence has Stopped Playing in either Sequencer or the Sequence Panel */
 	void NotifyOnSequenceStopped();
 
+	/** Called when an actor's label has changed */
+	void OnActorLabelChanged(AActor* InActor);
+
 	void OnMovieSceneBindingsPasted(const TArray<FMovieSceneBinding>& InBindings);
 
 	const TArray<FAvaSequenceItemPtr>& GetRootSequenceItems() const { return RootSequenceItems; }
@@ -319,6 +322,9 @@ private:
 
 	FDelegateHandle OnSequenceFinishedHandle;
 
+	/** Handle to when the actor label has changed */
+	FDelegateHandle OnActorLabelChangedHandle;
+
 	TArray<TSharedRef<FAvaSequencerAction>> SequencerActions;
 
 	/** prevents reentry when synchronizing Sequencer <-> Provider Selection  */
