MotionDesign : Modifiers
- Fix modifier may crash when retrieving shared object in playback
- Pick modifier log verbosity

diff --git a/Engine/Plugins/Experimental/ActorModifier/Source/ActorModifierLayout/Private/Modifiers/ActorModifierGridArrangeModifier.cpp b/Engine/Plugins/Experimental/ActorModifier/Source/ActorModifierLayout/Private/Modifiers/ActorModifierGridArrangeModifier.cpp
--- a/Engine/Plugins/Experimental/ActorModifier/Source/ActorModifierLayout/Private/Modifiers/ActorModifierGridArrangeModifier.cpp
+++ b/Engine/Plugins/Experimental/ActorModifier/Source/ActorModifierLayout/Private/Modifiers/ActorModifierGridArrangeModifier.cpp
@@ -124,8 +124,15 @@ void UActorModifierGridArrangeModifier::Apply()
 	auto GetReversedGridX = [this, GetGridX](const int32 ChildIndex) { return (Count.X - 1) - GetGridX(ChildIndex); };
 	auto GetReversedGridY = [this, GetGridY](const int32 ChildIndex) { return (Count.Y - 1) - GetGridY(ChildIndex); };
 
-	UActorModifierVisibilityShared* VisibilityShared = GetShared<UActorModifierVisibilityShared>(true);
-	UActorModifierTransformShared* LayoutShared = GetShared<UActorModifierTransformShared>(true);
+	constexpr bool bCreate = true;
+	UActorModifierVisibilityShared* VisibilityShared = GetShared<UActorModifierVisibilityShared>(bCreate);
+	UActorModifierTransformShared* LayoutShared = GetShared<UActorModifierTransformShared>(bCreate);
+
+	if (!VisibilityShared || !LayoutShared)
+	{
+		Fail(LOCTEXT("InvalidSharedObject", "Invalid modifier shared object retrieved"));
+		return;
+	}
 
 	TSet<TWeakObjectPtr<AActor>> NewChildrenActorsWeak;
 	for (int32 ChildIndex = 0; ChildIndex < AttachedActorCount; ++ChildIndex)
diff --git a/Engine/Plugins/Experimental/ActorModifier/Source/ActorModifierLayout/Private/Modifiers/ActorModifierJustifyModifier.cpp b/Engine/Plugins/Experimental/ActorModifier/Source/ActorModifierLayout/Private/Modifiers/ActorModifierJustifyModifier.cpp
--- a/Engine/Plugins/Experimental/ActorModifier/Source/ActorModifierLayout/Private/Modifiers/ActorModifierJustifyModifier.cpp
+++ b/Engine/Plugins/Experimental/ActorModifier/Source/ActorModifierLayout/Private/Modifiers/ActorModifierJustifyModifier.cpp
@@ -68,7 +68,14 @@ void UActorModifierJustifyModifier::Apply()
 	const FVector ConstraintVector = GetConstraintVector(BoundsCenter, ActorModified->GetActorLocation());
 	const FVector ChildLocationOffset = ConstraintVector + AlignmentOffset - AnchorOffsetVector;
 
-	UActorModifierTransformShared* LayoutShared = GetShared<UActorModifierTransformShared>(true);
+	constexpr bool bCreate = true;
+	UActorModifierTransformShared* LayoutShared = GetShared<UActorModifierTransformShared>(bCreate);
+
+	if (!LayoutShared)
+	{
+		Fail(LOCTEXT("InvalidSharedObject", "Invalid modifier shared object retrieved"));
+		return;
+	}
 
 	// unregister children transform update callbacks, we don't want them firing while the modifier is updating
 	if (FActorModifierTransformUpdateExtension* TransformExtension = GetExtension<FActorModifierTransformUpdateExtension>())
diff --git a/Engine/Plugins/Experimental/ActorModifier/Source/ActorModifierLayout/Private/Modifiers/ActorModifierRadialArrangeModifier.cpp b/Engine/Plugins/Experimental/ActorModifier/Source/ActorModifierLayout/Private/Modifiers/ActorModifierRadialArrangeModifier.cpp
--- a/Engine/Plugins/Experimental/ActorModifier/Source/ActorModifierLayout/Private/Modifiers/ActorModifierRadialArrangeModifier.cpp
+++ b/Engine/Plugins/Experimental/ActorModifier/Source/ActorModifierLayout/Private/Modifiers/ActorModifierRadialArrangeModifier.cpp
@@ -247,8 +247,15 @@ void UActorModifierRadialArrangeModifier::Apply()
 		return FVector(ChildRadius * SlotCos, ChildRadius * SlotSin, 0);
 	};
 
-	UActorModifierTransformShared* LayoutShared = GetShared<UActorModifierTransformShared>(true);
-	UActorModifierVisibilityShared* VisibilityShared = GetShared<UActorModifierVisibilityShared>(true);
+	constexpr bool bCreate = true;
+	UActorModifierTransformShared* LayoutShared = GetShared<UActorModifierTransformShared>(bCreate);
+	UActorModifierVisibilityShared* VisibilityShared = GetShared<UActorModifierVisibilityShared>(bCreate);
+
+	if (!LayoutShared || !VisibilityShared)
+	{
+		Fail(LOCTEXT("InvalidSharedObject", "Invalid modifier shared object retrieved"));
+		return;
+	}
 
 	TSet<TWeakObjectPtr<AActor>> NewChildrenActorsWeak;
 	for (int32 ChildIndex = 0; ChildIndex < AttachedActorCount; ++ChildIndex)
diff --git a/Engine/Plugins/Experimental/ActorModifierCore/Source/ActorModifierCore/Private/Modifiers/ActorModifierCoreBase.cpp b/Engine/Plugins/Experimental/ActorModifierCore/Source/ActorModifierCore/Private/Modifiers/ActorModifierCoreBase.cpp
--- a/Engine/Plugins/Experimental/ActorModifierCore/Source/ActorModifierCore/Private/Modifiers/ActorModifierCoreBase.cpp
+++ b/Engine/Plugins/Experimental/ActorModifierCore/Source/ActorModifierCore/Private/Modifiers/ActorModifierCoreBase.cpp
@@ -13,24 +13,31 @@ DEFINE_LOG_CATEGORY_STATIC(LogActorModifierCoreBase, Log, All);
 
 UActorModifierCoreSharedObject* UActorModifierCoreBase::GetShared(TSubclassOf<UActorModifierCoreSharedObject> InClass, bool bInCreateIfNone) const
 {
+	UActorModifierCoreSharedObject* SharedObject = nullptr;
+
 	if (!InClass.Get())
 	{
-		return nullptr;
+		return SharedObject;
 	}
 
 	const AActor* ActorModified = GetModifiedActor();
 	if (!ActorModified)
 	{
-		return nullptr;
+		return SharedObject;
 	}
 
 	if (const UActorModifierCoreSubsystem* Subsystem = UActorModifierCoreSubsystem::Get())
 	{
 		ULevel* Level = ActorModified->GetLevel();
-		return Subsystem->GetModifierSharedObject(Level, InClass, bInCreateIfNone);
+		SharedObject = Subsystem->GetModifierSharedObject(Level, InClass, bInCreateIfNone);
+	}
+
+	if (bInCreateIfNone && !SharedObject)
+	{
+		LogModifier(FString::Printf(TEXT("Failed to create and retrieve the specific shared object : %s"), *InClass.Get()->GetName()), /** Force */true, EActorModifierCoreStatus::Error);
 	}
 
-	return nullptr;
+	return SharedObject;
 }
 
 void UActorModifierCoreBase::Apply()
@@ -605,7 +612,7 @@ bool UActorModifierCoreBase::RemoveExtension(const FName& InExtensionType)
 	return false;
 }
 
-void UActorModifierCoreBase::LogModifier(const FString& InLog, bool bInForce) const
+void UActorModifierCoreBase::LogModifier(const FString& InLog, bool bInForce, EActorModifierCoreStatus InType) const
 {
 	if (IsModifierProfiling() || bInForce)
 	{
@@ -613,7 +620,18 @@ void UActorModifierCoreBase::LogModifier(const FString& InLog, bool bInForce) co
 		const FString ModifierLabel = GetModifierName().ToString();
 		const FString ClassLabel = GetClass()->GetName();
 
-		UE_LOG(LogActorModifierCoreBase, Log, TEXT("[%s][%s][%s] %s"), *ActorLabel, *ClassLabel, *ModifierLabel, *InLog);
+		switch (InType)
+		{
+		case EActorModifierCoreStatus::Success:
+			UE_LOG(LogActorModifierCoreBase, Log, TEXT("[%s][%s][%s] %s"), *ActorLabel, *ClassLabel, *ModifierLabel, *InLog);
+			break;
+		case EActorModifierCoreStatus::Warning:
+			UE_LOG(LogActorModifierCoreBase, Warning, TEXT("[%s][%s][%s] %s"), *ActorLabel, *ClassLabel, *ModifierLabel, *InLog);
+			break;
+		case EActorModifierCoreStatus::Error:
+			UE_LOG(LogActorModifierCoreBase, Error, TEXT("[%s][%s][%s] %s"), *ActorLabel, *ClassLabel, *ModifierLabel, *InLog);
+			break;
+		}
 	}
 }
 
diff --git a/Engine/Plugins/Experimental/ActorModifierCore/Source/ActorModifierCore/Private/Modifiers/ActorModifierCoreDefs.cpp b/Engine/Plugins/Experimental/ActorModifierCore/Source/ActorModifierCore/Private/Modifiers/ActorModifierCoreDefs.cpp
--- a/Engine/Plugins/Experimental/ActorModifierCore/Source/ActorModifierCore/Private/Modifiers/ActorModifierCoreDefs.cpp
+++ b/Engine/Plugins/Experimental/ActorModifierCore/Source/ActorModifierCore/Private/Modifiers/ActorModifierCoreDefs.cpp
@@ -364,7 +364,7 @@ void FActorModifierCoreExecutionTask::Restore()
 	Modifiers.Empty();
 }
 
-void FActorModifierCoreExecutionTask::Apply(TConstArrayView<TObjectPtr<UActorModifierCoreBase>> InModifiers, UActorModifierCoreStack* InStack)
+void FActorModifierCoreExecutionTask::Apply(TConstArrayView<UActorModifierCoreBase*> InModifiers, UActorModifierCoreStack* InStack)
 {
 	check(Modifiers.IsEmpty() && CurrentIndex == INDEX_NONE)
 
@@ -454,7 +454,7 @@ void FActorModifierCoreExecutionTask::RemoveModifier(UActorModifierCoreBase& InM
 	Modifiers.Remove(&InModifier);
 }
 
-void FActorModifierCoreExecutionTask::BuildExecutionChain(TConstArrayView<TObjectPtr<UActorModifierCoreBase>> InModifiers, UActorModifierCoreStack* InStack)
+void FActorModifierCoreExecutionTask::BuildExecutionChain(TConstArrayView<UActorModifierCoreBase*> InModifiers, UActorModifierCoreStack* InStack)
 {
 	ModifierStack = InStack;
 	Modifiers = InModifiers;
diff --git a/Engine/Plugins/Experimental/ActorModifierCore/Source/ActorModifierCore/Private/Modifiers/ActorModifierCoreStack.cpp b/Engine/Plugins/Experimental/ActorModifierCore/Source/ActorModifierCore/Private/Modifiers/ActorModifierCoreStack.cpp
--- a/Engine/Plugins/Experimental/ActorModifierCore/Source/ActorModifierCore/Private/Modifiers/ActorModifierCoreStack.cpp
+++ b/Engine/Plugins/Experimental/ActorModifierCore/Source/ActorModifierCore/Private/Modifiers/ActorModifierCoreStack.cpp
@@ -1622,7 +1622,7 @@ void UActorModifierCoreStack::Apply()
 	// restore original state before dirty modifiers
 	RestorePreState();
 
-	ExecutionTask.Apply(Modifiers, this);
+	ExecutionTask.Apply(GetModifiers(), this);
 }
 
 void UActorModifierCoreStack::OnModifierRemoved(EActorModifierCoreDisableReason InReason)
diff --git a/Engine/Plugins/Experimental/ActorModifierCore/Source/ActorModifierCore/Public/Modifiers/ActorModifierCoreBase.h b/Engine/Plugins/Experimental/ActorModifierCore/Source/ActorModifierCore/Public/Modifiers/ActorModifierCoreBase.h
--- a/Engine/Plugins/Experimental/ActorModifierCore/Source/ActorModifierCore/Public/Modifiers/ActorModifierCoreBase.h
+++ b/Engine/Plugins/Experimental/ActorModifierCore/Source/ActorModifierCore/Public/Modifiers/ActorModifierCoreBase.h
@@ -352,7 +352,7 @@ protected:
 	ACTORMODIFIERCORE_API bool RemoveExtension(const FName& InExtensionType);
 
 	/** Logs modifier message if in profiling mode or forced */
-	void LogModifier(const FString& InLog, bool bInForce = false) const;
+	void LogModifier(const FString& InLog, bool bInForce = false, EActorModifierCoreStatus InType = EActorModifierCoreStatus::Success) const;
 
 	void DeferInitializeModifier();
 
diff --git a/Engine/Plugins/Experimental/ActorModifierCore/Source/ActorModifierCore/Public/Modifiers/ActorModifierCoreDefs.h b/Engine/Plugins/Experimental/ActorModifierCore/Source/ActorModifierCore/Public/Modifiers/ActorModifierCoreDefs.h
--- a/Engine/Plugins/Experimental/ActorModifierCore/Source/ActorModifierCore/Public/Modifiers/ActorModifierCoreDefs.h
+++ b/Engine/Plugins/Experimental/ActorModifierCore/Source/ActorModifierCore/Public/Modifiers/ActorModifierCoreDefs.h
@@ -458,7 +458,7 @@ struct FActorModifierCoreExecutionTask
 	void Restore();
 	
 	/** Build the execution chain of dirty modifiers and execute, first -> last */
-	void Apply(TConstArrayView<TObjectPtr<UActorModifierCoreBase>> InModifiers, UActorModifierCoreStack* InStack);
+	void Apply(TConstArrayView<UActorModifierCoreBase*> InModifiers, UActorModifierCoreStack* InStack);
 
 	void Next();
 	void Fail();
@@ -470,7 +470,7 @@ struct FActorModifierCoreExecutionTask
 	void RemoveModifier(UActorModifierCoreBase& InModifier);
 
 private:
-	void BuildExecutionChain(TConstArrayView<TObjectPtr<UActorModifierCoreBase>> InModifiers, UActorModifierCoreStack* InStack);
+	void BuildExecutionChain(TConstArrayView<UActorModifierCoreBase*> InModifiers, UActorModifierCoreStack* InStack);
 	void ExecuteCurrentModifier();
 
 	UPROPERTY()
diff --git a/Engine/Plugins/Experimental/Avalanche/Source/AvalancheModifiers/Private/Modifiers/AvaBooleanModifier.cpp b/Engine/Plugins/Experimental/Avalanche/Source/AvalancheModifiers/Private/Modifiers/AvaBooleanModifier.cpp
--- a/Engine/Plugins/Experimental/Avalanche/Source/AvalancheModifiers/Private/Modifiers/AvaBooleanModifier.cpp
+++ b/Engine/Plugins/Experimental/Avalanche/Source/AvalancheModifiers/Private/Modifiers/AvaBooleanModifier.cpp
@@ -563,23 +563,24 @@ void UAvaBooleanModifier::UpdateMaskVisibility()
 {
 	if (AActor* ActorModified = GetModifiedActor())
 	{
-		UActorModifierVisibilityShared* VisibilityShared = GetShared<UActorModifierVisibilityShared>(true);
-
-		if (Mode != EAvaBooleanMode::None)
+		if (UActorModifierVisibilityShared* VisibilityShared = GetShared<UActorModifierVisibilityShared>(true))
 		{
-			// Save state
-			if (IsModifierEnabled())
+			if (Mode != EAvaBooleanMode::None)
 			{
-				VisibilityShared->SaveActorState(this, ActorModified);
-				VisibilityShared->SetActorVisibility(this, ActorModified, true, false, EActorModifierVisibilityActor::Game);
-				return;
+				// Save state
+				if (IsModifierEnabled())
+				{
+					VisibilityShared->SaveActorState(this, ActorModified);
+					VisibilityShared->SetActorVisibility(this, ActorModified, true, false, EActorModifierVisibilityActor::Game);
+					return;
+				}
 			}
-		}
 
-		// Restore state
-		if (VisibilityShared->IsActorStateSaved(this, ActorModified))
-		{
-			VisibilityShared->RestoreActorState(this, ActorModified);
+			// Restore state
+			if (VisibilityShared->IsActorStateSaved(this, ActorModified))
+			{
+				VisibilityShared->RestoreActorState(this, ActorModified);
+			}
 		}
 	}
 }
diff --git a/Engine/Plugins/Experimental/Avalanche/Source/AvalancheModifiers/Private/Modifiers/AvaHideEmptyModifier.cpp b/Engine/Plugins/Experimental/Avalanche/Source/AvalancheModifiers/Private/Modifiers/AvaHideEmptyModifier.cpp
--- a/Engine/Plugins/Experimental/Avalanche/Source/AvalancheModifiers/Private/Modifiers/AvaHideEmptyModifier.cpp
+++ b/Engine/Plugins/Experimental/Avalanche/Source/AvalancheModifiers/Private/Modifiers/AvaHideEmptyModifier.cpp
@@ -72,6 +72,11 @@ void UAvaHideEmptyModifier::Apply()
 	}
 
 	UActorModifierVisibilityShared* VisibilityShared = GetShared<UActorModifierVisibilityShared>(true);
+	if (!VisibilityShared)
+	{
+		Fail(LOCTEXT("InvalidSharedObject", "Invalid modifier shared object retrieved"));
+		return;
+	}
 
 	if (Text3DComponent->GetText().IsEmpty())
 	{
diff --git a/Engine/Plugins/Experimental/Avalanche/Source/AvalancheModifiers/Private/Modifiers/AvaTranslucentPriorityModifier.cpp b/Engine/Plugins/Experimental/Avalanche/Source/AvalancheModifiers/Private/Modifiers/AvaTranslucentPriorityModifier.cpp
--- a/Engine/Plugins/Experimental/Avalanche/Source/AvalancheModifiers/Private/Modifiers/AvaTranslucentPriorityModifier.cpp
+++ b/Engine/Plugins/Experimental/Avalanche/Source/AvalancheModifiers/Private/Modifiers/AvaTranslucentPriorityModifier.cpp
@@ -88,38 +88,40 @@ void UAvaTranslucentPriorityModifier::SavePreState()
 {
 	Super::SavePreState();
 
-	UAvaTranslucentPriorityModifierShared* SharedObject = GetShared<UAvaTranslucentPriorityModifierShared>(true);
-
-	AActor* ActorModified = GetModifiedActor();
+	constexpr bool bCreate = true;
+	if (UAvaTranslucentPriorityModifierShared* SharedObject = GetShared<UAvaTranslucentPriorityModifierShared>(bCreate))
+	{
+		AActor* ActorModified = GetModifiedActor();
 
-	ChildrenActorsWeak.Reset();
-	PrimitiveComponentsWeak.Reset();
-	ForEachComponent<UPrimitiveComponent>([this, &SharedObject, &ActorModified](UPrimitiveComponent* InComponent)->bool
-		{
-			// Component is already linked to a modifier
-			if (const UAvaTranslucentPriorityModifier* ModifierContext = SharedObject->FindModifierContext(InComponent))
+		ChildrenActorsWeak.Reset();
+		PrimitiveComponentsWeak.Reset();
+		ForEachComponent<UPrimitiveComponent>([this, &SharedObject, &ActorModified](UPrimitiveComponent* InComponent)->bool
 			{
-				const AActor* ModifierActor = ModifierContext->GetModifiedActor();
-
-				// Skip if a translucent priority modifier is attached on the component owner since it has priority over us
-				// or if the modifier handles children and is attached to the current owning actor
-				if (ModifierContext != this
-					&& (ModifierActor == InComponent->GetOwner()
-					|| (ModifierContext->bIncludeChildren && ModifierActor->IsAttachedTo(ActorModified))))
+				// Component is already linked to a modifier
+				if (const UAvaTranslucentPriorityModifier* ModifierContext = SharedObject->FindModifierContext(InComponent))
 				{
-					return true;
+					const AActor* ModifierActor = ModifierContext->GetModifiedActor();
+
+					// Skip if a translucent priority modifier is attached on the component owner since it has priority over us
+					// or if the modifier handles children and is attached to the current owning actor
+					if (ModifierContext != this
+						&& (ModifierActor == InComponent->GetOwner()
+						|| (ModifierContext->bIncludeChildren && ModifierActor->IsAttachedTo(ActorModified))))
+					{
+						return true;
+					}
 				}
-			}
 
-			PrimitiveComponentsWeak.Add(InComponent);
-			ChildrenActorsWeak.Add(InComponent->GetOwner());
+				PrimitiveComponentsWeak.Add(InComponent);
+				ChildrenActorsWeak.Add(InComponent->GetOwner());
 
-			return true;
-		}
-		, EActorModifierCoreComponentType::All
-		, bIncludeChildren ? EActorModifierCoreLookup::SelfAndAllChildren : EActorModifierCoreLookup::Self);
+				return true;
+			}
+			, EActorModifierCoreComponentType::All
+			, bIncludeChildren ? EActorModifierCoreLookup::SelfAndAllChildren : EActorModifierCoreLookup::Self);
 
-	SharedObject->SetComponentsState(this, PrimitiveComponentsWeak);
+		SharedObject->SetComponentsState(this, PrimitiveComponentsWeak);
+	}
 
 	if (FActorModifierTransformUpdateExtension* TransformExtension = GetExtension<FActorModifierTransformUpdateExtension>())
 	{
diff --git a/Engine/Plugins/Experimental/Avalanche/Source/AvalancheModifiers/Private/Modifiers/AvaVisibilityModifier.cpp b/Engine/Plugins/Experimental/Avalanche/Source/AvalancheModifiers/Private/Modifiers/AvaVisibilityModifier.cpp
--- a/Engine/Plugins/Experimental/Avalanche/Source/AvalancheModifiers/Private/Modifiers/AvaVisibilityModifier.cpp
+++ b/Engine/Plugins/Experimental/Avalanche/Source/AvalancheModifiers/Private/Modifiers/AvaVisibilityModifier.cpp
@@ -27,8 +27,15 @@ void UAvaVisibilityModifier::Apply()
 		return;
 	}
 
+	constexpr bool bCreate = true;
+	UActorModifierVisibilityShared* VisibilityShared = GetShared<UActorModifierVisibilityShared>(bCreate);
+	if (!VisibilityShared)
+	{
+		Fail(LOCTEXT("InvalidSharedObject", "Invalid modifier shared object retrieved"));
+		return;
+	}
+
 	const TArray<TWeakObjectPtr<AActor>> AttachedActors = SceneExtension->GetDirectChildrenActor(ModifyActor);
-	UActorModifierVisibilityShared* VisibilityShared = GetShared<UActorModifierVisibilityShared>(true);
 
 #if WITH_EDITOR
 	const bool bHiddenInEditor = ModifyActor->IsTemporarilyHiddenInEditor(/** IncludeParent */false);
