MotionDesign : Shapes
- Fix shape actor name assigned for variant to match tool name, eg: Rectangle vs Square or Cone vs Cylinder

diff --git a/Engine/Plugins/Experimental/Avalanche/Source/AvalancheInteractiveTools/Private/Tools/AvaInteractiveToolsToolBase.cpp b/Engine/Plugins/Experimental/Avalanche/Source/AvalancheInteractiveTools/Private/Tools/AvaInteractiveToolsToolBase.cpp
--- a/Engine/Plugins/Experimental/Avalanche/Source/AvalancheInteractiveTools/Private/Tools/AvaInteractiveToolsToolBase.cpp
+++ b/Engine/Plugins/Experimental/Avalanche/Source/AvalancheInteractiveTools/Private/Tools/AvaInteractiveToolsToolBase.cpp
@@ -351,7 +351,8 @@ AActor* UAvaInteractiveToolsToolBase::SpawnActor(TSubclassOf<AActor> InActorClas
 	{
 		FActorLabelUtilities::SetActorLabelUnique(NewActor, *InActorLabelOverride);
 	}
-	else
+	// Actor factory handles actor label
+	else if (!ActorFactory)
 	{
 		FActorLabelUtilities::SetActorLabelUnique(NewActor, NewActor->GetDefaultActorLabel());
 	}
diff --git a/Engine/Plugins/Experimental/Avalanche/Source/AvalancheShapesEditor/Private/AvaShapeFactory.cpp b/Engine/Plugins/Experimental/Avalanche/Source/AvalancheShapesEditor/Private/AvaShapeFactory.cpp
--- a/Engine/Plugins/Experimental/Avalanche/Source/AvalancheShapesEditor/Private/AvaShapeFactory.cpp
+++ b/Engine/Plugins/Experimental/Avalanche/Source/AvalancheShapesEditor/Private/AvaShapeFactory.cpp
@@ -14,7 +14,7 @@
 UAvaShapeFactory::UAvaShapeFactory()
 {
 	NewActorClass = AAvaShapeActor::StaticClass();
-	MeshSize = FVector(UAvaShapesEditorShapeToolBase::DefaultDim);
+	MeshSize = UAvaShapesEditorShapeToolBase::DefaultParameters.Size;
 }
 
 void UAvaShapeFactory::SetMeshClass(TSubclassOf<UAvaShapeDynamicMeshBase> InMeshClass)
@@ -32,6 +32,11 @@ void UAvaShapeFactory::SetMeshFunction(TFunction<void(UAvaShapeDynamicMeshBase*)
 	MeshFunction = InFunction;
 }
 
+void UAvaShapeFactory::SetMeshNameOverride(const TOptional<FString>& InMeshNameOverride)
+{
+	MeshNameOverride = InMeshNameOverride;
+}
+
 bool UAvaShapeFactory::CanCreateActorFrom(const FAssetData& AssetData, FText& OutErrorMsg)
 {
 	if (!MeshClass)
@@ -91,7 +96,7 @@ void UAvaShapeFactory::PostSpawnActor(UObject* InAsset, AActor* InNewActor)
 
 	if (AAvaShapeActor* NewShapeActor = Cast<AAvaShapeActor>(InNewActor))
 	{
-		FActorLabelUtilities::RenameExistingActor(NewShapeActor, NewShapeActor->GetDefaultActorLabel(), /** Unique*/true);
+		FActorLabelUtilities::RenameExistingActor(NewShapeActor, MeshNameOverride.Get(NewShapeActor->GetDefaultActorLabel()), /** Unique*/true);
 	}
 }
 
@@ -109,3 +114,16 @@ void UAvaShapeFactory::PostPlaceAsset(TArrayView<const FTypedElementHandle> InHa
 		FEngineAnalytics::GetProvider().RecordEvent(TEXT("Editor.Usage.MotionDesign.PlaceActor"), Attributes);
 	}
 }
+
+FString UAvaShapeFactory::GetDefaultActorLabel(UObject* InAsset) const
+{
+	if (const UClass* AssetClass = Cast<UClass>(InAsset))
+	{
+		if (const AAvaShapeActor* CDO = AssetClass->GetDefaultObject<AAvaShapeActor>())
+		{
+			return MeshNameOverride.Get(CDO->GetDefaultActorLabel());
+		}
+	}
+
+	return Super::GetDefaultActorLabel(InAsset);
+}
diff --git a/Engine/Plugins/Experimental/Avalanche/Source/AvalancheShapesEditor/Private/AvaShapeFactory.h b/Engine/Plugins/Experimental/Avalanche/Source/AvalancheShapesEditor/Private/AvaShapeFactory.h
--- a/Engine/Plugins/Experimental/Avalanche/Source/AvalancheShapesEditor/Private/AvaShapeFactory.h
+++ b/Engine/Plugins/Experimental/Avalanche/Source/AvalancheShapesEditor/Private/AvaShapeFactory.h
@@ -4,6 +4,7 @@
 
 #include "ActorFactories/ActorFactory.h"
 #include "DynamicMeshes/AvaShapeDynMeshBase.h"
+#include "Misc/Optional.h"
 #include "AvaShapeFactory.generated.h"
 
 class UAvaShapeDynamicMeshBase;
@@ -19,6 +20,7 @@ public:
 	void SetMeshClass(TSubclassOf<UAvaShapeDynamicMeshBase> InMeshClass);
 	void SetMeshSize(const FVector& InMeshSize);
 	void SetMeshFunction(TFunction<void(UAvaShapeDynamicMeshBase*)> InFunction);
+	void SetMeshNameOverride(const TOptional<FString>& InMeshNameOverride);
 
 protected:
 	//~ Begin UActorFactory
@@ -27,9 +29,11 @@ protected:
 	virtual AActor* SpawnActor(UObject* InAsset, ULevel* InLevel, const FTransform& InTransform, const FActorSpawnParameters& InSpawnParams) override;
 	virtual void PostSpawnActor(UObject* InAsset, AActor* InNewActor) override;
 	virtual void PostPlaceAsset(TArrayView<const FTypedElementHandle> InHandle, const FAssetPlacementInfo& InPlacementInfo, const FPlacementOptions& InPlacementOptions) override;
+	virtual FString GetDefaultActorLabel(UObject* InAsset) const override;
 	//~ End UActorFactory
 
 	TSubclassOf<UAvaShapeDynamicMeshBase> MeshClass;
 	FVector MeshSize = FVector::ZeroVector;
 	TFunction<void(UAvaShapeDynamicMeshBase*)> MeshFunction;
+	TOptional<FString> MeshNameOverride;
 };
diff --git a/Engine/Plugins/Experimental/Avalanche/Source/AvalancheShapesEditor/Private/Tools/AvaShapesEditorShapeToolBase.cpp b/Engine/Plugins/Experimental/Avalanche/Source/AvalancheShapesEditor/Private/Tools/AvaShapesEditorShapeToolBase.cpp
--- a/Engine/Plugins/Experimental/Avalanche/Source/AvalancheShapesEditor/Private/Tools/AvaShapesEditorShapeToolBase.cpp
+++ b/Engine/Plugins/Experimental/Avalanche/Source/AvalancheShapesEditor/Private/Tools/AvaShapesEditorShapeToolBase.cpp
@@ -11,8 +11,7 @@
 #include "InteractiveToolManager.h"
 #include "ToolContextInterfaces.h"
 
-const FVector UAvaShapesEditorShapeToolBase::DefaultDim = FVector(100);
-TFunction<void(UAvaShapeDynamicMeshBase*)> UAvaShapesEditorShapeToolBase::DefaultFunction = [](UAvaShapeDynamicMeshBase*){};
+const UAvaShapesEditorShapeToolBase::FShapeFactoryParameters UAvaShapesEditorShapeToolBase::DefaultParameters = FShapeFactoryParameters();
 
 UAvaShapesEditorShapeToolBase::UAvaShapesEditorShapeToolBase()
 {
@@ -39,16 +38,6 @@ void UAvaShapesEditorShapeToolBase::OnActorSpawned(AActor* InActor)
 	}
 }
 
-FString UAvaShapesEditorShapeToolBase::GetActorNameOverride() const
-{
-	if (const UAvaShapeDynamicMeshBase* DefaultObject = ShapeClass.GetDefaultObject())
-	{
-		return DefaultObject->GetMeshName();
-	}
-
-	return "";
-}
-
 void UAvaShapesEditorShapeToolBase::SetShapeSize(AAvaShapeActor* InShapeActor, const FVector2D& InShapeSize) const
 {
 	if (!InShapeActor)
diff --git a/Engine/Plugins/Experimental/Avalanche/Source/AvalancheShapesEditor/Private/Tools/AvaShapesEditorShapeToolBase.h b/Engine/Plugins/Experimental/Avalanche/Source/AvalancheShapesEditor/Private/Tools/AvaShapesEditorShapeToolBase.h
--- a/Engine/Plugins/Experimental/Avalanche/Source/AvalancheShapesEditor/Private/Tools/AvaShapesEditorShapeToolBase.h
+++ b/Engine/Plugins/Experimental/Avalanche/Source/AvalancheShapesEditor/Private/Tools/AvaShapesEditorShapeToolBase.h
@@ -18,20 +18,27 @@ class UAvaShapesEditorShapeToolBase : public UAvaInteractiveToolsActorToolBase
 public:
 	using UAvaInteractiveToolsActorToolBase::SpawnActor;
 
-	static const FVector DefaultDim;
-	static TFunction<void(UAvaShapeDynamicMeshBase*)> DefaultFunction;
+	struct FShapeFactoryParameters
+	{
+		FVector Size = FVector(100);
+		TFunction<void(UAvaShapeDynamicMeshBase*)> Functor = [](UAvaShapeDynamicMeshBase*){};
+		TOptional<FString> NameOverride;
+	};
+
+	static const FShapeFactoryParameters DefaultParameters;
 
 	UAvaShapesEditorShapeToolBase();
 
 protected:
 	template<typename InMeshClass
 		UE_REQUIRES(std::derived_from<InMeshClass, UAvaShapeDynamicMeshBase>)>
-	static UAvaShapeFactory* CreateFactory(const FVector& InMeshSize = DefaultDim, TFunction<void(UAvaShapeDynamicMeshBase*)> InFunction = DefaultFunction)
+	static UAvaShapeFactory* CreateFactory(const FShapeFactoryParameters& InParameters = DefaultParameters)
 	{
 		UAvaShapeFactory* Factory = CreateActorFactory<UAvaShapeFactory>();
 		Factory->SetMeshClass(InMeshClass::StaticClass());
-		Factory->SetMeshSize(InMeshSize);
-		Factory->SetMeshFunction(InFunction);
+		Factory->SetMeshSize(InParameters.Size);
+		Factory->SetMeshFunction(InParameters.Functor);
+		Factory->SetMeshNameOverride(InParameters.NameOverride);
 		return Factory;
 	}
 
@@ -45,7 +52,6 @@ protected:
 	//~ End UAvaInteractiveToolsToolBase
 
 	virtual void SetShapeSize(AAvaShapeActor* InShapeActor, const FVector2D& InShapeSize) const;
-	virtual FString GetActorNameOverride() const;
 
 	UPROPERTY()
 	TSubclassOf<UAvaShapeDynamicMeshBase> ShapeClass = nullptr;
diff --git a/Engine/Plugins/Experimental/Avalanche/Source/AvalancheShapesEditor/Private/Tools/AvaShapesEditorShapeToolCone.cpp b/Engine/Plugins/Experimental/Avalanche/Source/AvalancheShapesEditor/Private/Tools/AvaShapesEditorShapeToolCone.cpp
--- a/Engine/Plugins/Experimental/Avalanche/Source/AvalancheShapesEditor/Private/Tools/AvaShapesEditorShapeToolCone.cpp
+++ b/Engine/Plugins/Experimental/Avalanche/Source/AvalancheShapesEditor/Private/Tools/AvaShapesEditorShapeToolCone.cpp
@@ -30,7 +30,17 @@ void UAvaShapesEditorShapeToolCone::OnRegisterTool(IAvalancheInteractiveToolsMod
 	};
 
 	InAITModule->RegisterTool(IAvalancheInteractiveToolsModule::CategoryName3D, MoveTemp(ConeToolParameters));
-	
+
+	const FShapeFactoryParameters CylinderFactoryParameters =
+	{
+		.Functor = [](UAvaShapeDynamicMeshBase* InMesh)
+		{
+			UAvaShapeConeDynamicMesh* Cone = Cast<UAvaShapeConeDynamicMesh>(InMesh);
+			Cone->SetTopRadius(1.f);
+		},
+		.NameOverride = FString("Cylinder"),
+	};
+
 	FAvaInteractiveToolsToolParameters CylinderToolParameters = 
 	{
 		FAvaShapesEditorCommands::Get().Tool_Shape_Cylinder,
@@ -42,14 +52,7 @@ void UAvaShapesEditorShapeToolCone::OnRegisterTool(IAvalancheInteractiveToolsMod
 				return UAvaInteractiveToolsToolBuilder::CreateToolBuilder<UAvaShapesEditorShapeToolCone>(InEdMode);
 			}),
 		nullptr,
-		CreateFactory<UAvaShapeConeDynamicMesh>(
-		DefaultDim,
-			[](UAvaShapeDynamicMeshBase* InMesh)
-			{
-				UAvaShapeConeDynamicMesh* Cone = Cast<UAvaShapeConeDynamicMesh>(InMesh);
-				Cone->SetTopRadius(1.f);
-			}
-		)
+		CreateFactory<UAvaShapeConeDynamicMesh>(CylinderFactoryParameters)
 	};
 
 	InAITModule->RegisterTool(IAvalancheInteractiveToolsModule::CategoryName3D, MoveTemp(CylinderToolParameters));
diff --git a/Engine/Plugins/Experimental/Avalanche/Source/AvalancheShapesEditor/Private/Tools/AvaShapesEditorShapeToolLine.cpp b/Engine/Plugins/Experimental/Avalanche/Source/AvalancheShapesEditor/Private/Tools/AvaShapesEditorShapeToolLine.cpp
--- a/Engine/Plugins/Experimental/Avalanche/Source/AvalancheShapesEditor/Private/Tools/AvaShapesEditorShapeToolLine.cpp
+++ b/Engine/Plugins/Experimental/Avalanche/Source/AvalancheShapesEditor/Private/Tools/AvaShapesEditorShapeToolLine.cpp
@@ -63,6 +63,15 @@ void UAvaShapesEditorShapeToolLine::OnRegisterTool(IAvalancheInteractiveToolsMod
 {
 	Super::OnRegisterTool(InAITModule);
 
+	const FShapeFactoryParameters LineFactoryParameters =
+	{
+		.Functor = [](UAvaShapeDynamicMeshBase* InMesh)
+		{
+			UAvaShapeLineDynamicMesh* Line = Cast<UAvaShapeLineDynamicMesh>(InMesh);
+			Line->SetLineWidth(3.f);
+		}
+	};
+
 	FAvaInteractiveToolsToolParameters ToolParameters =
 	{
 		FAvaShapesEditorCommands::Get().Tool_Shape_Line,
@@ -74,14 +83,7 @@ void UAvaShapesEditorShapeToolLine::OnRegisterTool(IAvalancheInteractiveToolsMod
 				return UAvaInteractiveToolsToolBuilder::CreateToolBuilder<UAvaShapesEditorShapeToolLine>(InEdMode);
 			}),
 		nullptr,
-		CreateFactory<UAvaShapeLineDynamicMesh>(
-			DefaultDim,
-			[](UAvaShapeDynamicMeshBase* InMesh)
-			{
-				UAvaShapeLineDynamicMesh* Line = Cast<UAvaShapeLineDynamicMesh>(InMesh);
-				Line->SetLineWidth(3.f);
-			}
-		)
+		CreateFactory<UAvaShapeLineDynamicMesh>(LineFactoryParameters)
 	};
 
 	InAITModule->RegisterTool(IAvalancheInteractiveToolsModule::CategoryName2D, MoveTemp(ToolParameters));
diff --git a/Engine/Plugins/Experimental/Avalanche/Source/AvalancheShapesEditor/Private/Tools/AvaShapesEditorShapeToolRectangle.cpp b/Engine/Plugins/Experimental/Avalanche/Source/AvalancheShapesEditor/Private/Tools/AvaShapesEditorShapeToolRectangle.cpp
--- a/Engine/Plugins/Experimental/Avalanche/Source/AvalancheShapesEditor/Private/Tools/AvaShapesEditorShapeToolRectangle.cpp
+++ b/Engine/Plugins/Experimental/Avalanche/Source/AvalancheShapesEditor/Private/Tools/AvaShapesEditorShapeToolRectangle.cpp
@@ -17,6 +17,15 @@ void UAvaShapesEditorShapeToolRectangle::OnRegisterTool(IAvalancheInteractiveToo
 {
 	Super::OnRegisterTool(InAITModule);
 
+	FShapeFactoryParameters RectangleFactoryParameters =
+	{
+		.Size = FVector(0, 160, 90),
+		.Functor = [](UAvaShapeDynamicMeshBase* InMesh)
+		{
+			InMesh->SetMaterialUVMode(UAvaShapeDynamicMeshBase::MESH_INDEX_PRIMARY, EAvaShapeUVMode::Stretch);
+		}
+	};
+
 	FAvaInteractiveToolsToolParameters RectangleToolParameters = 
 	{
 		FAvaShapesEditorCommands::Get().Tool_Shape_Rectangle,
@@ -28,17 +37,16 @@ void UAvaShapesEditorShapeToolRectangle::OnRegisterTool(IAvalancheInteractiveToo
 				return UAvaInteractiveToolsToolBuilder::CreateToolBuilder<UAvaShapesEditorShapeToolRectangle>(InEdMode);
 			}),
 		nullptr,
-		CreateFactory<UAvaShapeRectangleDynamicMesh>(
-			FVector(0, 160, 90),
-			[](UAvaShapeDynamicMeshBase* InMesh)
-			{
-				InMesh->SetMaterialUVMode(UAvaShapeDynamicMeshBase::MESH_INDEX_PRIMARY, EAvaShapeUVMode::Stretch);
-			}
-		)
+		CreateFactory<UAvaShapeRectangleDynamicMesh>(RectangleFactoryParameters)
 	};
 
 	InAITModule->RegisterTool(IAvalancheInteractiveToolsModule::CategoryName2D, MoveTemp(RectangleToolParameters));
 
+	FShapeFactoryParameters SquareFactoryParameters =
+	{
+		.NameOverride = FString("Square")
+	};
+
 	FAvaInteractiveToolsToolParameters SquareToolParameters = 
 	{
 		FAvaShapesEditorCommands::Get().Tool_Shape_Square,
@@ -50,7 +58,7 @@ void UAvaShapesEditorShapeToolRectangle::OnRegisterTool(IAvalancheInteractiveToo
 				return UAvaInteractiveToolsToolBuilder::CreateToolBuilder<UAvaShapesEditorShapeToolRectangle>(InEdMode);
 			}),
 		nullptr,
-		CreateFactory<UAvaShapeRectangleDynamicMesh>()
+		CreateFactory<UAvaShapeRectangleDynamicMesh>(SquareFactoryParameters)
 	};
 
 	InAITModule->RegisterTool(IAvalancheInteractiveToolsModule::CategoryName2D, MoveTemp(SquareToolParameters));
