Motion Design: Overlay/zoom fix redux
- PIE is now detected as "GEditor->PlayWorld" rather than checking the associated viewport world as the viewport world was inconsistent.

diff --git a/Engine/Plugins/Experimental/Avalanche/Source/AvalancheLevelViewport/Private/Widgets/SAvaLevelViewportCameraBounds.cpp b/Engine/Plugins/Experimental/Avalanche/Source/AvalancheLevelViewport/Private/Widgets/SAvaLevelViewportCameraBounds.cpp
--- a/Engine/Plugins/Experimental/Avalanche/Source/AvalancheLevelViewport/Private/Widgets/SAvaLevelViewportCameraBounds.cpp
+++ b/Engine/Plugins/Experimental/Avalanche/Source/AvalancheLevelViewport/Private/Widgets/SAvaLevelViewportCameraBounds.cpp
@@ -21,16 +21,10 @@ int32 SAvaLevelViewportCameraBounds::OnPaint(const FPaintArgs& InPaintArgs, cons
 
 	if (UAvaViewportSettings const* AvaViewportSettings = GetDefault<UAvaViewportSettings>())
 	{
-		if (TSharedPtr<FAvaLevelViewportClient> AvaLevelViewportClient = AvaLevelViewportClientWeak.Pin())
+		if (!GEditor->PlayWorld)
 		{
-			UWorld* World = AvaLevelViewportClient->GetViewportWorld();
-
-			// Do not draw camera frames during PIE.
-			if (World && World->WorldType == EWorldType::Editor)
-			{
-				DrawCameraBounds(InPaintArgs, InAllottedGeometry, InMyCullingRect, OutDrawElements, InLayerId,
-					AvaViewportSettings->CameraBoundsShadeColor);
-			}
+			DrawCameraBounds(InPaintArgs, InAllottedGeometry, InMyCullingRect, OutDrawElements, InLayerId,
+				AvaViewportSettings->CameraBoundsShadeColor);
 		}
 	}
 
diff --git a/Engine/Plugins/Experimental/Avalanche/Source/AvalancheViewport/Private/Interaction/AvaCameraZoomController.cpp b/Engine/Plugins/Experimental/Avalanche/Source/AvalancheViewport/Private/Interaction/AvaCameraZoomController.cpp
--- a/Engine/Plugins/Experimental/Avalanche/Source/AvalancheViewport/Private/Interaction/AvaCameraZoomController.cpp
+++ b/Engine/Plugins/Experimental/Avalanche/Source/AvalancheViewport/Private/Interaction/AvaCameraZoomController.cpp
@@ -68,7 +68,7 @@ namespace UE::AvaViewport::Private
 
 bool FAvaCameraZoomController::IsCameraZoomPossible()
 {
-	return true;
+	return !GEditor->PlayWorld;
 }
 
 FAvaCameraZoomController::FAvaCameraZoomController(TSharedRef<IAvaViewportClient> InAvaViewportClient, float InFallbackFOV)
@@ -83,19 +83,7 @@ FAvaCameraZoomController::FAvaCameraZoomController(TSharedRef<IAvaViewportClient
 
 bool FAvaCameraZoomController::CanZoom() const
 {
-	if (IsCameraZoomPossible())
-	{
-		if (TSharedPtr<IAvaViewportClient> AvaViewportClient = GetViewportClient())
-		{
-			if (UWorld* World = AvaViewportClient->GetViewportWorld())
-			{
-				// Zoom should not be enabled during PIE.
-				return World->WorldType == EWorldType::Editor;
-			}
-		}
-	}
-
-	return false;
+	return IsCameraZoomPossible();
 }
 
 void FAvaCameraZoomController::SetZoomLevel(uint8 InZoomLevel)
