Motion Design: Guide enable toggle fix

diff --git a/Engine/Plugins/Experimental/Avalanche/Source/AvalancheLevelViewport/Private/SAvaLevelViewport.cpp b/Engine/Plugins/Experimental/Avalanche/Source/AvalancheLevelViewport/Private/SAvaLevelViewport.cpp
--- a/Engine/Plugins/Experimental/Avalanche/Source/AvalancheLevelViewport/Private/SAvaLevelViewport.cpp
+++ b/Engine/Plugins/Experimental/Avalanche/Source/AvalancheLevelViewport/Private/SAvaLevelViewport.cpp
@@ -423,6 +423,8 @@ TSharedPtr<SAvaLevelViewportGuide> SAvaLevelViewport::AddGuide(EOrientation InOr
 
 	Guides.Add(NewGuide);
 
+	SaveGuides();
+
 	GuideCanvas->AddSlot()
 		.Position(NewGuide.Get(), &SAvaLevelViewportGuide::GetPosition)
 		.Size(NewGuide.Get(), &SAvaLevelViewportGuide::GetSize)
diff --git a/Engine/Plugins/Experimental/Avalanche/Source/AvalancheLevelViewport/Private/Widgets/SAvaLevelViewportGuide.cpp b/Engine/Plugins/Experimental/Avalanche/Source/AvalancheLevelViewport/Private/Widgets/SAvaLevelViewportGuide.cpp
--- a/Engine/Plugins/Experimental/Avalanche/Source/AvalancheLevelViewport/Private/Widgets/SAvaLevelViewportGuide.cpp
+++ b/Engine/Plugins/Experimental/Avalanche/Source/AvalancheLevelViewport/Private/Widgets/SAvaLevelViewportGuide.cpp
@@ -594,7 +594,7 @@ float SAvaLevelViewportGuide::GetOffset() const
 	return Offset;
 }
 
-bool SAvaLevelViewportGuide::SetOffset(float Offset)
+bool SAvaLevelViewportGuide::SetOffset(float InOffset)
 {
 	const FAvaLevelViewportGuideFrameAndClient FrameAndClient(ViewportFrameWeak);
 
@@ -611,23 +611,23 @@ bool SAvaLevelViewportGuide::SetOffset(float Offset)
 	switch (Info.Orientation)
 	{
 		case Orient_Horizontal:
-			Offset = FMath::Clamp(
-				Offset,
+			InOffset = FMath::Clamp(
+				InOffset,
 				MinimumViewportEdgeDistance,
 				ViewportSize.Y - MinimumViewportEdgeDistance
 			);
 
-			NewOffsetFraction = Offset / ViewportSize.Y;
+			NewOffsetFraction = InOffset / ViewportSize.Y;
 			break;
 
 		case Orient_Vertical:
-			Offset = FMath::Clamp(
-				Offset,
+			InOffset = FMath::Clamp(
+				InOffset,
 				MinimumViewportEdgeDistance,
 				ViewportSize.X - MinimumViewportEdgeDistance
 			);
 
-			NewOffsetFraction = Offset / ViewportSize.X;
+			NewOffsetFraction = InOffset / ViewportSize.X;
 			break;
 
 		default:
@@ -894,6 +894,27 @@ void SAvaLevelViewportGuide::ToggleEnabled()
 	}
 }
 
+void SAvaLevelViewportGuide::SetState(EAvaViewportGuideState InNewState)
+{
+	Info.State = InNewState;
+
+	TSharedPtr<SAvaLevelViewportFrame> ViewportFrame = GetViewportFrame();
+
+	if (!ViewportFrame.IsValid())
+	{
+		return;
+	}
+
+	TSharedPtr<SAvaLevelViewport> ViewportWidget = ViewportFrame->GetViewportWidget();
+
+	if (!ViewportWidget.IsValid())
+	{
+		return;
+	}
+
+	ViewportWidget->SaveGuides();
+}
+
 ECheckBoxState SAvaLevelViewportGuide::GetLockedCheckState() const
 {
 	return Info.bLocked
diff --git a/Engine/Plugins/Experimental/Avalanche/Source/AvalancheLevelViewport/Private/Widgets/SAvaLevelViewportGuide.h b/Engine/Plugins/Experimental/Avalanche/Source/AvalancheLevelViewport/Private/Widgets/SAvaLevelViewportGuide.h
--- a/Engine/Plugins/Experimental/Avalanche/Source/AvalancheLevelViewport/Private/Widgets/SAvaLevelViewportGuide.h
+++ b/Engine/Plugins/Experimental/Avalanche/Source/AvalancheLevelViewport/Private/Widgets/SAvaLevelViewportGuide.h
@@ -59,12 +59,12 @@ public:
 	void Construct(const FArguments& InArgs, TSharedPtr<SAvaLevelViewportFrame> InViewportFrame, int32 InIndex);
 
 	const FAvaViewportGuideInfo& GetGuideInfo() const { return Info; }
-	void SetState(EAvaViewportGuideState NewState) { Info.State = NewState; }
+	void SetState(EAvaViewportGuideState InNewState);
 	void SetLocked(bool bInLocked) { Info.bLocked = bInLocked; }
 	bool IsBeingDragged() const { return bBeingDragged; }
 	EVisibility GetGuideVisibility() const;
 	float GetOffset() const;
-	bool SetOffset(float Offset);
+	bool SetOffset(float InOffset);
 	FVector2D GetSize() const;
 	FVector2D GetPosition() const;
 	void SetIndex(int32 InNewIndex) { Index = InNewIndex; }
