Motion Design: fixed issue with Translate/Rotate/Scale shortcut keys in viewport while in Motion Design not working after changing viewport types.

diff --git a/Engine/Plugins/Experimental/Avalanche/Source/AvalancheEditorCore/Private/AvaEdMode.cpp b/Engine/Plugins/Experimental/Avalanche/Source/AvalancheEditorCore/Private/AvaEdMode.cpp
--- a/Engine/Plugins/Experimental/Avalanche/Source/AvalancheEditorCore/Private/AvaEdMode.cpp
+++ b/Engine/Plugins/Experimental/Avalanche/Source/AvalancheEditorCore/Private/AvaEdMode.cpp
@@ -28,8 +28,17 @@ void UAvaEdMode::Initialize()
 
 bool UAvaEdMode::UsesToolkits() const
 {
-	// Note: Uses Toolkits returns false as this is for FModeToolkit only
-	return true;
+	// UsesToolkits() is only used in UEdMode::CreateToolkit, FEditorModeTools::ShouldShowModeToolbox (deprecated) and where the Editor Mode Tools processes input.
+	// CreateToolkit is already overriden by UAvaEdMode (below), so UsesToolkits() should return false to block EditorModeTools from processing input.
+	// Note: EditorModeTools processes input via SViewport (before SEditorViewport).
+	// This toolkit command list ends up being appended to the Global Level Editor command list as a parent.
+	// This is so that any action bound to any command list in that parent-child chain can be processed without needing to remap.
+	// However, UToolMenu objects do not fully release their command list shared ref until GC.
+	// Due to these tool menu objects, changing viewport types temporarily breaks the command input processing for viewport-related actions,
+	// and because of how FUICommandList input processing returns true for actions that have invalid objects (since CanExecute will return true if not bound and Execute is not checked).
+	// By having this false, the toolkit commands still get processed in SLevelEditor::OnKeyDown,
+	// but at a much later time, allowing others like valid SEditorViewport instances to process their command list first.
+	return false;
 }
 
 void UAvaEdMode::CreateToolkit()
