Motion Design: Missing viewport toolbar entries
- The viewport toolbar entries on the camera menu have been added to the new toolbar redesign.

diff --git a/Engine/Plugins/Experimental/Avalanche/Source/AvalancheLevelViewport/Private/AvaLevelViewportModule.cpp b/Engine/Plugins/Experimental/Avalanche/Source/AvalancheLevelViewport/Private/AvaLevelViewportModule.cpp
--- a/Engine/Plugins/Experimental/Avalanche/Source/AvalancheLevelViewport/Private/AvaLevelViewportModule.cpp
+++ b/Engine/Plugins/Experimental/Avalanche/Source/AvalancheLevelViewport/Private/AvaLevelViewportModule.cpp
@@ -5,6 +5,7 @@
 #include "AvaLevelViewportStyle.h"
 #include "AvaViewportCameraHistory.h"
 #include "AvaViewportUtils.h"
+#include "Customizations/AvaLevelViewportCameraCustomization.h"
 #include "Modules/ModuleManager.h"
 #include "ToolMenus.h"
 #include "ViewportClient/AvaLevelViewportClient.h"
@@ -46,6 +47,8 @@ void FAvaLevelViewportModule::ShutdownModule()
 
 	FAvaLevelViewportCommands::Unregister();
 
+	FAvaLevelViewportCameraCustomization::Get()->Unregister();
+
 	if (AvaLevelViewportClientCasterDelegateHandle.IsValid())
 	{
 		FAvaViewportUtils::UnregisterViewportClientCaster(AvaLevelViewportClientCasterDelegateHandle);
@@ -65,6 +68,8 @@ void FAvaLevelViewportModule::RegisterMenus()
 		Menu->SetStyleSet(&FAvaLevelViewportStyle::Get());
 		Menu->StyleName = "StatusBar";
 	}
+
+	FAvaLevelViewportCameraCustomization::Get()->Register();
 }
 
 IMPLEMENT_MODULE(FAvaLevelViewportModule, AvalancheLevelViewport)
diff --git a/Engine/Plugins/Experimental/Avalanche/Source/AvalancheLevelViewport/Private/Customizations/AvaLevelViewportCameraCustomization.cpp b/Engine/Plugins/Experimental/Avalanche/Source/AvalancheLevelViewport/Private/Customizations/AvaLevelViewportCameraCustomization.cpp
new file mode 100644
--- /dev/null
+++ b/Engine/Plugins/Experimental/Avalanche/Source/AvalancheLevelViewport/Private/Customizations/AvaLevelViewportCameraCustomization.cpp
@@ -0,0 +1,69 @@
+// Copyright Epic Games, Inc. All Rights Reserved.
+
+#include "Customizations/AvaLevelViewportCameraCustomization.h"
+
+#include "SAvaLevelViewport.h"
+#include "ToolMenu.h"
+#include "ToolMenus.h"
+#include "ViewportToolbar/UnrealEdViewportToolbarContext.h"
+
+namespace UE::AvalancheLevelViewport::Private
+{
+	constexpr const TCHAR* ToolbarMenuName = TEXT("LevelEditor.LevelViewportToolBar.Camera");
+	constexpr const TCHAR* ToolbarSectionName = TEXT("VirtualViewport");
+}
+
+const TSharedRef<FAvaLevelViewportCameraCustomization>& FAvaLevelViewportCameraCustomization::Get()
+{
+	static TSharedRef<FAvaLevelViewportCameraCustomization> Instance = MakeShared<FAvaLevelViewportCameraCustomization>();
+	return Instance;
+}
+
+void FAvaLevelViewportCameraCustomization::Register()
+{
+	using namespace UE::AvalancheLevelViewport::Private;
+
+	UToolMenus* ToolMenus = UToolMenus::Get();
+
+	if (!ToolMenus->IsMenuRegistered(ToolbarMenuName))
+	{
+		ToolMenus->RegisterMenu(ToolbarMenuName);
+	}
+
+	UToolMenu* Menu = ToolMenus->ExtendMenu(ToolbarMenuName);
+
+	Menu->AddDynamicSection(
+		ToolbarSectionName,
+		FNewSectionConstructChoice(FNewToolMenuDelegate::CreateSP(this, &FAvaLevelViewportCameraCustomization::ExtendLevelViewportToolbar))
+	);
+}
+
+void FAvaLevelViewportCameraCustomization::Unregister()
+{
+	using namespace UE::AvalancheLevelViewport::Private;
+
+	UToolMenus* ToolMenus = UToolMenus::Get();
+
+	UToolMenu* Menu = ToolMenus->ExtendMenu(ToolbarMenuName);
+	Menu->RemoveSection(ToolbarSectionName);
+}
+
+void FAvaLevelViewportCameraCustomization::ExtendLevelViewportToolbar(UToolMenu* InToolMenu)
+{
+	UUnrealEdViewportToolbarContext* LevelViewportContext = InToolMenu->FindContext<UUnrealEdViewportToolbarContext>();
+
+	if (!LevelViewportContext)
+	{
+		return;
+	}
+
+	TSharedPtr<SEditorViewport> ViewportWidget = LevelViewportContext->Viewport.Pin();
+
+	if (!ViewportWidget.IsValid() || ViewportWidget->GetWidgetClass().GetWidgetType() != SAvaLevelViewport::StaticWidgetClass().GetWidgetType())
+	{
+		return;
+	}
+
+	TSharedPtr<SAvaLevelViewport> AvaViewportWidget = StaticCastSharedPtr<SAvaLevelViewport>(ViewportWidget);
+	AvaViewportWidget->FillCameraMenu(InToolMenu, /* Include Cameras */ false);
+}
diff --git a/Engine/Plugins/Experimental/Avalanche/Source/AvalancheLevelViewport/Private/Customizations/AvaLevelViewportCameraCustomization.h b/Engine/Plugins/Experimental/Avalanche/Source/AvalancheLevelViewport/Private/Customizations/AvaLevelViewportCameraCustomization.h
new file mode 100644
--- /dev/null
+++ b/Engine/Plugins/Experimental/Avalanche/Source/AvalancheLevelViewport/Private/Customizations/AvaLevelViewportCameraCustomization.h
@@ -0,0 +1,20 @@
+// Copyright Epic Games, Inc. All Rights Reserved.
+
+#pragma once
+
+#include "Templates/SharedPointer.h"
+
+class UToolMenu;
+
+class FAvaLevelViewportCameraCustomization : public TSharedFromThis<FAvaLevelViewportCameraCustomization>
+{
+public:
+	static const TSharedRef<FAvaLevelViewportCameraCustomization>& Get();
+
+	void Register();
+
+	void Unregister();
+
+protected:
+	void ExtendLevelViewportToolbar(UToolMenu* InToolMenu);
+};
\ No newline at end of file
diff --git a/Engine/Plugins/Experimental/Avalanche/Source/AvalancheLevelViewport/Private/SAvaLevelViewport.h b/Engine/Plugins/Experimental/Avalanche/Source/AvalancheLevelViewport/Private/SAvaLevelViewport.h
--- a/Engine/Plugins/Experimental/Avalanche/Source/AvalancheLevelViewport/Private/SAvaLevelViewport.h
+++ b/Engine/Plugins/Experimental/Avalanche/Source/AvalancheLevelViewport/Private/SAvaLevelViewport.h
@@ -153,6 +153,8 @@ public:
 	bool CanResetPilotedCameraTransform() const;
 	void ResetPilotedCameraTransform();
 
+	void FillCameraMenu(UToolMenu* InMenu, bool bInIncludeCamerasAndViewports);
+
 	//~ Begin FSelfRegisteringEditorUndoClient
 	virtual void PostUndo(bool bInSuccess);
 	virtual void PostRedo(bool bInSuccess);
@@ -180,7 +182,6 @@ private:
 	FSlateIcon GetCameraIcon() const;
 
 	TSharedRef<SWidget> GenerateCameraMenu();
-	void FillCameraMenu(UToolMenu* InMenu);
 	bool CanChangeCameraInCameraMenu() const;
 
 	void AddCameraEntries(UToolMenu* InMenu, TArray<TWeakObjectPtr<ACameraActor>> InCameraActors);
diff --git a/Engine/Plugins/Experimental/Avalanche/Source/AvalancheLevelViewport/Private/SAvaLevelViewport_Menus.cpp b/Engine/Plugins/Experimental/Avalanche/Source/AvalancheLevelViewport/Private/SAvaLevelViewport_Menus.cpp
--- a/Engine/Plugins/Experimental/Avalanche/Source/AvalancheLevelViewport/Private/SAvaLevelViewport_Menus.cpp
+++ b/Engine/Plugins/Experimental/Avalanche/Source/AvalancheLevelViewport/Private/SAvaLevelViewport_Menus.cpp
@@ -143,7 +143,7 @@ TSharedRef<SWidget> SAvaLevelViewport::GenerateCameraMenu()
 				{
 					if (TSharedPtr<SAvaLevelViewport> Viewport = Context->GetViewport())
 					{
-						Viewport->FillCameraMenu(InMenu);
+						Viewport->FillCameraMenu(InMenu, /* Include cameras */ true);
 					}
 				}
 			})
@@ -159,7 +159,7 @@ TSharedRef<SWidget> SAvaLevelViewport::GenerateCameraMenu()
 	return ToolMenus->GenerateWidget(MenuName, MenuContext);
 }
 
-void SAvaLevelViewport::FillCameraMenu(UToolMenu* InMenu)
+void SAvaLevelViewport::FillCameraMenu(UToolMenu* InMenu, bool bInIncludeCamerasAndViewports)
 {
 	UWorld* const World = GetWorld();
 	if (!World)
@@ -169,26 +169,29 @@ void SAvaLevelViewport::FillCameraMenu(UToolMenu* InMenu)
 
 	const FAvaLevelViewportCommands& LevelViewportCommands = FAvaLevelViewportCommands::GetInternal();
 
-	TArray<TWeakObjectPtr<ACameraActor>> CameraActors;
-	for (ACameraActor* const CameraActor : TActorRange<ACameraActor>(World))
+	if (bInIncludeCamerasAndViewports)
 	{
-		CameraActors.Add(CameraActor);
-	}
+		TArray<TWeakObjectPtr<ACameraActor>> CameraActors;
+		for (ACameraActor* const CameraActor : TActorRange<ACameraActor>(World))
+		{
+			CameraActors.Add(CameraActor);
+		}
 
-	static const FText CamerasLabel = LOCTEXT("CameraActors", "Placed Cameras");
+		static const FText CamerasLabel = LOCTEXT("CameraActors", "Placed Cameras");
 
-	if (CameraActors.Num() > UE::AvaLevelViewport::Private::MaxCameraEntries)
-	{
-		FToolMenuSection& Section = InMenu->AddSection("CameraActors");
-		Section.AddSubMenu("CameraActors"
-			, CamerasLabel
-			, LOCTEXT("CameraActorsToolTip", "Scene Cameras")
-			, FNewToolMenuDelegate::CreateSP(this, &SAvaLevelViewport::AddCameraEntries, CameraActors));
-	}
-	else
-	{
-		FToolMenuSection& Section = InMenu->AddSection("CameraActors", CamerasLabel);
-		AddCameraEntries(Section, CameraActors);
+		if (CameraActors.Num() > UE::AvaLevelViewport::Private::MaxCameraEntries)
+		{
+			FToolMenuSection& Section = InMenu->AddSection("CameraActors");
+			Section.AddSubMenu("CameraActors"
+				, CamerasLabel
+				, LOCTEXT("CameraActorsToolTip", "Scene Cameras")
+				, FNewToolMenuDelegate::CreateSP(this, &SAvaLevelViewport::AddCameraEntries, CameraActors));
+		}
+		else
+		{
+			FToolMenuSection& Section = InMenu->AddSection("CameraActors", CamerasLabel);
+			AddCameraEntries(Section, CameraActors);
+		}
 	}
 
 	FToolMenuSection& VirtualSection = InMenu->AddSection("VirtualViewport", LOCTEXT("VirtualViewport", "Virtual Viewport"));
@@ -228,6 +231,11 @@ void SAvaLevelViewport::FillCameraMenu(UToolMenu* InMenu)
 		LOCTEXT("CameraResetTransform", "Reset Transform")
 	);
 
+	if (!bInIncludeCamerasAndViewports)
+	{
+		return;
+	}
+
 	int32 ViewportTypeCount = 0;
 
 	if (FLevelEditorModule* LevelEditorModule = FAvaLevelEditorUtils::GetLevelEditorModule())
