Motion Design: Fixed crash on shutdown

diff --git a/Engine/Plugins/Experimental/Avalanche/Source/AvalancheLevelViewport/Private/Customizations/AvaLevelViewportCameraCustomization.cpp b/Engine/Plugins/Experimental/Avalanche/Source/AvalancheLevelViewport/Private/Customizations/AvaLevelViewportCameraCustomization.cpp
--- a/Engine/Plugins/Experimental/Avalanche/Source/AvalancheLevelViewport/Private/Customizations/AvaLevelViewportCameraCustomization.cpp
+++ b/Engine/Plugins/Experimental/Avalanche/Source/AvalancheLevelViewport/Private/Customizations/AvaLevelViewportCameraCustomization.cpp
@@ -42,6 +42,11 @@ void FAvaLevelViewportCameraCustomization::Unregister()
 {
 	using namespace UE::AvalancheLevelViewport::Private;
 
+	if (!UObjectInitialized())
+	{
+		return;
+	}
+
 	UToolMenus* ToolMenus = UToolMenus::Get();
 
 	UToolMenu* Menu = ToolMenus->ExtendMenu(ToolbarMenuName);
