Motion Design: added special support for SetAssetByPath behaviors in'GetEntitiesControlledByController' to return the target entity for this behavior.
SetAssetByPath has a target entity living outside rc action list so existing logic was not taking into account this entity as 'controlled' and was not propagating the rundown rc values properly.

diff --git a/Engine/Plugins/Experimental/Avalanche/Source/AvalancheMedia/Private/Playable/AvaPlayableRemoteControl.cpp b/Engine/Plugins/Experimental/Avalanche/Source/AvalancheMedia/Private/Playable/AvaPlayableRemoteControl.cpp
--- a/Engine/Plugins/Experimental/Avalanche/Source/AvalancheMedia/Private/Playable/AvaPlayableRemoteControl.cpp
+++ b/Engine/Plugins/Experimental/Avalanche/Source/AvalancheMedia/Private/Playable/AvaPlayableRemoteControl.cpp
@@ -8,6 +8,7 @@
 #include "AvaMediaSerializationUtils.h"
 #include "Backends/JsonStructDeserializerBackend.h"
 #include "Backends/JsonStructSerializerBackend.h"
+#include "Behaviour/Builtin/Path/RCSetAssetByPathBehaviour.h"
 #include "Behaviour/RCBehaviour.h"
 #include "Controller/RCController.h"
 #include "IRemoteControlModule.h"
@@ -279,14 +280,19 @@ bool UE::AvaPlayableRemoteControl::GetEntitiesControlledByController(const URemo
 		return false;
 	}
 
-	for (const TObjectPtr<URCBehaviour>& Behaviour : Controller->Behaviours)
+	for (const URCBehaviour* Behavior : Controller->Behaviours)
 	{
-		const TSet<TObjectPtr<URCAction>>& Actions = Behaviour->ActionContainer->GetActions();
-		for (const TObjectPtr<URCAction>& Action : Actions)
+		// Special case: SetAssetByPath behaviors have an extra 'TargetEntity' that these behaviors control outside of RC Actions.
+		if (const URCSetAssetByPathBehaviour* SetAssetByPathBehavior = Cast<URCSetAssetByPathBehaviour>(Behavior))
+		{
+			OutEntityIds.Add(SetAssetByPathBehavior->GetTargetEntityId());
+		}
+
+		for (const URCAction* Action : Behavior->ActionContainer->GetActions())
 		{
 			if (Action)
 			{
-				if (const URCPropertyIdAction* IdentityAction = Cast<URCPropertyIdAction>(Action.Get()))
+				if (const URCPropertyIdAction* IdentityAction = Cast<URCPropertyIdAction>(Action))
 				{
 					if (const URemoteControlPropertyIdRegistry* IdentityRegistry = InRemoteControlPreset->GetPropertyIdRegistry())
 					{
diff --git a/Engine/Plugins/VirtualProduction/RemoteControl/Source/RemoteControlLogic/Private/Behaviour/Builtin/Path/RCSetAssetByPathBehaviour.cpp b/Engine/Plugins/VirtualProduction/RemoteControl/Source/RemoteControlLogic/Private/Behaviour/Builtin/Path/RCSetAssetByPathBehaviour.cpp
--- a/Engine/Plugins/VirtualProduction/RemoteControl/Source/RemoteControlLogic/Private/Behaviour/Builtin/Path/RCSetAssetByPathBehaviour.cpp
+++ b/Engine/Plugins/VirtualProduction/RemoteControl/Source/RemoteControlLogic/Private/Behaviour/Builtin/Path/RCSetAssetByPathBehaviour.cpp
@@ -396,6 +396,11 @@ TWeakPtr<const FRemoteControlEntity> URCSetAssetByPathBehaviour::GetTargetEntity
 	return TargetEntity;
 }
 
+const FGuid& URCSetAssetByPathBehaviour::GetTargetEntityId() const
+{
+	return TargetEntityId;
+}
+
 bool URCSetAssetByPathBehaviour::SetExternalAsset(FString InExternalPath)
 {
 	URCController* Controller = ControllerWeakPtr.Get();
diff --git a/Engine/Plugins/VirtualProduction/RemoteControl/Source/RemoteControlLogic/Public/Behaviour/Builtin/Path/RCSetAssetByPathBehaviour.h b/Engine/Plugins/VirtualProduction/RemoteControl/Source/RemoteControlLogic/Public/Behaviour/Builtin/Path/RCSetAssetByPathBehaviour.h
--- a/Engine/Plugins/VirtualProduction/RemoteControl/Source/RemoteControlLogic/Public/Behaviour/Builtin/Path/RCSetAssetByPathBehaviour.h
+++ b/Engine/Plugins/VirtualProduction/RemoteControl/Source/RemoteControlLogic/Public/Behaviour/Builtin/Path/RCSetAssetByPathBehaviour.h
@@ -130,6 +130,9 @@ public:
 	/** Returns a Pointer to the current Target Entity */
 	TWeakPtr<const FRemoteControlEntity> GetTargetEntity() const;
 
+	/** Returns the current Target Entity Id */
+	const FGuid& GetTargetEntityId() const;
+
 	/** Auxiliary Function to apply to update the Target Texture */
 	void UpdateTargetEntity();
 
