Motion Design: fixed issue where load default layout would cause tabs to disappear

diff --git a/Engine/Plugins/Experimental/Avalanche/Source/AvalancheEditorCore/Private/LevelEditor/AvaLevelEditor.cpp b/Engine/Plugins/Experimental/Avalanche/Source/AvalancheEditorCore/Private/LevelEditor/AvaLevelEditor.cpp
--- a/Engine/Plugins/Experimental/Avalanche/Source/AvalancheEditorCore/Private/LevelEditor/AvaLevelEditor.cpp
+++ b/Engine/Plugins/Experimental/Avalanche/Source/AvalancheEditorCore/Private/LevelEditor/AvaLevelEditor.cpp
@@ -165,10 +165,12 @@ void FAvaLevelEditor::OnLevelEditorCreated(TSharedPtr<ILevelEditor> InLevelEdito
 	SetToolkitHost(InLevelEditor.ToSharedRef());
 
 	DeactivateExtensions();
+	UnregisterTabSpawners();
 
 	if (ShouldAutoActivate())
 	{
 		TryOpenScene();
+		RegisterTabSpawners();
 	}
 }
 
