Motion Design: scene rig supported class to check for sub-classess too

diff --git a/Engine/Plugins/Experimental/Avalanche/Source/AvalancheSceneRig/Private/AvaSceneRigSubsystem.cpp b/Engine/Plugins/Experimental/Avalanche/Source/AvalancheSceneRig/Private/AvaSceneRigSubsystem.cpp
--- a/Engine/Plugins/Experimental/Avalanche/Source/AvalancheSceneRig/Private/AvaSceneRigSubsystem.cpp
+++ b/Engine/Plugins/Experimental/Avalanche/Source/AvalancheSceneRig/Private/AvaSceneRigSubsystem.cpp
@@ -141,9 +141,18 @@ const TSet<TSubclassOf<AActor>>& UAvaSceneRigSubsystem::GetSupportedActorClasses
 	return SupportedActorClasses;
 }
 
-bool UAvaSceneRigSubsystem::IsSupportedActorClass(const TSubclassOf<AActor>& InClass)
+bool UAvaSceneRigSubsystem::IsSupportedActorClass(TSubclassOf<AActor> InClass)
 {
-	return GetSupportedActorClasses().Contains(InClass);
+	// Iterate until the class is invalid (i.e. not an Actor subclass)
+	while (*InClass)
+	{
+		if (SupportedActorClasses.Contains(InClass))
+		{
+			return true;
+		}
+		InClass = InClass->GetSuperClass();
+	}
+	return false;
 }
 
 bool UAvaSceneRigSubsystem::AreActorsSupported(const TArray<AActor*>& InActors)
diff --git a/Engine/Plugins/Experimental/Avalanche/Source/AvalancheSceneRig/Public/AvaSceneRigSubsystem.h b/Engine/Plugins/Experimental/Avalanche/Source/AvalancheSceneRig/Public/AvaSceneRigSubsystem.h
--- a/Engine/Plugins/Experimental/Avalanche/Source/AvalancheSceneRig/Public/AvaSceneRigSubsystem.h
+++ b/Engine/Plugins/Experimental/Avalanche/Source/AvalancheSceneRig/Public/AvaSceneRigSubsystem.h
@@ -50,7 +50,7 @@ public:
 	AVALANCHESCENERIG_API static const TSet<TSubclassOf<AActor>>& GetSupportedActorClasses();
 
 	/** Returns true if the class is a supported Scene Rig class. */
-	AVALANCHESCENERIG_API static bool IsSupportedActorClass(const TSubclassOf<AActor>& InClass);
+	AVALANCHESCENERIG_API static bool IsSupportedActorClass(TSubclassOf<AActor> InClass);
 
 	/** Returns true if the list of actors are supported Scene Rig classes. */
 	AVALANCHESCENERIG_API static bool AreActorsSupported(const TArray<AActor*>& InActors);
