Motion Design: made top,right, left, bottom,front,back,perspective view modes collapsed as these were unmapped actions in Motion Design Viewport previously, but menu entries like the level editor viewport toolbar still added these as entries even though no action was mapped to them. Ideally, these view modes should be supported just like Cinematic Viewport that is camera driven but still supports these. However for 5.6, these view modes will be temporarily collapsed to minimize user confusion.

diff --git a/Engine/Plugins/Experimental/Avalanche/Source/AvalancheLevelViewport/Private/SAvaLevelViewport.cpp b/Engine/Plugins/Experimental/Avalanche/Source/AvalancheLevelViewport/Private/SAvaLevelViewport.cpp
--- a/Engine/Plugins/Experimental/Avalanche/Source/AvalancheLevelViewport/Private/SAvaLevelViewport.cpp
+++ b/Engine/Plugins/Experimental/Avalanche/Source/AvalancheLevelViewport/Private/SAvaLevelViewport.cpp
@@ -87,14 +87,17 @@ void SAvaLevelViewport::BindCommands()
 	const FEditorViewportCommands& EditorViewportCommands = FEditorViewportCommands::Get();
 	FUICommandList& CommandListRef = *CommandList;
 
-	// Remove View Entries as they don't make sense in a Camera-driven Viewport
-	CommandListRef.UnmapAction(EditorViewportCommands.Perspective);
-	CommandListRef.UnmapAction(EditorViewportCommands.Front);
-	CommandListRef.UnmapAction(EditorViewportCommands.Left);
-	CommandListRef.UnmapAction(EditorViewportCommands.Top);
-	CommandListRef.UnmapAction(EditorViewportCommands.Back);
-	CommandListRef.UnmapAction(EditorViewportCommands.Right);
-	CommandListRef.UnmapAction(EditorViewportCommands.Bottom);
+	FUIAction InvalidAction;
+	InvalidAction.IsActionVisibleDelegate.BindLambda([](){ return false; });
+
+	// Hide View Entries for this camera-driven Viewport
+	CommandListRef.MapAction(EditorViewportCommands.Perspective, InvalidAction);
+	CommandListRef.MapAction(EditorViewportCommands.Front, InvalidAction);
+	CommandListRef.MapAction(EditorViewportCommands.Left, InvalidAction);
+	CommandListRef.MapAction(EditorViewportCommands.Top, InvalidAction);
+	CommandListRef.MapAction(EditorViewportCommands.Back, InvalidAction);
+	CommandListRef.MapAction(EditorViewportCommands.Right, InvalidAction);
+	CommandListRef.MapAction(EditorViewportCommands.Bottom, InvalidAction);
 
 	const FAvaLevelViewportCommands& AvaLevelViewportCommands = FAvaLevelViewportCommands::GetInternal();
 
