Motion Design: fixed issue with open settings in "Apply Preset" section from sequence list not opening the motion sequencer settings

diff --git a/Engine/Plugins/Experimental/Avalanche/Source/AvalancheSequencer/Private/Settings/AvaSequencerSettings.cpp b/Engine/Plugins/Experimental/Avalanche/Source/AvalancheSequencer/Private/Settings/AvaSequencerSettings.cpp
--- a/Engine/Plugins/Experimental/Avalanche/Source/AvalancheSequencer/Private/Settings/AvaSequencerSettings.cpp
+++ b/Engine/Plugins/Experimental/Avalanche/Source/AvalancheSequencer/Private/Settings/AvaSequencerSettings.cpp
@@ -45,8 +45,8 @@ namespace UE::AvaSequencer::Private
 
 UAvaSequencerSettings::UAvaSequencerSettings()
 {
-	CategoryName = TEXT("Motion Design");
-	SectionName = TEXT("Sequencer");
+	CategoryName = SettingsCategory;
+	SectionName = SettingsSection;
 
 	DisplayRate.FrameRate = FFrameRate(60000, 1001);
 }
diff --git a/Engine/Plugins/Experimental/Avalanche/Source/AvalancheSequencer/Private/Settings/AvaSequencerSettings.h b/Engine/Plugins/Experimental/Avalanche/Source/AvalancheSequencer/Private/Settings/AvaSequencerSettings.h
--- a/Engine/Plugins/Experimental/Avalanche/Source/AvalancheSequencer/Private/Settings/AvaSequencerSettings.h
+++ b/Engine/Plugins/Experimental/Avalanche/Source/AvalancheSequencer/Private/Settings/AvaSequencerSettings.h
@@ -16,7 +16,7 @@ class UAvaSequencerSettings : public UDeveloperSettings
 
 public:
 	static constexpr const TCHAR* SettingsCategory = TEXT("Motion Design");
-	static constexpr const TCHAR* SettingsSection  = TEXT("AvaSequencerSettings");
+	static constexpr const TCHAR* SettingsSection  = TEXT("Sequencer");
 
 	UAvaSequencerSettings();
 
