Motion Design: fix issue with creating a sequence in Motion Design then undoing causing a 'None' sequence that can't be deleted.

diff --git a/Engine/Plugins/Experimental/Avalanche/Source/AvalancheSequencer/Private/AvaSequencer.cpp b/Engine/Plugins/Experimental/Avalanche/Source/AvalancheSequencer/Private/AvaSequencer.cpp
--- a/Engine/Plugins/Experimental/Avalanche/Source/AvalancheSequencer/Private/AvaSequencer.cpp
+++ b/Engine/Plugins/Experimental/Avalanche/Source/AvalancheSequencer/Private/AvaSequencer.cpp
@@ -1279,19 +1279,27 @@ UAvaSequence* FAvaSequencer::AddSequence(UAvaSequence* const InParentSequence)
 		return nullptr;
 	}
 
-	const FScopedTransaction Transaction(LOCTEXT("AddSequenceTransaction", "Add Sequence"));
+	FScopedTransaction Transaction(LOCTEXT("AddSequenceTransaction", "Add Sequence"));
+
+	if (UObject* SequenceProviderObject = SequenceProvider->ToUObject())
+	{
+		SequenceProviderObject->Modify();
+	}
 
 	UAvaSequence* const NewSequence = CreateSequence();
+	if (!NewSequence)
+	{
+		Transaction.Cancel();
+		return nullptr;
+	}
+
+	SequenceProvider->AddSequence(NewSequence);
 
-	if (IsValid(InParentSequence))
+	if (InParentSequence)
 	{
 		InParentSequence->Modify();
 		InParentSequence->AddChild(NewSequence);
 	}
-	else
-	{
-		SequenceProvider->AddSequence(NewSequence);
-	}
 
 	OnSequenceAddedDelegate.Broadcast(NewSequence);
 
