Motion Design: FIxed viewport mouse location
- The internal SViewport is now used to detect mouse position instead of the outer widget.

diff --git a/Engine/Plugins/Experimental/Avalanche/Source/AvalancheLevelViewport/Private/ViewportClient/AvaLevelViewportClient.cpp b/Engine/Plugins/Experimental/Avalanche/Source/AvalancheLevelViewport/Private/ViewportClient/AvaLevelViewportClient.cpp
--- a/Engine/Plugins/Experimental/Avalanche/Source/AvalancheLevelViewport/Private/ViewportClient/AvaLevelViewportClient.cpp
+++ b/Engine/Plugins/Experimental/Avalanche/Source/AvalancheLevelViewport/Private/ViewportClient/AvaLevelViewportClient.cpp
@@ -32,6 +32,12 @@
 namespace UE::AvaLevelViewport::Private
 {
 	static TSet<FEditorViewportClient*> ViewportClientRegistry;
+
+	class SEditorViewportExposeViewport : public SEditorViewport
+	{
+	public:
+		TSharedPtr<SViewport> GetViewportWidget() const { return ViewportWidget; }
+	};
 }
 
 bool FAvaLevelViewportClient::IsAvaLevelViewportClient(const FEditorViewportClient* InViewportClient)
@@ -254,7 +260,14 @@ FVector2f FAvaLevelViewportClient::GetUnconstrainedViewportMousePosition() const
 		return AbsoluteMousePosition;
 	}
 
-    return (AbsoluteMousePosition - ViewportWidget->GetTickSpaceGeometry().GetAbsolutePosition());
+	TSharedPtr<SViewport> InnerViewportWidget = static_cast<UE::AvaLevelViewport::Private::SEditorViewportExposeViewport*>(ViewportWidget.Get())->GetViewportWidget();
+
+	if (!InnerViewportWidget.IsValid())
+	{
+		return AbsoluteMousePosition;
+	}
+
+    return (AbsoluteMousePosition - InnerViewportWidget->GetTickSpaceGeometry().GetAbsolutePosition());
 }
 
 FVector2f FAvaLevelViewportClient::GetConstrainedViewportMousePosition() const
@@ -268,7 +281,14 @@ FVector2f FAvaLevelViewportClient::GetConstrainedViewportMousePosition() const
 		return AbsoluteMousePosition;
 	}
 
-	const FVector2f ViewportAbsolutePosition = ViewportWidget->GetTickSpaceGeometry().GetAbsolutePosition();
+	TSharedPtr<SViewport> InnerViewportWidget = static_cast<UE::AvaLevelViewport::Private::SEditorViewportExposeViewport*>(ViewportWidget.Get())->GetViewportWidget();
+
+	if (!InnerViewportWidget.IsValid())
+	{
+		return AbsoluteMousePosition;
+	}
+
+	const FVector2f ViewportAbsolutePosition = InnerViewportWidget->GetTickSpaceGeometry().GetAbsolutePosition();
 
 	return (AbsoluteMousePosition
 			- ViewportAbsolutePosition
