Motion Design: fixed issue with reparenting sequence not transacting

diff --git a/Engine/Plugins/Experimental/Avalanche/Source/AvalancheSequence/Private/AvaSequence.cpp b/Engine/Plugins/Experimental/Avalanche/Source/AvalancheSequence/Private/AvaSequence.cpp
--- a/Engine/Plugins/Experimental/Avalanche/Source/AvalancheSequence/Private/AvaSequence.cpp
+++ b/Engine/Plugins/Experimental/Avalanche/Source/AvalancheSequence/Private/AvaSequence.cpp
@@ -150,6 +150,13 @@ int32 UAvaSequence::AddChild(UAvaSequence* InChild)
 		return false;
 	}
 
+#if WITH_EDITOR
+	if (GUndo)
+	{
+		Modify();
+	}
+#endif
+
 	if (InChild->ParentAnimation.IsValid())
 	{
 		InChild->ParentAnimation->RemoveChild(InChild);
@@ -168,6 +175,13 @@ int32 UAvaSequence::AddChild(UAvaSequence* InChild)
 
 bool UAvaSequence::RemoveChild(UAvaSequence* InChild)
 {
+#if WITH_EDITOR
+	if (GUndo)
+	{
+		Modify();
+	}
+#endif
+
 	const bool bRemovedSequence = ChildAnimations.Remove(InChild) > 0;
 
 	if (InChild && bRemovedSequence)
@@ -185,6 +199,13 @@ bool UAvaSequence::RemoveChild(UAvaSequence* InChild)
 
 void UAvaSequence::RemoveAllChildren()
 {
+#if WITH_EDITOR
+	if (GUndo)
+	{
+		Modify();
+	}
+#endif
+
 	for (TWeakObjectPtr<UAvaSequence> Child : ChildAnimations)
 	{
 		if (Child.IsValid())
@@ -203,6 +224,12 @@ void UAvaSequence::RemoveAllChildren()
 
 void UAvaSequence::SetParent(UAvaSequence* InParent)
 {
+#if WITH_EDITOR
+	if (GUndo)
+	{
+		Modify();
+	}
+#endif
 	ParentAnimation = InParent;
 }
 
diff --git a/Engine/Plugins/Experimental/Avalanche/Source/AvalancheSequencer/Private/AvaSequencer.cpp b/Engine/Plugins/Experimental/Avalanche/Source/AvalancheSequencer/Private/AvaSequencer.cpp
--- a/Engine/Plugins/Experimental/Avalanche/Source/AvalancheSequencer/Private/AvaSequencer.cpp
+++ b/Engine/Plugins/Experimental/Avalanche/Source/AvalancheSequencer/Private/AvaSequencer.cpp
@@ -2166,6 +2166,11 @@ void FAvaSequencer::NotifyOnSequenceTreeChanged()
 
 void FAvaSequencer::PostUndo(const bool bInSuccess)
 {
+	if (IAvaSequenceProvider* const SequenceProvider = Provider.GetSequenceProvider())
+	{
+		SequenceProvider->RebuildSequenceTree();
+	}
+
 	// A just-added sequence might be removed due to this undo, so refresh 
 	NotifyOnSequenceTreeChanged();
 }
diff --git a/Engine/Plugins/Experimental/Avalanche/Source/AvalancheSequencer/Private/SequenceTree/DragDropOps/AvaSequenceItemDragDropOp.cpp b/Engine/Plugins/Experimental/Avalanche/Source/AvalancheSequencer/Private/SequenceTree/DragDropOps/AvaSequenceItemDragDropOp.cpp
--- a/Engine/Plugins/Experimental/Avalanche/Source/AvalancheSequencer/Private/SequenceTree/DragDropOps/AvaSequenceItemDragDropOp.cpp
+++ b/Engine/Plugins/Experimental/Avalanche/Source/AvalancheSequencer/Private/SequenceTree/DragDropOps/AvaSequenceItemDragDropOp.cpp
@@ -3,8 +3,11 @@
 #include "AvaSequenceItemDragDropOp.h"
 #include "AvaSequence.h"
 #include "Player/AvaSequenceActorFactory.h"
+#include "ScopedTransaction.h"
 #include "SequenceTree/IAvaSequenceItem.h"
 
+#define LOCTEXT_NAMESPACE "AvaSequenceItemDragDropOp"
+
 TSharedRef<FAvaSequenceItemDragDropOp> FAvaSequenceItemDragDropOp::New(const FAvaSequenceItemPtr& InSequenceItem, UActorFactory* InActorFactory)
 {
 	TSharedRef<FAvaSequenceItemDragDropOp> DragDropOp = MakeShared<FAvaSequenceItemDragDropOp>();
@@ -45,6 +48,8 @@ FReply FAvaSequenceItemDragDropOp::OnDropOnItem(FAvaSequenceItemPtr InTargetItem
 		return FReply::Unhandled();
 	}
 
+	FScopedTransaction Transaction(LOCTEXT("ReparentSequence", "Reparent Sequence"));
+
 	//Check if Target Item is a Parent of Sequence Item (so that we can move the item up the hierarchy)
 	if (TargetSequence == Sequence->GetParent())
 	{
@@ -92,3 +97,5 @@ TSharedPtr<SWidget> FAvaSequenceItemDragDropOp::GetDefaultDecorator() const
 	// Skipping FAssetDragDropOp::GetDefaultDecorator
 	return FDecoratedDragDropOp::GetDefaultDecorator();
 }
+
+#undef LOCTEXT_NAMESPACE
