[Motion Design] Rundown Page Editor: adding controller tooltip to show id and description. (helps debugging rundown controller issues)

diff --git a/Engine/Plugins/Experimental/Avalanche/Source/AvalancheMediaEditor/Private/Rundown/DetailsView/RemoteControl/Controllers/AvaRundownRCControllerItem.cpp b/Engine/Plugins/Experimental/Avalanche/Source/AvalancheMediaEditor/Private/Rundown/DetailsView/RemoteControl/Controllers/AvaRundownRCControllerItem.cpp
--- a/Engine/Plugins/Experimental/Avalanche/Source/AvalancheMediaEditor/Private/Rundown/DetailsView/RemoteControl/Controllers/AvaRundownRCControllerItem.cpp
+++ b/Engine/Plugins/Experimental/Avalanche/Source/AvalancheMediaEditor/Private/Rundown/DetailsView/RemoteControl/Controllers/AvaRundownRCControllerItem.cpp
@@ -7,6 +7,8 @@
 #include "SAvaRundownRCControllerItemRow.h"
 #include "UObject/WeakObjectPtrTemplates.h"
 
+#define LOCTEXT_NAMESPACE "AvaRundownRCControllerItem"
+
 FAvaRundownRCControllerItem::FAvaRundownRCControllerItem(int32 InInstanceIndex, FName InAssetName, URCController* InController, const TSharedRef<IDetailTreeNode>& InTreeNode)
 {
 	AssetName = InAssetName;
@@ -23,6 +25,36 @@ FAvaRundownRCControllerItem::FAvaRundownRCControllerItem(int32 InInstanceIndex,
 			: InController->DisplayName;
 		
 		DisplayNameText = FText::FromName(DisplayName);
+
+		{
+			// Build formatted tool tip lines with information from the controller.
+			TArray<FText> ToolTipLines;
+			ToolTipLines.Reserve(5 + InController->Metadata.Num());
+
+			ToolTipLines.Add(FText::Format(LOCTEXT("Controller_ToolTipLine_Name", "Name: \"{0}\""), FText::FromString(InController->PropertyName.ToString())));
+			ToolTipLines.Add(FText::Format(LOCTEXT("Controller_ToolTipLine_Id", "Id: {0}"), FText::FromString(InController->Id.ToString())));
+			if (!InController->FieldId.IsNone())
+			{
+				ToolTipLines.Add(FText::Format(LOCTEXT("Controller_ToolTipLine_FieldId", "FieldId: \"{0}\""), FText::FromString(InController->FieldId.ToString())));
+			}
+			if (!InController->Description.IsEmpty())
+			{
+				ToolTipLines.Add(FText::Format(LOCTEXT("Controller_ToolTipLine_Desc", "Description: {0}"), InController->Description));
+			}
+			ToolTipLines.Add(FText::Format(LOCTEXT("Controller_ToolTipLine_DisplayIndex", "Display Index: {0}"), InController->DisplayIndex));
+
+			for (const TPair<FName, FString>& MetaDataEntry : InController->Metadata)
+			{
+				if (!MetaDataEntry.Value.IsEmpty())
+				{
+					ToolTipLines.Add(FText::Format(LOCTEXT("Controller_ToolTipLine_Metadata", "Metadata: \"{0}\", Value: \"{1}\""),
+						FText::FromName(MetaDataEntry.Key), FText::FromString(MetaDataEntry.Value)));
+				}
+			}
+
+			ToolTipText = FText::Join(LOCTEXT("Controller_ToolTipLine_Delimiter", "\n"), ToolTipLines);
+		}
+		
 		NodeWidgets = InTreeNode->CreateNodeWidgets();
 
 		// Check if the Controller has a custom widget. In that case, overwrite the value widget with it
@@ -42,3 +74,5 @@ URCController* FAvaRundownRCControllerItem::GetController() const
 {
 	return Controller.Get();
 }
+
+#undef LOCTEXT_NAMESPACE
\ No newline at end of file
diff --git a/Engine/Plugins/Experimental/Avalanche/Source/AvalancheMediaEditor/Private/Rundown/DetailsView/RemoteControl/Controllers/AvaRundownRCControllerItem.h b/Engine/Plugins/Experimental/Avalanche/Source/AvalancheMediaEditor/Private/Rundown/DetailsView/RemoteControl/Controllers/AvaRundownRCControllerItem.h
--- a/Engine/Plugins/Experimental/Avalanche/Source/AvalancheMediaEditor/Private/Rundown/DetailsView/RemoteControl/Controllers/AvaRundownRCControllerItem.h
+++ b/Engine/Plugins/Experimental/Avalanche/Source/AvalancheMediaEditor/Private/Rundown/DetailsView/RemoteControl/Controllers/AvaRundownRCControllerItem.h
@@ -20,6 +20,7 @@ public:
 	TSharedRef<ITableRow> CreateWidget(TSharedRef<SAvaRundownRCControllerPanel> InControllerPanel, const TSharedRef<STableViewBase>& InOwnerTable) const;
 
 	FText GetDisplayName() const { return DisplayNameText; }
+	FText GetToolTipText() const { return ToolTipText; }
 	const FNodeWidgets& GetNodeWidgets() const { return NodeWidgets; }
 	
 	int32 GetInstanceIndex() const { return InstanceIndex; }
@@ -33,6 +34,7 @@ private:
 	int32 DisplayIndex = 0;
 	FName AssetName;
 	FText DisplayNameText;
+	FText ToolTipText;
 	FNodeWidgets NodeWidgets;
 
 	TWeakObjectPtr<URCController> Controller;
diff --git a/Engine/Plugins/Experimental/Avalanche/Source/AvalancheMediaEditor/Private/Rundown/DetailsView/RemoteControl/Controllers/SAvaRundownRCControllerItemRow.cpp b/Engine/Plugins/Experimental/Avalanche/Source/AvalancheMediaEditor/Private/Rundown/DetailsView/RemoteControl/Controllers/SAvaRundownRCControllerItemRow.cpp
--- a/Engine/Plugins/Experimental/Avalanche/Source/AvalancheMediaEditor/Private/Rundown/DetailsView/RemoteControl/Controllers/SAvaRundownRCControllerItemRow.cpp
+++ b/Engine/Plugins/Experimental/Avalanche/Source/AvalancheMediaEditor/Private/Rundown/DetailsView/RemoteControl/Controllers/SAvaRundownRCControllerItemRow.cpp
@@ -32,6 +32,7 @@ TSharedRef<SWidget> SAvaRundownRCControllerItemRow::GenerateWidgetForColumn(cons
 					[
 						SNew(STextBlock)
 						.Text(ItemPtr->GetDisplayName())
+						.ToolTipText(ItemPtr->GetToolTipText())
 					]
 				];
 		}
