Motion Design: fix issue with add sequence -> undo causing crashes on FSequencer

diff --git a/Engine/Plugins/Experimental/Avalanche/Source/AvalancheSequence/Private/AvaSequence.cpp b/Engine/Plugins/Experimental/Avalanche/Source/AvalancheSequence/Private/AvaSequence.cpp
--- a/Engine/Plugins/Experimental/Avalanche/Source/AvalancheSequence/Private/AvaSequence.cpp
+++ b/Engine/Plugins/Experimental/Avalanche/Source/AvalancheSequence/Private/AvaSequence.cpp
@@ -31,6 +31,10 @@
 
 #define LOCTEXT_NAMESPACE "AvaSequence"
 
+#if WITH_EDITOR
+UAvaSequence::FOnSequenceEditUndoEvent UAvaSequence::OnSequenceEditUndoDelegate;
+#endif
+
 UAvaSequence::UAvaSequence(const FObjectInitializer& InObjectInitializer)
 	: Super(InObjectInitializer)
 	, Label(TEXT("Sequence"))
@@ -526,6 +530,12 @@ void UAvaSequence::PostLoad()
 }
 
 #if WITH_EDITOR
+void UAvaSequence::PostEditUndo()
+{
+	Super::PostEditUndo();
+	OnSequenceEditUndoDelegate.Broadcast(this);
+}
+
 void UAvaSequence::OnObjectTransacted(UObject* InObject, const FTransactionObjectEvent& InTransactionEvent)
 {
 	if (MovieScene == InObject)
diff --git a/Engine/Plugins/Experimental/Avalanche/Source/AvalancheSequence/Public/AvaSequence.h b/Engine/Plugins/Experimental/Avalanche/Source/AvalancheSequence/Public/AvaSequence.h
--- a/Engine/Plugins/Experimental/Avalanche/Source/AvalancheSequence/Public/AvaSequence.h
+++ b/Engine/Plugins/Experimental/Avalanche/Source/AvalancheSequence/Public/AvaSequence.h
@@ -38,6 +38,16 @@ public:
 
 	virtual ~UAvaSequence() override;
 
+#if WITH_EDITOR
+	DECLARE_MULTICAST_DELEGATE_OneParam(FOnSequenceEditUndoEvent, UAvaSequence*);
+
+	/** Called on Sequence Edit Undo */
+	static FOnSequenceEditUndoEvent::RegistrationType& OnSequenceEditUndo()
+	{
+		return OnSequenceEditUndoDelegate;
+	}
+#endif
+
 	IAvaSequenceProvider* GetSequenceProvider() const;
 
 	/**
@@ -126,6 +136,7 @@ public:
 	virtual void Serialize(FArchive& Ar) override;
 	virtual void PostLoad() override;
 #if WITH_EDITOR
+	virtual void PostEditUndo() override;
 	void OnObjectTransacted(UObject* InObject, const FTransactionObjectEvent& InTransactionEvent);
 #endif
 	//~ End UObject
@@ -145,6 +156,11 @@ public:
 	void OnWorldCleanup(UWorld* InWorld, bool bInSessionEnded, bool bInCleanupResources);
 
 protected:
+#if WITH_EDITOR
+	/** Called on Sequence Edit Undo */
+	static FOnSequenceEditUndoEvent OnSequenceEditUndoDelegate;
+#endif
+
 	/** Label used to identify a sequence. Multiple sequences can share the same label. */
 	UPROPERTY()
 	FName Label;
diff --git a/Engine/Plugins/Experimental/Avalanche/Source/AvalancheSequencer/Private/AvaSequencer.cpp b/Engine/Plugins/Experimental/Avalanche/Source/AvalancheSequencer/Private/AvaSequencer.cpp
--- a/Engine/Plugins/Experimental/Avalanche/Source/AvalancheSequencer/Private/AvaSequencer.cpp
+++ b/Engine/Plugins/Experimental/Avalanche/Source/AvalancheSequencer/Private/AvaSequencer.cpp
@@ -207,6 +207,8 @@ FAvaSequencer::FAvaSequencer(IAvaSequencerProvider& InProvider, FAvaSequencerArg
 
 	BindCommands();
 
+	OnSequenceEditUndoHandle = UAvaSequence::OnSequenceEditUndo().AddRaw(this, &FAvaSequencer::OnSequenceEditUndo);
+
 	OnSequenceStartedHandle = UAvaSequencePlayer::OnSequenceStarted().AddLambda(
 		[this](UAvaSequencePlayer*, UAvaSequence*){ return NotifyOnSequencePlayed(); });
 
@@ -249,9 +251,13 @@ FAvaSequencer::FAvaSequencer(IAvaSequencerProvider& InProvider, FAvaSequencerArg
 
 FAvaSequencer::~FAvaSequencer()
 {
+	UAvaSequence::OnSequenceEditUndo().Remove(OnSequenceEditUndoHandle);
+	OnSequenceEditUndoHandle.Reset();
+
 	UAvaSequencePlayer::OnSequenceStarted().Remove(OnSequenceStartedHandle);
-	UAvaSequencePlayer::OnSequenceFinished().Remove(OnSequenceFinishedHandle);
 	OnSequenceStartedHandle.Reset();
+
+	UAvaSequencePlayer::OnSequenceFinished().Remove(OnSequenceFinishedHandle);
 	OnSequenceFinishedHandle.Reset();
 
 	FCoreDelegates::OnActorLabelChanged.Remove(OnActorLabelChangedHandle);
@@ -1266,6 +1272,27 @@ void FAvaSequencer::StopSequences_Execute()
 	}
 }
 
+void FAvaSequencer::OnSequenceEditUndo(UAvaSequence* InSequence)
+{
+	if (!InSequence)
+	{
+		return;
+	}
+
+	const UAvaSequence* ViewedSequence = ViewedSequenceWeak.Get(/*bEvenIfPendingKill*/true);
+
+	// Set to the default sequence if the sequence is the one being edited and has just been marked invalid
+	if (InSequence == ViewedSequence && !IsValid(ViewedSequence))
+	{
+		// Since a new viewed sequence is being set, Sequencer will try to restore the current viewed sequence (now invalid)
+		// and ensure fail when trying to resolve its weak ptrs (defaulting to bEvenIfPendingKill=false).
+		// To circumvent this, temporarily unmark this sequence as garbage, so that these weak ptrs can resolve properly.
+		InSequence->ClearGarbage();
+		SetViewedSequence(GetDefaultSequence());
+		InSequence->MarkAsGarbage();
+	}
+}
+
 bool FAvaSequencer::CanAddSequence() const
 {
 	return GetProvider().CanEditOrPlaySequences();
diff --git a/Engine/Plugins/Experimental/Avalanche/Source/AvalancheSequencer/Private/AvaSequencer.h b/Engine/Plugins/Experimental/Avalanche/Source/AvalancheSequencer/Private/AvaSequencer.h
--- a/Engine/Plugins/Experimental/Avalanche/Source/AvalancheSequencer/Private/AvaSequencer.h
+++ b/Engine/Plugins/Experimental/Avalanche/Source/AvalancheSequencer/Private/AvaSequencer.h
@@ -253,6 +253,9 @@ public:
 	bool StopSequences_CanExecute() const;
 	void StopSequences_Execute();
 
+	/** Called when PostUndo has been called for the given sequence */
+	void OnSequenceEditUndo(UAvaSequence* InSequence);
+
 	//~ Begin FEditorUndoClient
 	virtual void PostUndo(bool bInSuccess) override;
 	virtual void PostRedo(bool bInSuccess) override;
@@ -318,6 +321,9 @@ private:
 
 	TSharedRef<FAvaSequencerCleanView> CleanView;
 
+	/** Handle to the Sequence Edit Undo delegate*/
+	FDelegateHandle OnSequenceEditUndoHandle;
+
 	FDelegateHandle OnSequenceStartedHandle;
 
 	FDelegateHandle OnSequenceFinishedHandle;
