MotionDesign : Modifiers
- Fix bend modifier crash caused by multiple disconnected components like Text

diff --git a/Engine/Plugins/Experimental/Avalanche/Source/AvalancheModifiers/Private/Modifiers/AvaBendModifier.cpp b/Engine/Plugins/Experimental/Avalanche/Source/AvalancheModifiers/Private/Modifiers/AvaBendModifier.cpp
--- a/Engine/Plugins/Experimental/Avalanche/Source/AvalancheModifiers/Private/Modifiers/AvaBendModifier.cpp
+++ b/Engine/Plugins/Experimental/Avalanche/Source/AvalancheModifiers/Private/Modifiers/AvaBendModifier.cpp
@@ -4,6 +4,8 @@
 
 #include "Async/Async.h"
 #include "Components/DynamicMeshComponent.h"
+#include "DynamicMeshEditor.h"
+#include "Selections/MeshConnectedComponents.h"
 #include "SpaceDeformerOps/BendMeshOp.h"
 
 #define LOCTEXT_NAMESPACE "AvaBendModifier"
@@ -28,52 +30,77 @@ void UAvaBendModifier::Apply()
 		Fail(LOCTEXT("InvalidDynamicMeshComponent", "Invalid dynamic mesh component on modified actor"));
 		return;
 	}
-	
-	// no need to bend here!
-	if (Angle == KINDA_SMALL_NUMBER)
+
+	if (FMath::IsNearlyZero(Angle))
 	{
 		Next();
 		return;
 	}
 
-	const FTransform BendFrame(BendRotation, BendPosition);
-	if (!BendFrame.IsValid())
+	const FTransform BendTransform(BendRotation, BendPosition);
+	if (!BendTransform.IsValid())
 	{
 		Next();
 		return;
 	}
 	
 	using namespace UE::Geometry;
-
-	// update extent
-	DynMeshComp->ProcessMesh([this](const FDynamicMesh3& ProcessMesh)
-	{
-		const FBox MeshBounds = static_cast<FBox>(ProcessMesh.GetBounds(true));
-		// this is the extent value used for maximum extent, meaning the bend will be done over the entire geometry along its X axis
-		ModifiedMeshMaxExtent = MeshBounds.GetSize().Z / 2.0;
-	});
 			
 	constexpr float LowerExtent = 10.0f;
-	const float BendExtent = ModifiedMeshMaxExtent * Extent;
-			
-	const FFrame3d BendFrameAsFrame = FFrame3d(BendFrame);
+	const FFrame3d BendFrame = FFrame3d(BendTransform);
 
-	DynMeshComp->GetDynamicMesh()->EditMesh([this, BendFrameAsFrame, BendExtent, LowerExtent](FDynamicMesh3& EditMesh) 
+	DynMeshComp->GetDynamicMesh()->EditMesh([this, BendFrame, LowerExtent](FDynamicMesh3& EditMesh) 
 	{
-		const TSharedPtr<FDynamicMesh3> TmpMeshPtr = MakeShared<FDynamicMesh3>();
-		*TmpMeshPtr = MoveTemp(EditMesh);
-			    
-		FBendMeshOp BendOperation;
-		BendOperation.OriginalMesh = TmpMeshPtr;
-		BendOperation.GizmoFrame = BendFrameAsFrame;
-		BendOperation.LowerBoundsInterval = (bSymmetricExtents) ? -BendExtent : -LowerExtent;
-		BendOperation.UpperBoundsInterval = BendExtent;
-		BendOperation.BendDegrees = Angle;
-		BendOperation.bLockBottom = !bBidirectional;
-		BendOperation.CalculateResult(nullptr);
-		TUniquePtr<FDynamicMesh3> NewResultMesh = BendOperation.ExtractResult();				
-		FDynamicMesh3* NewResultMeshPtr = NewResultMesh.Release();
-		EditMesh = MoveTemp(*NewResultMeshPtr);
+		FMeshConnectedComponents Components(&EditMesh);
+		Components.FindConnectedTriangles();
+
+		TArray<TSharedPtr<FDynamicMesh3>> Submeshes;
+		for (const FMeshConnectedComponents::FComponent& MeshComponent : Components.Components)
+		{
+			TSharedRef<FDynamicMesh3> Submesh = MakeShared<FDynamicMesh3>();
+			Submesh->EnableMatchingAttributes(EditMesh);
+
+			FDynamicMeshEditor Editor(&Submesh.Get());
+
+			FMeshIndexMappings IndexMap;
+			FDynamicMeshEditResult Result;
+			Editor.AppendTriangles(&EditMesh, MeshComponent.Indices, IndexMap, Result);
+
+			if (Submesh->TriangleCount() > 0)
+			{
+				Submeshes.Add(MoveTemp(Submesh));
+			}
+		}
+
+		for (const TSharedPtr<FDynamicMesh3>& Submesh : Submeshes)
+		{
+			const FBox MeshBounds = static_cast<FBox>(Submesh->GetBounds(true));
+			const float BendExtent = (MeshBounds.GetSize().Z / 2.0) * Extent;
+
+			FBendMeshOp BendOperation;
+			BendOperation.OriginalMesh = Submesh;
+			BendOperation.GizmoFrame = BendFrame;
+			BendOperation.LowerBoundsInterval = (bSymmetricExtents) ? -BendExtent : -LowerExtent;
+			BendOperation.UpperBoundsInterval = BendExtent;
+			BendOperation.BendDegrees = Angle;
+			BendOperation.bLockBottom = !bBidirectional;
+			BendOperation.CalculateResult(nullptr);
+			*Submesh = MoveTemp(*BendOperation.ExtractResult().Release());
+		}
+
+		if (!Submeshes.IsEmpty())
+		{
+			FDynamicMesh3 MergeMesh;
+			MergeMesh.EnableMatchingAttributes(EditMesh);
+			FDynamicMeshEditor Editor(&MergeMesh);
+			for (const TSharedPtr<FDynamicMesh3>& Submesh : Submeshes)
+			{
+				FMeshIndexMappings IndexMap;
+				Editor.AppendMesh(Submesh.Get(), IndexMap);
+			}
+
+			EditMesh = MoveTemp(MergeMesh);
+		}
 
 	}, EDynamicMeshChangeType::GeneralEdit, EDynamicMeshAttributeChangeFlags::Unknown, false);
 
diff --git a/Engine/Plugins/Experimental/Avalanche/Source/AvalancheModifiers/Public/Modifiers/AvaBendModifier.h b/Engine/Plugins/Experimental/Avalanche/Source/AvalancheModifiers/Public/Modifiers/AvaBendModifier.h
--- a/Engine/Plugins/Experimental/Avalanche/Source/AvalancheModifiers/Public/Modifiers/AvaBendModifier.h
+++ b/Engine/Plugins/Experimental/Avalanche/Source/AvalancheModifiers/Public/Modifiers/AvaBendModifier.h
@@ -97,7 +97,4 @@ protected:
 
 	UPROPERTY(EditInstanceOnly, Setter="SetBidirectional", Getter="GetBidirectional", Category="Bend", meta=(AllowPrivateAccess="true"))
 	bool bBidirectional = false;
-
-	UPROPERTY()
-	float ModifiedMeshMaxExtent = 50;
 };
