MotionDesign : Shapes
- Fix UV mode behavior not updating properly
- Fix undo causes crash after aligning shapes
- Disable two sided material by default to avoid inverted normals
- Removing blueprint specifier and instead exposing getter/setter where relevant
- Fix modifiers not updating due to multiple mesh events broadcasted

diff --git a/Engine/Plugins/Experimental/Avalanche/Source/AvalancheEditor/Private/DetailView/Customizations/AvaMeshesDetailCustomization.cpp b/Engine/Plugins/Experimental/Avalanche/Source/AvalancheEditor/Private/DetailView/Customizations/AvaMeshesDetailCustomization.cpp
--- a/Engine/Plugins/Experimental/Avalanche/Source/AvalancheEditor/Private/DetailView/Customizations/AvaMeshesDetailCustomization.cpp
+++ b/Engine/Plugins/Experimental/Avalanche/Source/AvalancheEditor/Private/DetailView/Customizations/AvaMeshesDetailCustomization.cpp
@@ -59,7 +59,7 @@ void FAvaMeshesDetailCustomization::CustomizeDetails(IDetailLayoutBuilder& InDet
 		{
 			if (bFirstElement)
 			{
-				const TArray<FName> MeshNames = MeshGenerator->GetMeshDataNames();
+				const TArray<FName> MeshNames = MeshGenerator->GetMeshSectionNames();
 				for (int32 Index : MeshGenerator->GetMeshesIndexes())
 				{
 					if (MeshNames.IsValidIndex(Index))
@@ -522,7 +522,7 @@ void FAvaMeshesDetailCustomization::ApplySizeToTexture(UAvaShapeDynamicMeshBase*
 		return;
 	}
 
-	const FAvaShapeParametricMaterial* PrimaryParametricMaterial = InMeshGenerator->GetParametricMaterial(UAvaShapeDynamicMeshBase::MESH_INDEX_PRIMARY);
+	const FAvaShapeParametricMaterial* PrimaryParametricMaterial = InMeshGenerator->GetParametricMaterialPtr(UAvaShapeDynamicMeshBase::MESH_INDEX_PRIMARY);
 	if (!PrimaryParametricMaterial)
 	{
 		return;
diff --git a/Engine/Plugins/Experimental/Avalanche/Source/AvalancheMask/Private/Handling/AvaMaskAvaShapeMaterialCollectionHandle.cpp b/Engine/Plugins/Experimental/Avalanche/Source/AvalancheMask/Private/Handling/AvaMaskAvaShapeMaterialCollectionHandle.cpp
--- a/Engine/Plugins/Experimental/Avalanche/Source/AvalancheMask/Private/Handling/AvaMaskAvaShapeMaterialCollectionHandle.cpp
+++ b/Engine/Plugins/Experimental/Avalanche/Source/AvalancheMask/Private/Handling/AvaMaskAvaShapeMaterialCollectionHandle.cpp
@@ -385,7 +385,7 @@ bool FAvaMaskAvaShapeMaterialCollectionHandle::RefreshMaterials()
 			{
 				if (ShapeComponent->IsMaterialType(InMeshIdx, EMaterialType::Parametric))
 				{
-					if (ShapeComponent->GetParametricMaterial(InMeshIdx))
+					if (ShapeComponent->GetParametricMaterialPtr(InMeshIdx))
 					{
 						MaterialHandle = HandleSubsystem->MakeHandle<IAvaMaskMaterialHandle>(FAvaMaskMaterialReference(ShapeComponent, InMeshIdx), UE::AvaMask::Internal::HandleTag);
 
diff --git a/Engine/Plugins/Experimental/Avalanche/Source/AvalancheMask/Private/Handling/AvaParametricMaterialHandle.cpp b/Engine/Plugins/Experimental/Avalanche/Source/AvalancheMask/Private/Handling/AvaParametricMaterialHandle.cpp
--- a/Engine/Plugins/Experimental/Avalanche/Source/AvalancheMask/Private/Handling/AvaParametricMaterialHandle.cpp
+++ b/Engine/Plugins/Experimental/Avalanche/Source/AvalancheMask/Private/Handling/AvaParametricMaterialHandle.cpp
@@ -10,7 +10,7 @@ namespace UE::AvaMask::Private
 	{
 		if (UAvaShapeDynamicMeshBase* ShapeComponent = InParametricMaterial.GetTypedObject<UAvaShapeDynamicMeshBase>())
 		{
-			return ShapeComponent->GetParametricMaterial(InParametricMaterial.Index);
+			return ShapeComponent->GetParametricMaterialPtr(InParametricMaterial.Index);
 		}
 		return nullptr;
 	}
diff --git a/Engine/Plugins/Experimental/Avalanche/Source/AvalancheShapes/Private/DynamicMeshes/AvaShape2DDynMeshBase.cpp b/Engine/Plugins/Experimental/Avalanche/Source/AvalancheShapes/Private/DynamicMeshes/AvaShape2DDynMeshBase.cpp
--- a/Engine/Plugins/Experimental/Avalanche/Source/AvalancheShapes/Private/DynamicMeshes/AvaShape2DDynMeshBase.cpp
+++ b/Engine/Plugins/Experimental/Avalanche/Source/AvalancheShapes/Private/DynamicMeshes/AvaShape2DDynMeshBase.cpp
@@ -197,7 +197,6 @@ void UAvaShape2DDynMeshBase::OnSizeChanged()
 		return;
 	}
 
-	const bool bWasMinSize = PreEditSize2D.GetMin() <= UAvaShapeDynamicMeshBase::MinSizeValue;
 	if (PreEditSize2D.IsZero())
 	{
 		PreEditSize2D = Size2D;
@@ -210,19 +209,6 @@ void UAvaShape2DDynMeshBase::OnSizeChanged()
 	PreEditSize2D = Size2D;
 
 	ScaleVertices(Scale);
-
-	for (FAvaSnapPoint& SnapPoint : LocalSnapPoints)
-	{
-		SnapPoint.Location.Y *= Scale.X;
-		SnapPoint.Location.Z *= Scale.Y;
-	}
-
-	MarkVerticesDirty();
-
-	if (bWasMinSize)
-	{
-		MarkAllMeshesDirty();
-	}
 }
 
 bool UAvaShape2DDynMeshBase::CreateUVs(FAvaShapeMesh& InMesh, FAvaShapeMaterialUVParameters& InParams)
diff --git a/Engine/Plugins/Experimental/Avalanche/Source/AvalancheShapes/Private/DynamicMeshes/AvaShape3DDynMeshBase.cpp b/Engine/Plugins/Experimental/Avalanche/Source/AvalancheShapes/Private/DynamicMeshes/AvaShape3DDynMeshBase.cpp
--- a/Engine/Plugins/Experimental/Avalanche/Source/AvalancheShapes/Private/DynamicMeshes/AvaShape3DDynMeshBase.cpp
+++ b/Engine/Plugins/Experimental/Avalanche/Source/AvalancheShapes/Private/DynamicMeshes/AvaShape3DDynMeshBase.cpp
@@ -147,7 +147,6 @@ void UAvaShape3DDynMeshBase::OnSizeChanged()
 		return;
 	}
 
-	const bool bWasMinSize = PreEditSize3D.GetMin() <= UAvaShapeDynamicMeshBase::MinSizeValue;
 	if (PreEditSize3D.IsZero())
 	{
 		PreEditSize3D = Size3D;
@@ -157,13 +156,6 @@ void UAvaShape3DDynMeshBase::OnSizeChanged()
 	PreEditSize3D = Size3D;
 
 	ScaleVertices(Scale);
-
-	MarkVerticesDirty();
-
-	if (bWasMinSize)
-	{
-		MarkAllMeshesDirty();
-	}
 }
 
 void UAvaShape3DDynMeshBase::OnScaledSizeChanged()
diff --git a/Engine/Plugins/Experimental/Avalanche/Source/AvalancheShapes/Private/DynamicMeshes/AvaShapeDynMeshBase.cpp b/Engine/Plugins/Experimental/Avalanche/Source/AvalancheShapes/Private/DynamicMeshes/AvaShapeDynMeshBase.cpp
--- a/Engine/Plugins/Experimental/Avalanche/Source/AvalancheShapes/Private/DynamicMeshes/AvaShapeDynMeshBase.cpp
+++ b/Engine/Plugins/Experimental/Avalanche/Source/AvalancheShapes/Private/DynamicMeshes/AvaShapeDynMeshBase.cpp
@@ -6,6 +6,7 @@
 #include "Async/ParallelFor.h"
 #include "AvaShapeActor.h"
 #include "Components/DynamicMeshComponent.h"
+#include "Containers/Ticker.h"
 #include "DynamicMesh/DynamicMesh3.h"
 #include "DynamicMesh/DynamicMeshAttributeSet.h"
 #include "DynamicMesh/MeshTransforms.h"
@@ -65,14 +66,12 @@ namespace UE::Private::AvaShape
 UAvaShapeDynamicMeshBase::UAvaShapeDynamicMeshBase(const FLinearColor& InVertexColor, float InUniformScaledSize, bool bInAllowEditSize)
 	: bAllowEditSize(bInAllowEditSize)
 	, UniformScaledSize(InUniformScaledSize)
-	, VertexColor(InVertexColor)
 	, bUsePrimaryMaterialEverywhere(true)
 	, bMeshRegistered(false)
 	, bInCreateMode(false)
 	, bAllMeshDirty(true)
 	, bAnyMeshDirty(bAllMeshDirty)
 	, bVerticesDirty(false)
-	, bColorsDirty(false)
 {
 	if (!IsTemplate())
 	{
@@ -85,28 +84,23 @@ UAvaShapeDynamicMeshBase::UAvaShapeDynamicMeshBase(const FLinearColor& InVertexC
 	}
 }
 
-EMaterialType& UAvaShapeDynamicMeshBase::GetMaterialType(int32 MeshIndex)
+EMaterialType UAvaShapeDynamicMeshBase::GetMaterialType(int32 InMeshIndex)
 {
-	FAvaShapeMeshData* MeshData = GetMeshData(MeshIndex);
-	if (!ensure(MeshData != nullptr))
+	if (const FAvaShapeMeshData* MeshData = GetMeshData(InMeshIndex))
 	{
-		static EMaterialType DefaultMaterialType = EMaterialType::Default;
-		return DefaultMaterialType;
+		return MeshData->MaterialType;
 	}
 
-	return MeshData->MaterialType;
+	return EMaterialType::Default;
 }
 
-bool UAvaShapeDynamicMeshBase::SetMaterialType(int32 MeshIndex, EMaterialType Type)
+void UAvaShapeDynamicMeshBase::SetMaterialType(int32 InMeshIndex, EMaterialType InType)
 {
-	if (FAvaShapeMeshData* MeshData = GetMeshData(MeshIndex))
+	if (FAvaShapeMeshData* MeshData = GetMeshData(InMeshIndex))
 	{
-		MeshData->MaterialType = Type;
-		OnMaterialTypeChanged(MeshIndex);
-		return true;
+		MeshData->MaterialType = InType;
+		OnMaterialTypeChanged(InMeshIndex);
 	}
-
-	return false;
 }
 
 bool UAvaShapeDynamicMeshBase::IsMaterialType(int32 MeshIndex, EMaterialType Type)
@@ -133,7 +127,6 @@ void UAvaShapeDynamicMeshBase::SetupMeshes()
 			// register other meshes
 			RegisterMeshes();
 
-			const bool bIs2DShape = MeshDatas.Num() == 1;
 			for (const int32& MeshIdx : GetMeshesIndexes())
 			{
 				FAvaShapeMeshData* MeshData = GetMeshData(MeshIdx);
@@ -146,9 +139,6 @@ void UAvaShapeDynamicMeshBase::SetupMeshes()
 				// Enable translucency by default
 				MeshData->ParametricMaterial.SetTranslucency(EAvaShapeParametricMaterialTranslucency::Enabled);
 
-				// Disable two sided materials to avoid inverted normal for 3D shapes
-				MeshData->ParametricMaterial.SetUseTwoSidedMaterial(bIs2DShape);
-
 				OnMaterialTypeChanged(MeshIdx);
 			}
 		}
@@ -182,7 +172,7 @@ void UAvaShapeDynamicMeshBase::OnRegister()
 
 	if (bAnyMeshDirty)
 	{
-		RequireUpdate();
+		RequestUpdate();
 	}
 }
 
@@ -299,37 +289,34 @@ void UAvaShapeDynamicMeshBase::ToggleGizmo_Implementation(
 	}
 }
 
-void UAvaShapeDynamicMeshBase::UpdateDynamicMesh()
+void UAvaShapeDynamicMeshBase::UpdateMesh()
 {
-	RestoreCachedMesh();
-	CachedMesh.Reset();
+	RestoreCachedMesh(/** Reset */true);
 
-	bool bUpdateMesh = false;
+	bool bUpdateMesh = GenerateMesh();
 
-	if (bAllMeshDirty || bAnyMeshDirty)
-	{
-		bUpdateMesh = GenerateMesh();
-	}
+	UpdateVertices();
 
-	if (bVerticesDirty)
-	{
-		UpdateVertices();
-	}
+	// Apply vertices, triangles, normals
+	ApplyMeshSectionChanges();
 
-	for (const int32 MeshIndex : GetMeshesIndexes())
+	bUpdateMesh |= GenerateUV();
+
+	if (bUpdateMesh)
 	{
-		FAvaShapeMesh* Mesh = GetMesh(MeshIndex);
-		if (Mesh->bUpdateRequired)
-		{
-			OnMeshSectionUpdated(*Mesh);
-		}
+		GenerateTangents();
 	}
 
+	// Apply Update UV, tangents
+	ApplyMeshSectionChanges();
+
+	// Offset to match alignments
 	if (bHasNewMeshRegenWorldLocation)
 	{
-		AActor* OwningActor = GetOwner();
 		bHasNewMeshRegenWorldLocation = false;
-		FVector Delta = MeshRegenWorldLocation - OwningActor->GetRootComponent()->GetComponentLocation();
+
+		AActor* OwningActor = GetOwner();
+		const FVector Delta = MeshRegenWorldLocation - OwningActor->GetRootComponent()->GetComponentLocation();
 
 #if WITH_EDITOR
 		FActorElementEditorViewportInteractionCustomization::ApplyDeltaToActor(
@@ -342,17 +329,22 @@ void UAvaShapeDynamicMeshBase::UpdateDynamicMesh()
 		MeshRegenWorldLocation = FVector::ZeroVector;
 	}
 
-	GenerateUV();
+	SaveCachedMesh();
+}
 
-	if (bUpdateMesh)
+void UAvaShapeDynamicMeshBase::ApplyMeshSectionChanges()
+{
+	for (const int32 MeshIndex : GetMeshesIndexes())
 	{
-		GenerateTangents();
+		FAvaShapeMesh* Mesh = GetMesh(MeshIndex);
+		if (Mesh->bUpdateRequired)
+		{
+			OnMeshSectionUpdated(*Mesh);
+		}
 	}
-
-	SaveCachedMesh();
 }
 
-void UAvaShapeDynamicMeshBase::RunUpdate(bool bDoAsync)
+void UAvaShapeDynamicMeshBase::RunUpdate()
 {
 	if (MeshUpdateStatus == EAvaDynamicMeshUpdateState::UpdateInProgress)
 	{
@@ -366,52 +358,51 @@ void UAvaShapeDynamicMeshBase::RunUpdate(bool bDoAsync)
 
 	MeshUpdateStatus = EAvaDynamicMeshUpdateState::UpdateInProgress;
 
-	auto ExecUpdate = [this]()
+	auto ExecUpdate = [this]()->bool
 	{
-		if (IsValid(this) && !HasAnyFlags(RF_BeginDestroyed | RF_FinishDestroyed | RF_ClassDefaultObject))
-		{
-			if (!UE::IsSavingPackage(this))
-			{
-				FActorModifierCoreScopedLock ModifierLock(UActorModifierCoreSubsystem::Get()->GetActorModifierStack(GetShapeActor()));
+		UDynamicMeshComponent* DynamicMeshComponent = GetShapeMeshComponent();
 
-				UpdateDynamicMesh();
-				MeshUpdateStatus = EAvaDynamicMeshUpdateState::UpToDate;
-				OnMeshUpdateFinished();
-			}
-			else
+		if (IsValid(this)
+			&& IsValid(DynamicMeshComponent)
+			&& !HasAnyFlags(RF_BeginDestroyed | RF_FinishDestroyed | RF_ClassDefaultObject))
+		{
+			if (UE::IsSavingPackage(this))
 			{
-				MeshUpdateStatus = EAvaDynamicMeshUpdateState::UpdateRequired;
+				return false;
 			}
+
+			FActorModifierCoreScopedLock ModifierLock(UActorModifierCoreSubsystem::Get()->GetActorModifierStack(GetShapeActor()));
+
+			UpdateMesh();
+			DynamicMeshComponent->NotifyMeshModified();
+			MeshUpdateStatus = EAvaDynamicMeshUpdateState::UpToDate;
+			OnMeshUpdateFinished();
 		}
+
+		return true;
 	};
 
-	if (bDoAsync)
+	if (GetRunAsync())
 	{
-		TWeakObjectPtr<UAvaShapeDynamicMeshBase> ThisWeak(this);
-		Async(EAsyncExecution::TaskGraphMainThread, [ThisWeak, ExecUpdate]()
+		FTSTicker::GetCoreTicker().AddTicker(FTickerDelegate::CreateWeakLambda(this, [this, ExecUpdate](float)
 		{
-			UAvaShapeDynamicMeshBase* const This = ThisWeak.Get(true);
-			if (!This)
-			{
-				return;
-			}
-
-			ExecUpdate();
-		});
+			// execute next tick if failed
+			return !ExecUpdate();
+		}));
 	}
-	else
+	else if (!ExecUpdate())
 	{
-		ExecUpdate();
+		MeshUpdateStatus = EAvaDynamicMeshUpdateState::UpdateRequired;
 	}
 }
 
-void UAvaShapeDynamicMeshBase::RequireUpdate(bool bDoAsync)
+void UAvaShapeDynamicMeshBase::RequestUpdate()
 {
 	// only runs if the mesh was up to date previously
 	if (MeshUpdateStatus == EAvaDynamicMeshUpdateState::UpToDate)
 	{
 		MeshUpdateStatus = EAvaDynamicMeshUpdateState::UpdateRequired;
-		RunUpdate(bDoAsync);
+		RunUpdate();
 	}
 }
 
@@ -444,6 +435,35 @@ const FString& UAvaShapeDynamicMeshBase::GetMeshName() const
 	return MeshName;
 }
 
+void UAvaShapeDynamicMeshBase::SetSize2D(const FVector2D& InSize2D)
+{
+	if (!bAllowEditSize)
+	{
+		return;
+	}
+
+	FVector Size3D = GetSize3D();
+	Size3D.Y = InSize2D.X;
+	Size3D.Z = InSize2D.Y;
+	SetSize3D(Size3D);
+}
+
+FVector2D UAvaShapeDynamicMeshBase::GetSize2D() const
+{
+	const FVector Size3D = GetSize3D();
+	return FVector2D(Size3D.Y, Size3D.Z);
+}
+
+bool UAvaShapeDynamicMeshBase::IsMeshSectionVisible(int32 InIndex) const
+{
+	if (const FAvaShapeMeshData* MeshData = GetMeshData(InIndex))
+	{
+		return MeshData->bMeshVisible;
+	}
+
+	return false;
+}
+
 void UAvaShapeDynamicMeshBase::SetUsePrimaryMaterialEverywhere(bool bInUse)
 {
 	if (bUsePrimaryMaterialEverywhere == bInUse)
@@ -477,7 +497,7 @@ void UAvaShapeDynamicMeshBase::OnUsePrimaryMaterialEverywhereChanged()
 	{
 		// create a copy
 		const EMaterialType PrimaryMaterialType = GetMaterialType(MESH_INDEX_PRIMARY);
-		const FAvaShapeParametricMaterial PrimaryParametricMaterial = *GetParametricMaterial(MESH_INDEX_PRIMARY);
+		const FAvaShapeParametricMaterial PrimaryParametricMaterial = GetParametricMaterial(MESH_INDEX_PRIMARY);
 
 		const TSet<int32> Keys = GetMeshesIndexes();
 
@@ -513,7 +533,7 @@ void UAvaShapeDynamicMeshBase::OnUsePrimaryMaterialEverywhereChanged()
 			MeshData->bUsesPrimaryMaterialParams = false;
 
 			// Reset to proper state if it was a parametric material
-			if (GetParametricMaterial(MESH_INDEX_PRIMARY)->IsParametricMaterial(GetMaterial(MeshIndex)))
+			if (GetParametricMaterialPtr(MESH_INDEX_PRIMARY)->IsParametricMaterial(GetMaterial(MeshIndex)))
 			{
 				MeshData->MaterialType = EMaterialType::Parametric;
 				MeshData->Material = MeshData->ParametricMaterial.GetOrCreateMaterial(GetOwner());
@@ -534,18 +554,26 @@ void UAvaShapeDynamicMeshBase::OnMeshTransformChanged(USceneComponent* InCompone
 	bZeroScaled = bIsZeroScaled;
 }
 
-bool UAvaShapeDynamicMeshBase::SetOverridePrimaryUVParams(int32 MeshIndex, bool bOverride)
+void UAvaShapeDynamicMeshBase::SetOverridePrimaryUVParams(int32 InMeshIndex, bool bInOverride)
 {
-	if (MeshIndex == MESH_INDEX_PRIMARY)
+	if (InMeshIndex != MESH_INDEX_PRIMARY)
 	{
-		return false;
+		if (FAvaShapeMeshData* MeshData = GetMeshData(InMeshIndex))
+		{
+			MeshData->bOverridePrimaryUVParams = bInOverride;
+			OnUsesPrimaryUVParamsChanged(InMeshIndex);
+		}
 	}
+}
 
-	if (FAvaShapeMeshData* MeshData = GetMeshData(MeshIndex))
+bool UAvaShapeDynamicMeshBase::GetOverridePrimaryUVParams(int32 InMeshIndex) const
+{
+	if (InMeshIndex != MESH_INDEX_PRIMARY)
 	{
-		MeshData->bOverridePrimaryUVParams = bOverride;
-		OnUsesPrimaryUVParamsChanged(MeshIndex);
-		return true;
+		if (const FAvaShapeMeshData* MeshData = GetMeshData(InMeshIndex))
+		{
+			return MeshData->bOverridePrimaryUVParams;
+		}
 	}
 
 	return false;
@@ -566,6 +594,11 @@ bool UAvaShapeDynamicMeshBase::IsMeshSizeValid() const
 
 bool UAvaShapeDynamicMeshBase::GenerateMesh()
 {
+	if (!bAllMeshDirty && !bAnyMeshDirty)
+	{
+		return false;
+	}
+
 	SetupMeshes();
 
 	bAllMeshDirty = false;
@@ -611,7 +644,6 @@ bool UAvaShapeDynamicMeshBase::GenerateMesh()
 		}
 	}
 
-	bColorsDirty = true;
 	bVerticesDirty = false;
 
 	return true;
@@ -621,6 +653,7 @@ bool UAvaShapeDynamicMeshBase::GenerateUV()
 {
 	const TSet<int32> Keys = GetMeshesIndexes();
 
+	bool bUVUpdated = false;
 	for (const int32 MeshIndex : Keys)
 	{
 		FAvaShapeMeshData* MeshData = GetMeshData(MeshIndex);
@@ -638,6 +671,9 @@ bool UAvaShapeDynamicMeshBase::GenerateUV()
 		MeshData->Mesh.UVs.Empty();
 		if (CreateUVs(MeshData->Mesh, *GetInUseMaterialUVParams(MeshIndex)))
 		{
+			bUVUpdated = true;
+			MeshData->Mesh.bUpdateRequired = true;
+
 			if (MeshIndex != MESH_INDEX_PRIMARY)
 			{
 				MeshData->MaterialUVParams.bUVsDirty = false;
@@ -645,9 +681,9 @@ bool UAvaShapeDynamicMeshBase::GenerateUV()
 		}
 	}
 
-	GetMaterialUVParams(MESH_INDEX_PRIMARY)->bUVsDirty = false;
+	GetMaterialUVParamsPtr(MESH_INDEX_PRIMARY)->bUVsDirty = false;
 
-	return true;
+	return bUVUpdated;
 }
 
 void UAvaShapeDynamicMeshBase::OnMeshUpdateFinished()
@@ -664,7 +700,7 @@ void UAvaShapeDynamicMeshBase::OnMeshUpdateFinished()
 				return;
 			}
 
-			This->RestoreCachedMesh();
+			This->RestoreCachedMesh(/** Reset */false);
 		};
 
 		// When stack is restored function will be called
@@ -685,30 +721,27 @@ FAvaShapeMesh* UAvaShapeDynamicMeshBase::GetMesh(int32 MeshIndex)
 	return nullptr;
 }
 
-bool UAvaShapeDynamicMeshBase::IsValidMeshIndex(int32 MeshIndex) const
+bool UAvaShapeDynamicMeshBase::IsValidMeshIndex(int32 InMeshIndex) const
 {
-	return MeshDatas.Contains(MeshIndex);
+	return MeshDatas.Contains(InMeshIndex);
 }
 
-bool UAvaShapeDynamicMeshBase::SetMaterial(int32 MeshIndex, UMaterialInterface* NewMaterial)
+void UAvaShapeDynamicMeshBase::SetMaterial(int32 InMeshIndex, UMaterialInterface* InNewMaterial)
 {
-	if (MeshIndex == INDEX_NONE)
+	if (InMeshIndex == INDEX_NONE)
 	{
 		SetUsePrimaryMaterialEverywhere(true);
-		MeshIndex = MESH_INDEX_PRIMARY;
+		InMeshIndex = MESH_INDEX_PRIMARY;
 	}
-	else if (MeshIndex != MESH_INDEX_PRIMARY && GetUsePrimaryMaterialEverywhere())
+	else if (InMeshIndex != MESH_INDEX_PRIMARY && GetUsePrimaryMaterialEverywhere())
 	{
 		SetUsePrimaryMaterialEverywhere(false);
 	}
 
-	if (SetMaterialDirect(MeshIndex, NewMaterial))
+	if (SetMaterialDirect(InMeshIndex, InNewMaterial))
 	{
-		OnMaterialChanged(MeshIndex);
-		return true;
+		OnMaterialChanged(InMeshIndex);
 	}
-
-	return false;
 }
 
 bool UAvaShapeDynamicMeshBase::SetMaterialDirect(int32 MeshIndex, UMaterialInterface* NewMaterial)
@@ -723,21 +756,20 @@ bool UAvaShapeDynamicMeshBase::SetMaterialDirect(int32 MeshIndex, UMaterialInter
 	return false;
 }
 
-const FAvaShapeParametricMaterial*
-	UAvaShapeDynamicMeshBase::GetParametricMaterial(int32 MeshIndex) const
+const FAvaShapeParametricMaterial& UAvaShapeDynamicMeshBase::GetParametricMaterial(int32 InMeshIndex) const
 {
-	if (const FAvaShapeMeshData* MeshData = GetMeshData(MeshIndex))
+	if (const FAvaShapeMeshData* MeshData = GetMeshData(InMeshIndex))
 	{
-		return &MeshData->ParametricMaterial;
+		return MeshData->ParametricMaterial;
 	}
 
-	return nullptr;
+	static const FAvaShapeParametricMaterial DefaultParametricMaterial;
+	return DefaultParametricMaterial;
 }
 
-FAvaShapeParametricMaterial*
-	UAvaShapeDynamicMeshBase::GetParametricMaterial(int32 MeshIndex)
+FAvaShapeParametricMaterial* UAvaShapeDynamicMeshBase::GetParametricMaterialPtr(int32 InMeshIndex)
 {
-	if (FAvaShapeMeshData* MeshData = GetMeshData(MeshIndex))
+	if (FAvaShapeMeshData* MeshData = GetMeshData(InMeshIndex))
 	{
 		return &MeshData->ParametricMaterial;
 	}
@@ -745,22 +777,18 @@ FAvaShapeParametricMaterial*
 	return nullptr;
 }
 
-bool UAvaShapeDynamicMeshBase::SetParametricMaterial(int32 MeshIndex,
-	const FAvaShapeParametricMaterial& NewMaterialParams)
+void UAvaShapeDynamicMeshBase::SetParametricMaterial(int32 InMeshIndex, const FAvaShapeParametricMaterial& InNewMaterialParams)
 {
-	if (FAvaShapeMeshData* MeshData = GetMeshData(MeshIndex))
+	if (FAvaShapeMeshData* MeshData = GetMeshData(InMeshIndex))
 	{
-		MeshData->ParametricMaterial = NewMaterialParams;
-		OnParametricMaterialChanged(MeshIndex);
-		return true;
+		MeshData->ParametricMaterial = InNewMaterialParams;
+		OnParametricMaterialChanged(InMeshIndex);
 	}
-
-	return false;
 }
 
-UMaterialInterface* UAvaShapeDynamicMeshBase::GetMaterial(int32 MeshIndex) const
+UMaterialInterface* UAvaShapeDynamicMeshBase::GetMaterial(int32 InMeshIndex) const
 {
-	if (const FAvaShapeMeshData* MeshData = GetMeshData(MeshIndex))
+	if (const FAvaShapeMeshData* MeshData = GetMeshData(InMeshIndex))
 	{
 		return MeshData->Material;
 	}
@@ -783,14 +811,15 @@ const FAvaShapeMaterialUVParameters* UAvaShapeDynamicMeshBase::GetInUseMaterialU
 	return nullptr;
 }
 
-const FAvaShapeMaterialUVParameters* UAvaShapeDynamicMeshBase::GetMaterialUVParams(int32 MeshIndex) const
+const FAvaShapeMaterialUVParameters& UAvaShapeDynamicMeshBase::GetMaterialUVParams(int32 MeshIndex) const
 {
 	if (const FAvaShapeMeshData* MeshData = GetMeshData(MeshIndex))
 	{
-		return &MeshData->MaterialUVParams;
+		return MeshData->MaterialUVParams;
 	}
 
-	return nullptr;
+	static const FAvaShapeMaterialUVParameters DefaultMaterialUVParams;
+	return DefaultMaterialUVParams;
 }
 
 FAvaShapeMaterialUVParameters* UAvaShapeDynamicMeshBase::GetInUseMaterialUVParams(int32 MeshIndex)
@@ -808,7 +837,7 @@ FAvaShapeMaterialUVParameters* UAvaShapeDynamicMeshBase::GetInUseMaterialUVParam
 	return nullptr;
 }
 
-FAvaShapeMaterialUVParameters* UAvaShapeDynamicMeshBase::GetMaterialUVParams(int32 MeshIndex)
+FAvaShapeMaterialUVParameters* UAvaShapeDynamicMeshBase::GetMaterialUVParamsPtr(int32 MeshIndex)
 {
 	if (FAvaShapeMeshData* MeshData = GetMeshData(MeshIndex))
 	{
@@ -818,30 +847,23 @@ FAvaShapeMaterialUVParameters* UAvaShapeDynamicMeshBase::GetMaterialUVParams(int
 	return nullptr;
 }
 
-bool UAvaShapeDynamicMeshBase::SetMaterialUVParams(int32 MeshIndex, const FAvaShapeMaterialUVParameters& InParams)
+void UAvaShapeDynamicMeshBase::SetMaterialUVParams(int32 InMeshIndex, const FAvaShapeMaterialUVParameters& InParams)
 {
-	if (FAvaShapeMeshData* MeshData = GetMeshData(MeshIndex))
+	if (FAvaShapeMeshData* MeshData = GetMeshData(InMeshIndex))
 	{
 		MeshData->MaterialUVParams = InParams;
-		OnUVParamsChanged(MeshIndex);
-		return true;
+		OnUVParamsChanged(InMeshIndex);
 	}
-
-	return false;
 }
 
 EAvaShapeUVMode UAvaShapeDynamicMeshBase::GetMaterialUVMode(int32 MeshIndex) const
 {
-	if (GetMaterialUVParams(MeshIndex) == nullptr)
-	{
-		return EAvaShapeUVMode::Stretch;
-	}
-	return GetMaterialUVParams(MeshIndex)->Mode;
+	return GetMaterialUVParams(MeshIndex).Mode;
 }
 
 bool UAvaShapeDynamicMeshBase::SetMaterialUVMode(int32 MeshIndex, EAvaShapeUVMode InUVMode)
 {
-	FAvaShapeMaterialUVParameters* Params = GetMaterialUVParams(MeshIndex);
+	FAvaShapeMaterialUVParameters* Params = GetMaterialUVParamsPtr(MeshIndex);
 
 	if (Params == nullptr)
 	{
@@ -861,18 +883,13 @@ bool UAvaShapeDynamicMeshBase::SetMaterialUVMode(int32 MeshIndex, EAvaShapeUVMod
 
 EAvaAnchors UAvaShapeDynamicMeshBase::GetMaterialUVAnchorPreset(int32 MeshIndex) const
 {
-	if (GetMaterialUVParams(MeshIndex) == nullptr)
-	{
-		return EAvaAnchors::TopLeft;
-	}
-
-	return GetMaterialUVParams(MeshIndex)->AnchorPreset;
+	return GetMaterialUVParams(MeshIndex).AnchorPreset;
 }
 
 bool UAvaShapeDynamicMeshBase::SetMaterialUVAnchorPreset(int32 MeshIndex,
 	EAvaAnchors InUVAnchorPreset)
 {
-	FAvaShapeMaterialUVParameters* Params = GetMaterialUVParams(MeshIndex);
+	FAvaShapeMaterialUVParameters* Params = GetMaterialUVParamsPtr(MeshIndex);
 
 	if (Params == nullptr)
 	{
@@ -892,17 +909,12 @@ bool UAvaShapeDynamicMeshBase::SetMaterialUVAnchorPreset(int32 MeshIndex,
 
 float UAvaShapeDynamicMeshBase::GetMaterialUVRotation(int32 MeshIndex) const
 {
-	if (GetMaterialUVParams(MeshIndex) == nullptr)
-	{
-		return 0.f;
-	}
-
-	return GetMaterialUVParams(MeshIndex)->Rotation;
+	return GetMaterialUVParams(MeshIndex).Rotation;
 }
 
 bool UAvaShapeDynamicMeshBase::SetMaterialUVRotation(int32 MeshIndex, float InUVRotation)
 {
-	FAvaShapeMaterialUVParameters* Params = GetMaterialUVParams(MeshIndex);
+	FAvaShapeMaterialUVParameters* Params = GetMaterialUVParamsPtr(MeshIndex);
 
 	if (Params == nullptr)
 	{
@@ -927,17 +939,12 @@ bool UAvaShapeDynamicMeshBase::SetMaterialUVRotation(int32 MeshIndex, float InUV
 
 const FVector2D& UAvaShapeDynamicMeshBase::GetMaterialUVAnchor(int32 MeshIndex) const
 {
-	if (GetMaterialUVParams(MeshIndex) == nullptr)
-	{
-		return FVector2D::ZeroVector;
-	}
-
-	return GetMaterialUVParams(MeshIndex)->Anchor;
+	return GetMaterialUVParams(MeshIndex).Anchor;
 }
 
 bool UAvaShapeDynamicMeshBase::SetMaterialUVAnchor(int32 MeshIndex, const FVector2D& InUVAnchor)
 {
-	FAvaShapeMaterialUVParameters* Params = GetMaterialUVParams(MeshIndex);
+	FAvaShapeMaterialUVParameters* Params = GetMaterialUVParamsPtr(MeshIndex);
 
 	if (Params == nullptr)
 	{
@@ -957,17 +964,12 @@ bool UAvaShapeDynamicMeshBase::SetMaterialUVAnchor(int32 MeshIndex, const FVecto
 
 const FVector2D& UAvaShapeDynamicMeshBase::GetMaterialUVScale(int32 MeshIndex) const
 {
-	if (GetMaterialUVParams(MeshIndex) == nullptr)
-	{
-		return FVector2D::ZeroVector;
-	}
-
-	return GetMaterialUVParams(MeshIndex)->Scale;
+	return GetMaterialUVParams(MeshIndex).Scale;
 }
 
 bool UAvaShapeDynamicMeshBase::SetMaterialUVScale(int32 MeshIndex, const FVector2D& InUVScale)
 {
-	FAvaShapeMaterialUVParameters* Params = GetMaterialUVParams(MeshIndex);
+	FAvaShapeMaterialUVParameters* Params = GetMaterialUVParamsPtr(MeshIndex);
 
 	if (Params == nullptr)
 	{
@@ -992,17 +994,12 @@ bool UAvaShapeDynamicMeshBase::SetMaterialUVScale(int32 MeshIndex, const FVector
 
 const FVector2D& UAvaShapeDynamicMeshBase::GetMaterialUVOffset(int32 MeshIndex) const
 {
-	if (GetMaterialUVParams(MeshIndex) == nullptr)
-	{
-		return FVector2D::ZeroVector;
-	}
-
-	return GetMaterialUVParams(MeshIndex)->Offset;
+	return GetMaterialUVParams(MeshIndex).Offset;
 }
 
 bool UAvaShapeDynamicMeshBase::SetMaterialUVOffset(int32 MeshIndex, const FVector2D& InUVOffset)
 {
-	FAvaShapeMaterialUVParameters* Params = GetMaterialUVParams(MeshIndex);
+	FAvaShapeMaterialUVParameters* Params = GetMaterialUVParamsPtr(MeshIndex);
 
 	if (Params == nullptr)
 	{
@@ -1022,17 +1019,12 @@ bool UAvaShapeDynamicMeshBase::SetMaterialUVOffset(int32 MeshIndex, const FVecto
 
 bool UAvaShapeDynamicMeshBase::GetMaterialHorizontalFlip(int32 MeshIndex) const
 {
-	if (GetMaterialUVParams(MeshIndex) == nullptr)
-	{
-		return false;
-	}
-
-	return GetMaterialUVParams(MeshIndex)->bFlipHorizontal;
+	return GetMaterialUVParams(MeshIndex).bFlipHorizontal;
 }
 
 bool UAvaShapeDynamicMeshBase::SetMaterialHorizontalFlip(int32 MeshIndex, bool InHorizontalFlip)
 {
-	FAvaShapeMaterialUVParameters* Params = GetMaterialUVParams(MeshIndex);
+	FAvaShapeMaterialUVParameters* Params = GetMaterialUVParamsPtr(MeshIndex);
 
 	if (Params == nullptr)
 	{
@@ -1052,17 +1044,12 @@ bool UAvaShapeDynamicMeshBase::SetMaterialHorizontalFlip(int32 MeshIndex, bool I
 
 bool UAvaShapeDynamicMeshBase::GetMaterialVerticalFlip(int32 MeshIndex) const
 {
-	if (GetMaterialUVParams(MeshIndex) == nullptr)
-	{
-		return false;
-	}
-
-	return GetMaterialUVParams(MeshIndex)->bFlipVertical;
+	return GetMaterialUVParams(MeshIndex).bFlipVertical;
 }
 
 bool UAvaShapeDynamicMeshBase::SetMaterialVerticalFlip(int32 MeshIndex, bool InVerticalFlip)
 {
-	FAvaShapeMaterialUVParameters* Params = GetMaterialUVParams(MeshIndex);
+	FAvaShapeMaterialUVParameters* Params = GetMaterialUVParamsPtr(MeshIndex);
 
 	if (Params == nullptr)
 	{
@@ -1115,15 +1102,11 @@ bool UAvaShapeDynamicMeshBase::GenerateNormals()
 
 void UAvaShapeDynamicMeshBase::OnUVParamsChanged(int32 MeshIndex)
 {
-	if (!IsValidMeshIndex(MeshIndex))
+	if (FAvaShapeMaterialUVParameters* UVParams = GetMaterialUVParamsPtr(MeshIndex))
 	{
-		return;
+		UVParams->bUVsDirty = true;
+		RequestUpdate();
 	}
-
-	FAvaShapeMaterialUVParameters* UVParams = GetMaterialUVParams(MeshIndex);
-	UVParams->bUVsDirty = true;
-
-	RequireUpdate();
 }
 
 void UAvaShapeDynamicMeshBase::OnUsesPrimaryUVParamsChanged(int32 MeshIndex)
@@ -1154,11 +1137,6 @@ void UAvaShapeDynamicMeshBase::OnMeshChanged(int32 MeshIndex)
 	MarkAllMeshesDirty();
 }
 
-void UAvaShapeDynamicMeshBase::OnVertexColorChanged()
-{
-	InvalidateSection(bColorsDirty);
-}
-
 void UAvaShapeDynamicMeshBase::OnMaterialChanged(int32 MaterialIndex)
 {
 	if (!IsValidMeshIndex(MaterialIndex))
@@ -1166,27 +1144,27 @@ void UAvaShapeDynamicMeshBase::OnMaterialChanged(int32 MaterialIndex)
 		return;
 	}
 
-	UMaterialInterface* Material = GetMaterial(MaterialIndex);
+	FAvaShapeMeshData* MeshSection = GetMeshData(MaterialIndex);
+	UMaterialInterface* Material = MeshSection->Material;
 
 	// Switch back to parametric material when material is cleared.
 	if (Material == nullptr)
 	{
-		FAvaShapeParametricMaterial* ParametricMat = GetParametricMaterial(MaterialIndex);
-		Material = ParametricMat->GetOrCreateMaterial(GetShapeActor());
+		Material = MeshSection->ParametricMaterial.GetOrCreateMaterial(GetShapeActor());
 		SetMaterialDirect(MaterialIndex, Material);
 	}
 
-	if (GetParametricMaterial(MaterialIndex)->IsParametricMaterial(Material))
+	if (MeshSection->ParametricMaterial.IsParametricMaterial(Material))
 	{
-		GetMaterialType(MaterialIndex) = EMaterialType::Parametric;
+		MeshSection->MaterialType = EMaterialType::Parametric;
 	}
-	else if (Cast<UDynamicMaterialInstance>(GetMaterial(MaterialIndex)) != nullptr)
+	else if (Cast<UDynamicMaterialInstance>(Material) != nullptr)
 	{
-		GetMaterialType(MaterialIndex) = EMaterialType::MaterialDesigner;
+		MeshSection->MaterialType = EMaterialType::MaterialDesigner;
 	}
 	else
 	{
-		GetMaterialType(MaterialIndex) = EMaterialType::Asset;
+		MeshSection->MaterialType = EMaterialType::Asset;
 	}
 
 	SetShapeComponentMaterial(MaterialIndex, Material);
@@ -1201,7 +1179,7 @@ void UAvaShapeDynamicMeshBase::OnMaterialChanged(int32 MaterialIndex)
 			OnUsePrimaryMaterialEverywhereChanged();
 		}
 		// disable override if we change a material different than primary
-		else if (GetMaterial(MaterialIndex) != GetMaterial(MESH_INDEX_PRIMARY))
+		else if (Material != GetMaterial(MESH_INDEX_PRIMARY))
 		{
 			bUsePrimaryMaterialEverywhere = false;
 		}
@@ -1272,7 +1250,7 @@ void UAvaShapeDynamicMeshBase::OnParametricMaterialChanged(int32 MaterialIndex)
 		return;
 	}
 
-	FAvaShapeParametricMaterial* ParametricMaterial = GetParametricMaterial(MaterialIndex);
+	FAvaShapeParametricMaterial* ParametricMaterial = GetParametricMaterialPtr(MaterialIndex);
 
 	UMaterialInterface* ParametricMaterialInstance = ParametricMaterial->GetOrCreateMaterial(GetShapeActor());
 
@@ -1347,18 +1325,13 @@ bool UAvaShapeDynamicMeshBase::IsDynamicMeshInitialized() const
 	return bHasMeshData;
 }
 
-bool UAvaShapeDynamicMeshBase::CreateColors(FAvaShapeMesh& InMesh)
+void UAvaShapeDynamicMeshBase::UpdateVertices()
 {
-	for (const FVector& Vertex : InMesh.Vertices)
+	if (!bVerticesDirty)
 	{
-		InMesh.VertexColours.Add(VertexColor);
+		return;
 	}
 
-	return true;
-}
-
-void UAvaShapeDynamicMeshBase::UpdateVertices()
-{
 	bVerticesDirty = false;
 
 	const TSet<int32> Keys = GetMeshesIndexes();
@@ -1370,22 +1343,6 @@ void UAvaShapeDynamicMeshBase::UpdateVertices()
 	}
 }
 
-bool UAvaShapeDynamicMeshBase::UpdateColors()
-{
-	const TSet<int32> Keys = GetMeshesIndexes();
-
-	for (const int32 MeshIndex : Keys)
-	{
-		FAvaShapeMeshData* MeshData = GetMeshData(MeshIndex);
-		MeshData->Mesh.VertexColours.Empty();
-		MeshData->Mesh.bUpdateRequired = CreateColors(MeshData->Mesh);
-	}
-
-	bColorsDirty = false;
-
-	return true;
-}
-
 bool UAvaShapeDynamicMeshBase::OnMeshSectionUpdated(FAvaShapeMesh& InMesh)
 {
 	if (InMesh.Vertices.Num() == 0)
@@ -1605,6 +1562,9 @@ bool UAvaShapeDynamicMeshBase::CreateDynamicMesh(FAvaShapeMesh& InMesh)
 				NormalOverlay->SetTriangle(TId, NormalIds, true);
 			}
 		}
+		, EDynamicMeshChangeType::GeneralEdit
+		, EDynamicMeshAttributeChangeFlags::MeshTopology
+		, /** Broadcast */false
 	);
 
 	ProcMesh->SetMaterial(InMesh.GetMeshIndex(), GetMaterial(InMesh.GetMeshIndex()));
@@ -1658,9 +1618,10 @@ bool UAvaShapeDynamicMeshBase::UpdateDynamicMesh(FAvaShapeMesh& InMesh)
 				int32 NId = InMesh.NormalIds[Idx];
 				NormalOverlay->SetElement(NId, static_cast<FVector3f>(InMesh.Normals[Idx]));
 			});
-		},
-		EDynamicMeshChangeType::MeshVertexChange,
-		EDynamicMeshAttributeChangeFlags::Unknown
+		}
+		, EDynamicMeshChangeType::MeshVertexChange
+		, EDynamicMeshAttributeChangeFlags::Unknown
+		, /** Broadcast */false
 	);
 
 	ProcMesh->SetMaterial(InMesh.GetMeshIndex(), GetMaterial(InMesh.GetMeshIndex()));
@@ -1673,18 +1634,13 @@ void UAvaShapeDynamicMeshBase::PostEditChangeProperty(FPropertyChangedEvent& Pro
 {
 	Super::PostEditChangeProperty(PropertyChangedEvent);
 
-	static const FName VertexColorName = GET_MEMBER_NAME_CHECKED(UAvaShapeDynamicMeshBase, VertexColor);
 	static const FName ScaledSizeName = GET_MEMBER_NAME_CHECKED(UAvaShapeDynamicMeshBase, UniformScaledSize);
 	static const FName UsePrimaryMaterialEverywhereName = GET_MEMBER_NAME_CHECKED(UAvaShapeDynamicMeshBase, bUsePrimaryMaterialEverywhere);
 
 	const FName MemberPropertyName = PropertyChangedEvent.GetMemberPropertyName();
 	const FName PropertyName = PropertyChangedEvent.GetPropertyName();
 
-	if (MemberPropertyName.IsEqual(VertexColorName))
-	{
-		OnVertexColorChanged();
-	}
-	else if (MemberPropertyName.IsEqual(ScaledSizeName))
+	if (MemberPropertyName.IsEqual(ScaledSizeName))
 	{
 		OnScaledSizeChanged();
 	}
@@ -1768,17 +1724,18 @@ void UAvaShapeDynamicMeshBase::PostEditChangeChainProperty(FPropertyChangedChain
 			static const FName AnchorName = GET_MEMBER_NAME_CHECKED(FAvaShapeMaterialUVParameters, Anchor);
 			static const FName ScaleName = GET_MEMBER_NAME_CHECKED(FAvaShapeMaterialUVParameters, Scale);
 
-			FAvaShapeMaterialUVParameters* UVParams = GetMaterialUVParams(MeshIndex);
+			FAvaShapeMaterialUVParameters* UVParams = GetMaterialUVParamsPtr(MeshIndex);
+			const FName PropertyName = PropertyChangedEvent.GetPropertyName();
 
-			if (PropertyChangedEvent.Property->GetFName() == AnchorPresetName)
+			if (PropertyName == AnchorPresetName)
 			{
 				UVParams->Anchor = GetNumericsFromAnchor(UVParams->AnchorPreset);
 			}
-			else if (PropertyChangedEvent.Property->GetFName() == AnchorName)
+			else if (PropertyName == AnchorName)
 			{
 				UVParams->AnchorPreset = GetAnchorFromNumerics(UVParams->Anchor);
 			}
-			else if (PropertyChangedEvent.Property->GetFName() == ScaleName)
+			else if (PropertyName == ScaleName)
 			{
 				if (UVParams->Scale.X == 0)
 				{
@@ -1872,7 +1829,6 @@ bool UAvaShapeDynamicMeshBase::CheckMaterialSlotChanges()
 		{
 			bResult = true;
 			SetMaterial(SectionIdx, SlotMat);
-			GetMaterialType(SectionIdx) = EMaterialType::Asset;
 		}
 	}
 
@@ -1902,6 +1858,11 @@ void UAvaShapeDynamicMeshBase::SetShapeComponentMaterial(int32 InMaterialIndex,
 
 void UAvaShapeDynamicMeshBase::ScaleVertices(const FVector& InScale)
 {
+	if (InScale.IsNearlyZero() || InScale.Equals(FVector::OneVector))
+	{
+		return;
+	}
+
 	for (const int32 MeshIndex : GetMeshesIndexes())
 	{
 		FAvaShapeMeshData* MeshData = GetMeshData(MeshIndex);
@@ -1918,25 +1879,19 @@ void UAvaShapeDynamicMeshBase::ScaleVertices(const FVector& InScale)
 			MeshData->MaterialUVParams.bUVsDirty = true;
 		}
 	}
-}
 
-void UAvaShapeDynamicMeshBase::ScaleVertices(const FVector2D& InScale)
-{
-	for (const int32 MeshIndex : GetMeshesIndexes())
+	for (FAvaSnapPoint& SnapPoint : LocalSnapPoints)
 	{
-		FAvaShapeMeshData* MeshData = GetMeshData(MeshIndex);
+		SnapPoint.Location.Y *= InScale.Y;
+		SnapPoint.Location.Z *= InScale.Z;
+	}
 
-		for (FVector& Vertex : MeshData->Mesh.Vertices)
-		{
-			Vertex.Y *= InScale.X;
-			Vertex.Z *= InScale.Y;
-		}
+	MarkVerticesDirty();
+}
 
-		if (MeshData->MaterialUVParams.Mode != EAvaShapeUVMode::Stretch)
-		{
-			MeshData->MaterialUVParams.bUVsDirty = true;
-		}
-	}
+void UAvaShapeDynamicMeshBase::ScaleVertices(const FVector2D& InScale)
+{
+	ScaleVertices(FVector(1.f, InScale.X, InScale.Y));
 }
 
 void UAvaShapeDynamicMeshBase::SetMeshRegenWorldLocation(const FVector& NewLocation, bool bImmediateUpdate)
@@ -1994,9 +1949,10 @@ bool UAvaShapeDynamicMeshBase::ApplyUVsManually(FAvaShapeMesh& InMesh)
 				}
 				UVOverlay->SetTriangle(TID, UVElem);
 			}
-		},
-		EDynamicMeshChangeType::AttributeEdit,
-		EDynamicMeshAttributeChangeFlags::UVs
+		}
+		, EDynamicMeshChangeType::AttributeEdit
+		, EDynamicMeshAttributeChangeFlags::UVs
+		, /** Broadcast */false
 	);
 }
 
@@ -2031,9 +1987,10 @@ bool UAvaShapeDynamicMeshBase::ApplyUVsPlanarProjection(FAvaShapeMesh& InMesh, F
 			{
 				InMesh.UVs.Add(static_cast<FVector2D>(UVOverlay->GetElement(UId)));
 			}
-		}, 
-		EDynamicMeshChangeType::AttributeEdit,
-		EDynamicMeshAttributeChangeFlags::UVs
+		}
+		, EDynamicMeshChangeType::AttributeEdit
+		, EDynamicMeshAttributeChangeFlags::UVs
+		, /** Broadcast */false
 	);
 }
 
@@ -2066,13 +2023,14 @@ bool UAvaShapeDynamicMeshBase::ApplyUVsBoxProjection(FAvaShapeMesh& InMesh, FRot
 			{
 				InMesh.UVs.Add(static_cast<FVector2D>(UVOverlay->GetElement(UId)));
 			}
-		},
-		EDynamicMeshChangeType::AttributeEdit,
-		EDynamicMeshAttributeChangeFlags::UVs
+		}
+		, EDynamicMeshChangeType::AttributeEdit
+		, EDynamicMeshAttributeChangeFlags::UVs
+		, /** Broadcast */false
 	);
 }
 
-bool UAvaShapeDynamicMeshBase::ApplyUVsTransform(FAvaShapeMesh& InMesh, FAvaShapeMaterialUVParameters& InParams, FVector2D ShapeSize, FVector2D UVOffset, float UVFixRotation) const
+bool UAvaShapeDynamicMeshBase::ApplyUVsTransform(FAvaShapeMesh& InMesh, FAvaShapeMaterialUVParameters& InParams, FVector2D ShapeSize, FVector2D UVOffset, float UVFixRotation)
 {
 	bool bResult = true;
 	return UE::Private::AvaShape::EditMesh(
@@ -2080,17 +2038,23 @@ bool UAvaShapeDynamicMeshBase::ApplyUVsTransform(FAvaShapeMesh& InMesh, FAvaShap
 		[InMesh, InParams, ShapeSize, UVOffset, UVFixRotation, &bResult](FDynamicMesh3& EditMesh)
 		{
 			bResult = UE::AvaShapes::TransformMeshUVs(EditMesh, InMesh.UVIds, InParams, ShapeSize, UVOffset, UVFixRotation);
-		},
-		EDynamicMeshChangeType::AttributeEdit,
-		EDynamicMeshAttributeChangeFlags::UVs
+		}
+		, EDynamicMeshChangeType::AttributeEdit
+		, EDynamicMeshAttributeChangeFlags::UVs
+		, /** Broadcast */false
 	) && bResult;
 }
 
+void UAvaShapeDynamicMeshBase::SetRunAsync(bool bInAsync)
+{
+	bRunAsync = bInAsync;
+}
+
 void UAvaShapeDynamicMeshBase::OnParametricMaterialChanged(const FAvaShapeParametricMaterial& InMaterial)
 {
 	for (const int32 Index : GetMeshesIndexes())
 	{
-		if (GetParametricMaterial(Index) == &InMaterial && GetMaterialType(Index) == EMaterialType::Parametric)
+		if (GetParametricMaterialPtr(Index) == &InMaterial && GetMaterialType(Index) == EMaterialType::Parametric)
 		{
 			OnParametricMaterialChanged(Index);
 			break;
@@ -2102,7 +2066,7 @@ void UAvaShapeDynamicMeshBase::OnParametricMaterialParameterChanged(const FAvaSh
 {
 	for (const int32 Index : GetMeshesIndexes())
 	{
-		if (GetParametricMaterial(Index) == &InMaterial && GetMaterialType(Index) == EMaterialType::Parametric)
+		if (GetParametricMaterialPtr(Index) == &InMaterial && GetMaterialType(Index) == EMaterialType::Parametric)
 		{
 			OnParametricMaterialChanged(Index);
 			break;
@@ -2110,17 +2074,27 @@ void UAvaShapeDynamicMeshBase::OnParametricMaterialParameterChanged(const FAvaSh
 	}
 }
 
-void UAvaShapeDynamicMeshBase::RestoreCachedMesh()
+void UAvaShapeDynamicMeshBase::RestoreCachedMesh(bool bInReset)
 {
 	if (CachedMesh.IsSet())
 	{
 		UE::Private::AvaShape::EditMesh(
 			this,
-			[this](FDynamicMesh3& EditMesh)
+			[this, bInReset](FDynamicMesh3& EditMesh)
 			{
-				EditMesh = CachedMesh.GetValue();
+				if (bInReset)
+				{
+					EditMesh = MoveTemp(CachedMesh.GetValue());
+					CachedMesh.Reset();
+				}
+				else
+				{
+					EditMesh = CachedMesh.GetValue();
+				}
 			}
-		);
+		, EDynamicMeshChangeType::GeneralEdit
+		,EDynamicMeshAttributeChangeFlags::Unknown
+		, /** Broadcast */false);
 	}
 }
 
@@ -2140,13 +2114,13 @@ bool UAvaShapeDynamicMeshBase::CreateUVs(FAvaShapeMesh& InMesh, FAvaShapeMateria
 	return ApplyUVsTransform(InMesh, InParams, FVector2D(1, 1), FVector2D(0, 0));
 }
 
-void UAvaShapeDynamicMeshBase::InvalidateSection(bool& bInvalidatedSection, bool bRequireUpdate, bool bDoAsync)
+void UAvaShapeDynamicMeshBase::InvalidateSection(bool& bInvalidatedSection, bool bRequireUpdate)
 {
 	bInvalidatedSection = true;
 
 	if (bRequireUpdate)
 	{
-		RequireUpdate(bDoAsync);
+		RequestUpdate();
 	}
 }
 
@@ -2244,7 +2218,7 @@ void UAvaShapeDynamicMeshBase::MarkVerticesDirty()
 	InvalidateSection(bVerticesDirty);
 }
 
-FName UAvaShapeDynamicMeshBase::GetMeshDataName(int32 InMeshIndex) const
+FName UAvaShapeDynamicMeshBase::GetMeshSectionName(int32 InMeshIndex) const
 {
 	if (const FAvaShapeMeshData* MeshData = GetMeshData(InMeshIndex))
 	{
@@ -2254,7 +2228,7 @@ FName UAvaShapeDynamicMeshBase::GetMeshDataName(int32 InMeshIndex) const
 	return FName();
 }
 
-TArray<FName> UAvaShapeDynamicMeshBase::GetMeshDataNames() const
+TArray<FName> UAvaShapeDynamicMeshBase::GetMeshSectionNames() const
 {
 	TArray<FName> Names;
 	const TSet<int32> Keys = GetMeshesIndexes();
@@ -2308,7 +2282,7 @@ void UAvaShapeDynamicMeshBase::InitializeDynamicMesh()
 			EditMesh.Attributes()->SetNumPolygroupLayers(MeshDatas.Num());
 			for (int32 Idx = 0; Idx < MeshDatas.Num(); Idx++)
 			{
-				const FName DataName = GetMeshDataName(Idx);
+				const FName DataName = GetMeshSectionName(Idx);
 				EditMesh.Attributes()->GetPolygroupLayer(Idx)->SetName(DataName);
 			}
 			EditMesh.Attributes()->SetNumUVLayers(1);
@@ -2316,6 +2290,9 @@ void UAvaShapeDynamicMeshBase::InitializeDynamicMesh()
 
 			CachedMesh = EditMesh;
 		}
+		, EDynamicMeshChangeType::GeneralEdit
+		, EDynamicMeshAttributeChangeFlags::Unknown
+		, /** Broadcast */false
 	);
 
 	bZeroScaled = MeshComponent->GetComponentScale().IsNearlyZero();
@@ -2421,6 +2398,9 @@ bool UAvaShapeDynamicMeshBase::ClearDynamicMeshSection(int32 MeshIndex)
 			// empty ids arrays
 			MeshData->Mesh.ClearIds();
 		}
+		, EDynamicMeshChangeType::GeneralEdit
+		, EDynamicMeshAttributeChangeFlags::MeshTopology
+		, /** Broadcast */false
 	);
 }
 
@@ -2438,6 +2418,9 @@ bool UAvaShapeDynamicMeshBase::ClearDynamicMesh()
 				}
 			}
 		}
+		, EDynamicMeshChangeType::GeneralEdit
+		, EDynamicMeshAttributeChangeFlags::MeshTopology
+		, /** Broadcast */false
 	);
 
 	if (!bResult)
@@ -2471,7 +2454,7 @@ bool UAvaShapeDynamicMeshBase::ExportToStaticMesh(UStaticMesh* DestinationMesh)
 	AssetOptions.bEnableRecomputeNormals = false;
 	AssetOptions.bEnableRecomputeTangents = false;
 	AssetOptions.bEnableRemoveDegenerates = true;
-	AssetOptions.NewMaterialSlotNames = GetMeshDataNames();
+	AssetOptions.NewMaterialSlotNames = GetMeshSectionNames();
 
 	// Only take asset materials instead of dynamic ones
 	UMaterial* DefaultMaterial = LoadObject<UMaterial>(nullptr, TEXT("/Script/Engine.Material'/Engine/EngineMaterials/WorldGridMaterial.WorldGridMaterial'"));
@@ -2507,7 +2490,7 @@ void UAvaShapeDynamicMeshBase::OnColorPicked(const FAvaColorChangeData& InNewCol
 		return;
 	}
 
-	FAvaShapeParametricMaterial* const ParametricMaterial = GetParametricMaterial(MeshIndex);
+	FAvaShapeParametricMaterial* const ParametricMaterial = GetParametricMaterialPtr(MeshIndex);
 	if (!ParametricMaterial)
 	{
 		return;
@@ -2578,14 +2561,7 @@ FAvaColorChangeData UAvaShapeDynamicMeshBase::GetActiveColor() const
 		return ActiveColor;
 	}
 
-	const FAvaShapeParametricMaterial* ParamMat = GetParametricMaterial(0);
-
-	if (!ParamMat)
-	{
-		return ActiveColor;
-	}
-
-	switch (ParamMat->GetStyle())
+	switch (MeshData->ParametricMaterial.GetStyle())
 	{
 		case EAvaShapeParametricMaterialStyle::Solid:
 			ActiveColor.ColorStyle = EAvaColorStyle::Solid;
@@ -2600,9 +2576,9 @@ FAvaColorChangeData UAvaShapeDynamicMeshBase::GetActiveColor() const
 			break;
 	}
 
-	ActiveColor.PrimaryColor = ParamMat->GetPrimaryColor();
-	ActiveColor.SecondaryColor = ParamMat->GetSecondaryColor();
-	ActiveColor.bIsUnlit = ParamMat->GetUseUnlitMaterial();
+	ActiveColor.PrimaryColor = MeshData->ParametricMaterial.GetPrimaryColor();
+	ActiveColor.SecondaryColor = MeshData->ParametricMaterial.GetSecondaryColor();
+	ActiveColor.bIsUnlit = MeshData->ParametricMaterial.GetUseUnlitMaterial();
 
 	return ActiveColor;
 }
diff --git a/Engine/Plugins/Experimental/Avalanche/Source/AvalancheShapes/Public/AvaShapeMesh.h b/Engine/Plugins/Experimental/Avalanche/Source/AvalancheShapes/Public/AvaShapeMesh.h
--- a/Engine/Plugins/Experimental/Avalanche/Source/AvalancheShapes/Public/AvaShapeMesh.h
+++ b/Engine/Plugins/Experimental/Avalanche/Source/AvalancheShapes/Public/AvaShapeMesh.h
@@ -118,7 +118,7 @@ struct FAvaShapeMeshData
 
 protected:
 	// whether the mesh is currently visible or not and should be editable
-	UPROPERTY(EditDefaultsOnly, BlueprintReadOnly, Category="Shape")
+	UPROPERTY(EditDefaultsOnly, Category="Shape")
 	bool bMeshVisible;
 
 	// mesh uses same material as primary
@@ -126,23 +126,23 @@ protected:
 	bool bUsesPrimaryMaterialParams;
 
 	// mesh material type
-	UPROPERTY(EditInstanceOnly, BlueprintReadWrite, Category="Material", meta=(DisplayName="Material Type", DisplayAfter="bUsePrimaryMaterialEverywhere", EditCondition="bMeshVisible && !bUsesPrimaryMaterialParams", EditConditionHides, AllowPrivateAccess = "true"))
+	UPROPERTY(EditInstanceOnly, Category="Material", meta=(NoResetToDefault, DisplayName="Material Type", DisplayAfter="bUsePrimaryMaterialEverywhere", EditCondition="bMeshVisible && !bUsesPrimaryMaterialParams", EditConditionHides, AllowPrivateAccess = "true"))
 	EMaterialType MaterialType;
 
 	// mesh material
-	UPROPERTY(EditInstanceOnly, BlueprintReadWrite, Category="Material", meta=(DisplayName="Material", DisplayAfter="MaterialType", EditCondition="bMeshVisible && !bUsesPrimaryMaterialParams && MaterialType != EMaterialType::Parametric", EditConditionHides, AllowPrivateAccess = "true"))
+	UPROPERTY(EditInstanceOnly, Category="Material", meta=(NoResetToDefault, DisplayName="Material", DisplayAfter="MaterialType", EditCondition="bMeshVisible && !bUsesPrimaryMaterialParams && MaterialType != EMaterialType::Parametric", EditConditionHides, AllowPrivateAccess = "true"))
 	TObjectPtr<UMaterialInterface> Material;
 
 	// parametric material settings
-	UPROPERTY(EditInstanceOnly, BlueprintReadWrite, Category="Material", meta=(DisplayName="Material", DisplayAfter="Material", EditCondition="bMeshVisible && !bUsesPrimaryMaterialParams && MaterialType == EMaterialType::Parametric", EditConditionHides, AllowPrivateAccess = "true"))
+	UPROPERTY(EditInstanceOnly, Category="Material", meta=(NoResetToDefault, DisplayName="Material", DisplayAfter="Material", EditCondition="bMeshVisible && !bUsesPrimaryMaterialParams && MaterialType == EMaterialType::Parametric", EditConditionHides, AllowPrivateAccess = "true"))
 	FAvaShapeParametricMaterial ParametricMaterial;
 
 	// mesh uses same uv params as primary
-	UPROPERTY(EditInstanceOnly, BlueprintReadWrite, Category="Material", meta=(DisplayAfter="ParametricMaterial", EditCondition="bMeshVisible && MeshIndex != 0", EditConditionHides, AllowPrivateAccess = "true"))
+	UPROPERTY(EditInstanceOnly, Category="Material", meta=(NoResetToDefault, DisplayAfter="ParametricMaterial", EditCondition="bMeshVisible && MeshIndex != 0", EditConditionHides, AllowPrivateAccess = "true"))
 	bool bOverridePrimaryUVParams;
 
 	// mesh material UV params
-	UPROPERTY(EditInstanceOnly, BlueprintReadWrite, Category="Material", meta=(DisplayName="UV Params", DisplayAfter="bOverridePrimaryUVParams", EditCondition="bMeshVisible && bOverridePrimaryUVParams", EditConditionHides, AllowPrivateAccess = "true"))
+	UPROPERTY(EditInstanceOnly, Category="Material", meta=(NoResetToDefault, DisplayName="UV Params", DisplayAfter="bOverridePrimaryUVParams", EditCondition="bMeshVisible && bOverridePrimaryUVParams", EditConditionHides, AllowPrivateAccess = "true"))
 	FAvaShapeMaterialUVParameters MaterialUVParams;
 
 	// do not change this once assigned, 0 means primary
diff --git a/Engine/Plugins/Experimental/Avalanche/Source/AvalancheShapes/Public/AvaShapeParametricMaterial.h b/Engine/Plugins/Experimental/Avalanche/Source/AvalancheShapes/Public/AvaShapeParametricMaterial.h
--- a/Engine/Plugins/Experimental/Avalanche/Source/AvalancheShapes/Public/AvaShapeParametricMaterial.h
+++ b/Engine/Plugins/Experimental/Avalanche/Source/AvalancheShapes/Public/AvaShapeParametricMaterial.h
@@ -155,39 +155,39 @@ struct FAvaShapeParametricMaterial
 
 protected:
 	/** Default style for the material */
-	UPROPERTY(EditInstanceOnly, BlueprintReadWrite, Category="Material", meta=(DisplayName="Material Style", AllowPrivateAccess = "true"))
+	UPROPERTY(EditInstanceOnly, BlueprintReadWrite, Category="Material", meta=(NoResetToDefault, DisplayName="Material Style", AllowPrivateAccess = "true"))
 	EAvaShapeParametricMaterialStyle Style = EAvaShapeParametricMaterialStyle::Solid;
 
 	/** Simple texture for the material */
-	UPROPERTY(EditInstanceOnly, BlueprintReadWrite, Category="Material", meta=(DisplayName="Texture", EditCondition="Style == EAvaShapeParametricMaterialStyle::Texture", EditConditionHides, AllowPrivateAccess = "true"))
+	UPROPERTY(EditInstanceOnly, BlueprintReadWrite, Category="Material", meta=(NoResetToDefault, DisplayName="Texture", EditCondition="Style == EAvaShapeParametricMaterialStyle::Texture", EditConditionHides, AllowPrivateAccess = "true"))
 	TObjectPtr<UTexture> Texture = nullptr;
 
 	/** Primary colour for the material */
-	UPROPERTY(EditInstanceOnly, BlueprintReadWrite, Category="Material", meta=(DisplayName="Primary Color", EditCondition="Style != EAvaShapeParametricMaterialStyle::Texture", EditConditionHides, AllowPrivateAccess = "true"))
+	UPROPERTY(EditInstanceOnly, BlueprintReadWrite, Category="Material", meta=(NoResetToDefault, DisplayName="Primary Color", EditCondition="Style != EAvaShapeParametricMaterialStyle::Texture", EditConditionHides, AllowPrivateAccess = "true"))
 	FLinearColor ColorA = FLinearColor::White;
 
 	/** Secondary colour for the material */
-	UPROPERTY(EditInstanceOnly, BlueprintReadWrite, Category="Material", meta=(DisplayName="Secondary Color", EditCondition="Style == EAvaShapeParametricMaterialStyle::LinearGradient", EditConditionHides, AllowPrivateAccess = "true"))
+	UPROPERTY(EditInstanceOnly, BlueprintReadWrite, Category="Material", meta=(NoResetToDefault, DisplayName="Secondary Color", EditCondition="Style == EAvaShapeParametricMaterialStyle::LinearGradient", EditConditionHides, AllowPrivateAccess = "true"))
 	FLinearColor ColorB = FLinearColor::Black;
 
 	/** Offset for gradient style material */
-	UPROPERTY(EditInstanceOnly, BlueprintReadWrite, Category="Material", meta=(ClampMin="0", ClampMax="1", DisplayName="Gradient Offset", EditCondition = "Style == EAvaShapeParametricMaterialStyle::LinearGradient", EditConditionHides, AllowPrivateAccess = "true"))
+	UPROPERTY(EditInstanceOnly, BlueprintReadWrite, Category="Material", meta=(NoResetToDefault, ClampMin="0", ClampMax="1", DisplayName="Gradient Offset", EditCondition = "Style == EAvaShapeParametricMaterialStyle::LinearGradient", EditConditionHides, AllowPrivateAccess = "true"))
 	float GradientOffset = 0.5f;
 
 	/** Rotation for gradient style material */
-	UPROPERTY(EditInstanceOnly, BlueprintReadWrite, Category="Material", meta=(ClampMin="0", ClampMax="1", DisplayName="Gradient Rotation", EditCondition = "Style == EAvaShapeParametricMaterialStyle::LinearGradient", EditConditionHides, AllowPrivateAccess = "true"))
+	UPROPERTY(EditInstanceOnly, BlueprintReadWrite, Category="Material", meta=(NoResetToDefault, ClampMin="0", ClampMax="1", DisplayName="Gradient Rotation", EditCondition = "Style == EAvaShapeParametricMaterialStyle::LinearGradient", EditConditionHides, AllowPrivateAccess = "true"))
 	float GradientRotation = 0.f;
 
 	/** whether the material is unlit or default lit */
-	UPROPERTY(EditInstanceOnly, BlueprintReadWrite, Category="Material", meta=(DisplayName="Use Unlit Material", AllowPrivateAccess = "true"))
+	UPROPERTY(EditInstanceOnly, BlueprintReadWrite, Category="Material", meta=(NoResetToDefault, DisplayName="Use Unlit Material", AllowPrivateAccess = "true"))
 	bool bUseUnlitMaterial = false;
 
 	/** whether the material is one sided or two sided */
-	UPROPERTY(EditInstanceOnly, BlueprintReadWrite, Category = "Material", meta = (DisplayName = "Use Two Sided Material", AllowPrivateAccess = "true"))
-	bool bUseTwoSidedMaterial = true;
+	UPROPERTY(EditInstanceOnly, BlueprintReadWrite, Category = "Material", meta=(NoResetToDefault, DisplayName = "Use Two Sided Material", AllowPrivateAccess = "true"))
+	bool bUseTwoSidedMaterial = false;
 
 	/** How to handle translucency for the underlying material */
-	UPROPERTY(EditInstanceOnly, BlueprintReadWrite, Category="Material")
+	UPROPERTY(EditInstanceOnly, BlueprintReadWrite, Category="Material", meta=(NoResetToDefault))
 	EAvaShapeParametricMaterialTranslucency Translucency = EAvaShapeParametricMaterialTranslucency::Auto;
 
 private:
diff --git a/Engine/Plugins/Experimental/Avalanche/Source/AvalancheShapes/Public/AvaShapeUVParameters.h b/Engine/Plugins/Experimental/Avalanche/Source/AvalancheShapes/Public/AvaShapeUVParameters.h
--- a/Engine/Plugins/Experimental/Avalanche/Source/AvalancheShapes/Public/AvaShapeUVParameters.h
+++ b/Engine/Plugins/Experimental/Avalanche/Source/AvalancheShapes/Public/AvaShapeUVParameters.h
@@ -67,28 +67,28 @@ struct FAvaShapeMaterialUVParameters
 	}
 
 protected:
-	UPROPERTY(EditInstanceOnly, BlueprintReadWrite, Category="Shape", meta=(AllowPrivateAccess = "true"))
+	UPROPERTY(EditInstanceOnly, BlueprintReadWrite, Category="Shape", meta=(NoResetToDefault, AllowPrivateAccess = "true"))
 	EAvaShapeUVMode Mode;
 
-	UPROPERTY(EditInstanceOnly, BlueprintReadWrite, Category="Shape", meta=(AllowPrivateAccess = "true"))
+	UPROPERTY(EditInstanceOnly, BlueprintReadWrite, Category="Shape", meta=(NoResetToDefault, AllowPrivateAccess = "true"))
 	EAvaAnchors AnchorPreset;
 
-	UPROPERTY(EditInstanceOnly, BlueprintReadWrite, Category="Shape", meta=(ClampMin="0.0", ClampMax="1.0", EditCondition="AnchorPreset == EAvaAnchors::Custom", AllowPrivateAccess = "true"))
+	UPROPERTY(EditInstanceOnly, BlueprintReadWrite, Category="Shape", meta=(NoResetToDefault, ClampMin="0.0", ClampMax="1.0", EditCondition="AnchorPreset == EAvaAnchors::Custom", AllowPrivateAccess = "true"))
 	FVector2D Anchor;
 
-	UPROPERTY(EditInstanceOnly, BlueprintReadWrite, Category="Shape", meta=(ClampMin="0.01", MotionDesignVectorWidget, AllowPreserveRatio="XY", Delta="0.05", AllowPrivateAccess = "true"))
+	UPROPERTY(EditInstanceOnly, BlueprintReadWrite, Category="Shape", meta=(NoResetToDefault, ClampMin="0.01", MotionDesignVectorWidget, AllowPreserveRatio="XY", Delta="0.05", AllowPrivateAccess = "true"))
 	FVector2D Scale;
 
-	UPROPERTY(EditInstanceOnly, BlueprintReadWrite, Category="Shape", meta=(Delta="0.05", AllowPrivateAccess = "true"))
+	UPROPERTY(EditInstanceOnly, BlueprintReadWrite, Category="Shape", meta=(NoResetToDefault, Delta="0.05", AllowPrivateAccess = "true"))
 	FVector2D Offset;
 
-	UPROPERTY(EditInstanceOnly, BlueprintReadWrite, Category="Shape", meta=(ClampMin="-360.0", ClampMax="360.0", AllowPrivateAccess = "true"))
+	UPROPERTY(EditInstanceOnly, BlueprintReadWrite, Category="Shape", meta=(NoResetToDefault, ClampMin="-360.0", ClampMax="360.0", AllowPrivateAccess = "true"))
 	float Rotation;
 
-	UPROPERTY(EditInstanceOnly, BlueprintReadWrite, Category="Shape", meta=(AllowPrivateAccess = "true"))
+	UPROPERTY(EditInstanceOnly, BlueprintReadWrite, Category="Shape", meta=(NoResetToDefault, AllowPrivateAccess = "true"))
 	bool bFlipHorizontal;
 
-	UPROPERTY(EditInstanceOnly, BlueprintReadWrite, Category="Shape", meta=(AllowPrivateAccess = "true"))
+	UPROPERTY(EditInstanceOnly, BlueprintReadWrite, Category="Shape", meta=(NoResetToDefault, AllowPrivateAccess = "true"))
 	bool bFlipVertical;
 
 	bool bUVsDirty;
diff --git a/Engine/Plugins/Experimental/Avalanche/Source/AvalancheShapes/Public/AvaShapesDefs.h b/Engine/Plugins/Experimental/Avalanche/Source/AvalancheShapes/Public/AvaShapesDefs.h
--- a/Engine/Plugins/Experimental/Avalanche/Source/AvalancheShapes/Public/AvaShapesDefs.h
+++ b/Engine/Plugins/Experimental/Avalanche/Source/AvalancheShapes/Public/AvaShapesDefs.h
@@ -35,8 +35,10 @@ enum class EAvaShapeRectangleCornerTypeIndex : uint8
 UENUM()
 enum class EAvaShapeUVMode : uint8
 {
-	Stretch, // uvs are stretched and fill the shape size
-	Uniform // uvs are uniform and fit the shape size (square)
+	// uvs are stretched to fit the shape size
+	Stretch		UMETA(DisplayName="Stretch to Fit"),
+	// uvs are uniform to not cause any distortion
+	Uniform		UMETA(DisplayName="Preserve Ratio"),
 };
 
 UENUM()
diff --git a/Engine/Plugins/Experimental/Avalanche/Source/AvalancheShapes/Public/DynamicMeshes/AvaShape2DDynMeshBase.h b/Engine/Plugins/Experimental/Avalanche/Source/AvalancheShapes/Public/DynamicMeshes/AvaShape2DDynMeshBase.h
--- a/Engine/Plugins/Experimental/Avalanche/Source/AvalancheShapes/Public/DynamicMeshes/AvaShape2DDynMeshBase.h
+++ b/Engine/Plugins/Experimental/Avalanche/Source/AvalancheShapes/Public/DynamicMeshes/AvaShape2DDynMeshBase.h
@@ -38,14 +38,14 @@ public:
 		return PixelSize2D;
 	}
 
-	AVALANCHESHAPES_API void SetSize2D(const FVector2D& InSize2D);
-	const FVector2D& GetSize2D() const
+	AVALANCHESHAPES_API virtual void SetSize2D(const FVector2D& InSize2D) override;
+	virtual FVector2D GetSize2D() const override
 	{
 		return Size2D;
 	}
 
 	AVALANCHESHAPES_API virtual void SetSize3D(const FVector& InSize) override;
-	virtual const FVector& GetSize3D() const override
+	virtual FVector GetSize3D() const override
 	{
 		return Size3D;
 	}
diff --git a/Engine/Plugins/Experimental/Avalanche/Source/AvalancheShapes/Public/DynamicMeshes/AvaShape3DDynMeshBase.h b/Engine/Plugins/Experimental/Avalanche/Source/AvalancheShapes/Public/DynamicMeshes/AvaShape3DDynMeshBase.h
--- a/Engine/Plugins/Experimental/Avalanche/Source/AvalancheShapes/Public/DynamicMeshes/AvaShape3DDynMeshBase.h
+++ b/Engine/Plugins/Experimental/Avalanche/Source/AvalancheShapes/Public/DynamicMeshes/AvaShape3DDynMeshBase.h
@@ -32,7 +32,7 @@ public:
 	}
 
 	AVALANCHESHAPES_API virtual void SetSize3D(const FVector& InSize) override;
-	virtual const FVector& GetSize3D() const override
+	virtual FVector GetSize3D() const override
 	{
 		return Size3D;
 	}
diff --git a/Engine/Plugins/Experimental/Avalanche/Source/AvalancheShapes/Public/DynamicMeshes/AvaShapeDynMeshBase.h b/Engine/Plugins/Experimental/Avalanche/Source/AvalancheShapes/Public/DynamicMeshes/AvaShapeDynMeshBase.h
--- a/Engine/Plugins/Experimental/Avalanche/Source/AvalancheShapes/Public/DynamicMeshes/AvaShapeDynMeshBase.h
+++ b/Engine/Plugins/Experimental/Avalanche/Source/AvalancheShapes/Public/DynamicMeshes/AvaShapeDynMeshBase.h
@@ -68,51 +68,102 @@ public:
 	/** Get the name of the shape */
 	virtual const FString& GetMeshName() const;
 
+	/** Returns the name of the sections composing this mesh */
+	UFUNCTION(BlueprintCallable, Category="Shape")
+	AVALANCHESHAPES_API TArray<FName> GetMeshSectionNames() const;
+
 	/** Can we change the size of this shape */
-	bool AllowsSizeEditing() const
+	UFUNCTION(BlueprintPure, Category="Shape")
+	bool GetAllowEditSize() const
 	{
 		return bAllowEditSize;
 	}
 
-	// override in child classes, represents the mesh size
+	UFUNCTION(BlueprintCallable, Category="Shape")
 	virtual void SetSize3D(const FVector& InSize) {}
-	virtual const FVector& GetSize3D() const
+
+	UFUNCTION(BlueprintPure, Category="Shape")
+	virtual FVector GetSize3D() const
 	{
 		return FVector::ZeroVector;
 	}
 
-	// Primary material override
+	UFUNCTION(BlueprintCallable, Category="Shape")
+	AVALANCHESHAPES_API virtual void SetSize2D(const FVector2D& InSize2D);
+
+	UFUNCTION(BlueprintPure, Category="Shape")
+	AVALANCHESHAPES_API virtual FVector2D GetSize2D() const;
+
+	/** Checks if a mesh section is currently visible */
+	UFUNCTION(BlueprintPure, Category="Shape")
+	bool IsMeshSectionVisible(int32 InIndex) const;
+
+	/** Set Primary material override */
+	UFUNCTION(BlueprintCallable, Category="Shape")
 	AVALANCHESHAPES_API void SetUsePrimaryMaterialEverywhere(bool bInUse);
+
+	UFUNCTION(BlueprintPure, Category="Shape")
 	bool GetUsePrimaryMaterialEverywhere() const
 	{
 		return bUsePrimaryMaterialEverywhere;
 	}
 
+	UFUNCTION(BlueprintCallable, Category="Shape")
 	AVALANCHESHAPES_API void SetUniformScaledSize(float InSize);
+
+	UFUNCTION(BlueprintPure, Category="Shape")
 	float GetUniformScaledSize() const
 	{
 		return UniformScaledSize;
 	}
 
-	/** Checks if shape size is render-able and not minimal */
-	bool IsMeshSizeValid() const;
-
 	/** Checks if the mesh section is valid */
-	AVALANCHESHAPES_API bool IsValidMeshIndex(int32 MeshIndex) const;
+	UFUNCTION(BlueprintPure, Category="Shape")
+	AVALANCHESHAPES_API bool IsValidMeshIndex(int32 InMeshIndex) const;
+
+	UFUNCTION(BlueprintCallable, Category="Shape")
+	void SetMaterialType(int32 InMeshIndex, EMaterialType InType);
+
+	UFUNCTION(BlueprintPure, Category="Shape")
+	EMaterialType GetMaterialType(int32 InMeshIndex);
+
+	AVALANCHESHAPES_API bool IsMaterialType(int32 MeshIndex, EMaterialType Type);
+
+	/** Get the material of a mesh section */
+	UFUNCTION(BlueprintPure, Category="Shape")
+	AVALANCHESHAPES_API UMaterialInterface* GetMaterial(int32 InMeshIndex) const;
+
+	/** Set the material of a mesh section */
+	UFUNCTION(BlueprintCallable, Category="Shape")
+	AVALANCHESHAPES_API void SetMaterial(int32 InMeshIndex, UMaterialInterface* InNewMaterial);
+
+	/** Set parametric material of a section, internal materials are not copied over, only settings */
+	UFUNCTION(BlueprintCallable, Category="Shape")
+	AVALANCHESHAPES_API void SetParametricMaterial(int32 InMeshIndex, const FAvaShapeParametricMaterial& InNewMaterialParams);
+
+	/** Get parametric material of a mesh section */
+	UFUNCTION(BlueprintPure, Category="Shape")
+	const FAvaShapeParametricMaterial& GetParametricMaterial(int32 InMeshIndex) const;
+
+	AVALANCHESHAPES_API FAvaShapeParametricMaterial* GetParametricMaterialPtr(int32 InMeshIndex);
 
-	AVALANCHESHAPES_API UMaterialInterface* GetMaterial(int32 MeshIndex) const;
-	AVALANCHESHAPES_API bool SetMaterial(int32 MeshIndex, UMaterialInterface* NewMaterial);
+	/** Use custom uv params per mesh section instead of using the primary one */
+	UFUNCTION(BlueprintCallable, Category="Shape")
+	void SetOverridePrimaryUVParams(int32 InMeshIndex, bool bInOverride);
 
-	AVALANCHESHAPES_API const FAvaShapeParametricMaterial* GetParametricMaterial(int32 MeshIndex) const;
-	AVALANCHESHAPES_API FAvaShapeParametricMaterial* GetParametricMaterial(int32 MeshIndex);
+	UFUNCTION(BlueprintPure, Category="Shape")
+	bool GetOverridePrimaryUVParams(int32 InMeshIndex) const;
 
-	// Internal materials are not copied over, only settings
-	AVALANCHESHAPES_API bool SetParametricMaterial(int32 MeshIndex, const FAvaShapeParametricMaterial& NewMaterialParams);
+	UFUNCTION(BlueprintCallable, Category="Shape")
+	void SetMaterialUVParams(int32 InMeshIndex, const FAvaShapeMaterialUVParameters& InParams);
 
-	// Allows custom uv params per mesh section instead of using the primary one
-	bool SetOverridePrimaryUVParams(int32 MeshIndex, bool bOverride);
+	UFUNCTION(BlueprintPure, Category="Shape")
+	const FAvaShapeMaterialUVParameters& GetMaterialUVParams(int32 InMeshIndex) const;
+
+	FAvaShapeMaterialUVParameters* GetMaterialUVParamsPtr(int32 MeshIndex);
 
 	// Helper function to quickly get what you need instead of using GetMeshData()
+
 	AVALANCHESHAPES_API EAvaShapeUVMode GetMaterialUVMode(int32 MeshIndex) const;
 	AVALANCHESHAPES_API bool SetMaterialUVMode(int32 MeshIndex, EAvaShapeUVMode InUVMode);
 
@@ -140,12 +191,6 @@ public:
 	AVALANCHESHAPES_API const FAvaShapeMaterialUVParameters* GetInUseMaterialUVParams(int32 MeshIndex) const;
 	AVALANCHESHAPES_API FAvaShapeMaterialUVParameters* GetInUseMaterialUVParams(int32 MeshIndex);
 
-	const FAvaShapeMaterialUVParameters* GetMaterialUVParams(int32 MeshIndex) const;
-	bool SetMaterialUVParams(int32 MeshIndex, const FAvaShapeMaterialUVParameters& InParams);
-
-	bool SetMaterialType(int32 MeshIndex, EMaterialType Type);
-	AVALANCHESHAPES_API bool IsMaterialType(int32 MeshIndex, EMaterialType Type);
-
 	bool HasMeshRegenWorldLocation() const { return bHasNewMeshRegenWorldLocation; }
 	const FVector& GetMeshRegenWorldLocation() const { return MeshRegenWorldLocation; }
 	AVALANCHESHAPES_API void SetMeshRegenWorldLocation(const FVector& NewLocation, bool bImmediateUpdate = false);
@@ -175,14 +220,18 @@ public:
 	/** Get the registered meshes indexes */
 	AVALANCHESHAPES_API TSet<int32> GetMeshesIndexes() const;
 
-	/** Returns the name of the meshes */
-	AVALANCHESHAPES_API TArray<FName> GetMeshDataNames() const;
-
 	/** find a registered mesh and gets a pointer to it */
 	FAvaShapeMeshData* GetMeshData(int32 MeshIndex);
 
 	// apply scale, offset, rotation on dynamic mesh section
-	bool ApplyUVsTransform(FAvaShapeMesh& InMesh, FAvaShapeMaterialUVParameters& InParams, FVector2D ShapeSize, FVector2D UVFixOffset, float UVFixRotation = 0.f) const;
+	bool ApplyUVsTransform(FAvaShapeMesh& InMesh, FAvaShapeMaterialUVParameters& InParams, FVector2D ShapeSize, FVector2D UVFixOffset, float UVFixRotation = 0.f);
+
+	/** Set whether mesh update should run async or not */
+	AVALANCHESHAPES_API void SetRunAsync(bool bInAsync);
+	bool GetRunAsync() const
+	{
+		return bRunAsync;
+	}
 
 protected:
 	static bool SetAlignmentSize(AActor* InActor, const FVector& InSizeMultiplier);
@@ -216,6 +265,9 @@ protected:
 
 	EAvaDynamicMeshUpdateState GetMeshUpdateStatus() const { return MeshUpdateStatus; }
 
+	/** Checks if shape size is render-able and not minimal */
+	bool IsMeshSizeValid() const;
+
 	/** Flag all mesh sections dirty to regenerate them */
 	void MarkAllMeshesDirty();
 
@@ -244,7 +296,7 @@ protected:
 	virtual bool IsMeshVisible(int32 MeshIndex);
 
 	/** Returns the name of the mesh */
-	FName GetMeshDataName(int32 InMeshIndex) const;
+	FName GetMeshSectionName(int32 InMeshIndex) const;
 
 	/** Returns the materials used by every mesh */
 	const TArray<UMaterialInterface*> GetMeshDataMaterials();
@@ -255,9 +307,7 @@ protected:
 	/** Clears everything and reinitialize the DM to use it properly */
 	void InitializeDynamicMesh();
 
-	EMaterialType& GetMaterialType(int32 MeshIndex);
 	bool SetMaterialDirect(int32 MeshIndex, UMaterialInterface* NewMaterial);
-	FAvaShapeMaterialUVParameters* GetMaterialUVParams(int32 MeshIndex);
 
 	// called once the mesh is done updating
 	virtual void OnMeshUpdateFinished();
@@ -266,7 +316,6 @@ protected:
 	virtual void OnPixelSizeChanged() {}
 
 	virtual void OnMeshChanged(int32 MeshIndex);
-	virtual void OnVertexColorChanged();
 	virtual void OnMaterialChanged(int32 MaterialIndex);
 	virtual void OnMaterialTypeChanged(int32 MaterialIndex);
 	virtual void OnParametricMaterialChanged(int32 MaterialIndex);
@@ -284,12 +333,6 @@ protected:
 	// Clears the mesh sections that are dirty
 	virtual bool ClearMesh();
 
-	// Adds vertex colors to the vertex colour array. Does not empty it.
-	virtual bool CreateColors(FAvaShapeMesh& InMesh);
-
-	// Performs a vertex-colour-only update. Then calls OnMeshUpdate to the PMC
-	virtual bool UpdateColors();
-
 	// Special case when the mesh is scaled, rather than reconfigured.
 	virtual void UpdateVertices();
 
@@ -318,16 +361,19 @@ protected:
 	bool UpdateDynamicMesh(FAvaShapeMesh& InMesh);
 
 	// Invalidates a section and runs the mesh update if necessary
-	void InvalidateSection(bool& bInvalidatedSection, bool bRequireUpdate = true, bool bDoAsync = true);
+	void InvalidateSection(bool& bInvalidatedSection, bool bRequireUpdate = true);
 
 	/** Flag a mesh section dirty to regenerate it */
 	void MarkMeshDirty(int32 MeshIndex);
 
 	// Takes all dirty flags, regenerate whatever is needed and updates the DMC, called by RunUpdate or RequireUpdate
-	void UpdateDynamicMesh();
+	void UpdateMesh();
+
+	/** Pushed mesh section data updates to dynamic mesh and creates/updates geometry data */
+	void ApplyMeshSectionChanges();
 
 	// runs the update for this mesh, async by default
-	void RequireUpdate(bool bDoAsync = true);
+	void RequestUpdate();
 
 	// used to check for material changes in slot of DMC to match with our current material settings
 	bool CheckMaterialSlotChanges();
@@ -339,27 +385,24 @@ protected:
 	void ScaleVertices(const FVector& InScale);
 	void ScaleVertices(const FVector2D& InScale);
 
-	// enable mesh size property editing
-	UPROPERTY(EditDefaultsOnly, BlueprintReadOnly, Category="Shape")
+	/** Is mesh size editing allowed */
+	UPROPERTY(Getter="GetAllowEditSize")
 	bool bAllowEditSize;
 
-	// the type of size you want to handle
-	UPROPERTY(EditInstanceOnly, BlueprintReadWrite, Category="Shape", Transient, meta=(DisplayName="Size Type", EditCondition="bAllowEditSize", EditConditionHides, AllowPrivateAccess="true"))
+	/** the type of size you want to handle */
+	UPROPERTY(EditInstanceOnly, Category="Shape", Transient, meta=(DisplayName="Size Type", HideEditConditionToggle, EditCondition="bAllowEditSize", EditConditionHides, AllowPrivateAccess="true"))
 	ESizeType SizeType = ESizeType::UnrealUnit;
 
-	// Uniform scaled size of the mesh
-	UPROPERTY(EditInstanceOnly, BlueprintReadWrite, Transient, Setter, Getter, Category="Shape", meta=(ClampMin="0.0", DisplayName="Uniform Scaled Size", Units="times", AllowPrivateAccess="true"))
+	/** Uniform scaled size of the mesh */
+	UPROPERTY(EditInstanceOnly, Transient, Setter, Getter, Category="Shape", meta=(ClampMin="0.0", DisplayName="Uniform Scaled Size", Units="times", AllowPrivateAccess="true"))
 	float UniformScaledSize = 1.f;
 
-	UPROPERTY(BlueprintReadWrite, AdvancedDisplay, Category="Shape", meta=(DisplayName="Material Vertex Color", AllowPrivateAccess="true"))
-	FLinearColor VertexColor;
-
 	// use primary material for every slot available
-	UPROPERTY(EditInstanceOnly, BlueprintReadWrite, Setter="SetUsePrimaryMaterialEverywhere", Getter="GetUsePrimaryMaterialEverywhere", Category="Material", meta=(DisplayName="Use Single Material", AllowPrivateAccess="true"))
+	UPROPERTY(EditInstanceOnly, Setter="SetUsePrimaryMaterialEverywhere", Getter="GetUsePrimaryMaterialEverywhere", Category="Material", meta=(DisplayName="Use Single Material", AllowPrivateAccess="true"))
 	bool bUsePrimaryMaterialEverywhere;
 
 	// Meshes used for the current shape sections
-	UPROPERTY(EditInstanceOnly, BlueprintReadOnly, Category="Material", meta=(EditFixedOrder, EditFixedSize, AllowPrivateAccess="true"))
+	UPROPERTY(EditInstanceOnly, Category="Material", meta=(EditFixedOrder, EditFixedSize, AllowPrivateAccess="true"))
 	TMap<int32, FAvaShapeMeshData> MeshDatas;
 
 	TArray<FAvaSnapPoint> LocalSnapPoints;
@@ -385,11 +428,11 @@ private:
 	/** Generates the normals for the mesh (all sections), needs the uv set on the mesh */
 	bool GenerateNormals();
 
-	/** runs the update for this mesh only if it's required (mesh is not up to date), async by default */
-	void RunUpdate(bool bDoAsync = true);
+	/** runs the update for this mesh only if it's required (mesh is not up to date) */
+	void RunUpdate();
 
 	/** Restore the existing saved cached mesh on the component */
-	void RestoreCachedMesh();
+	void RestoreCachedMesh(bool bInReset);
 
 	/** Save the existing component mesh into a saved cached mesh */
 	void SaveCachedMesh();
@@ -440,8 +483,9 @@ private:
 	// Should only be used when the shape's vertices have changed, but the overall shape has not changed.
 	bool bVerticesDirty;
 
-	bool bColorsDirty;
-
 	/** Is the mesh scaled to zero ? */
 	bool bZeroScaled = false;
+
+	/** Run mesh updates async */
+	bool bRunAsync = true;
 };
diff --git a/Engine/Plugins/Experimental/Avalanche/Source/AvalancheShapesEditor/Private/Visualizers/AvaShape2DDynMeshVis.cpp b/Engine/Plugins/Experimental/Avalanche/Source/AvalancheShapesEditor/Private/Visualizers/AvaShape2DDynMeshVis.cpp
--- a/Engine/Plugins/Experimental/Avalanche/Source/AvalancheShapesEditor/Private/Visualizers/AvaShape2DDynMeshVis.cpp
+++ b/Engine/Plugins/Experimental/Avalanche/Source/AvalancheShapesEditor/Private/Visualizers/AvaShape2DDynMeshVis.cpp
@@ -62,7 +62,7 @@ bool FAvaShape2DDynamicMeshVisualizer::VisProxyHandleClick(FEditorViewportClient
 
 	if (InVisProxy->IsA(HAvaShapeSizeHitProxy::StaticGetType()))
 	{
-		if (!DynMesh->AllowsSizeEditing())
+		if (!DynMesh->GetAllowEditSize())
 		{
 			return Super::VisProxyHandleClick(InViewportClient, InVisProxy, InClick);
 		}
@@ -105,7 +105,7 @@ void FAvaShape2DDynamicMeshVisualizer::DrawVisualizationNotEditing(const UActorC
 	DrawUVButton(DynMesh, InView, InPDI, InOutIconIndex, 0, GetIconColor(false));
 	++InOutIconIndex;
 
-	if (DynMesh->AllowsSizeEditing())
+	if (DynMesh->GetAllowEditSize())
 	{
 		DrawSizeButtons(DynMesh, InView, InPDI);
 	}
@@ -126,7 +126,7 @@ void FAvaShape2DDynamicMeshVisualizer::DrawVisualizationEditing(const UActorComp
 	DrawUVButton(DynMesh, View, InPDI, InOutIconIndex, 0, GetIconColor(UVSectionIdx == 0, bEditingUVAnchor));
 	++InOutIconIndex;
 
-	if (DynMesh->AllowsSizeEditing())
+	if (DynMesh->GetAllowEditSize())
 	{
 		DrawSizeButtons(DynMesh, View, InPDI);
 	}
diff --git a/Engine/Plugins/Experimental/Avalanche/Source/AvalancheShapesEditor/Private/Visualizers/AvaShape3DDynMeshVis.cpp b/Engine/Plugins/Experimental/Avalanche/Source/AvalancheShapesEditor/Private/Visualizers/AvaShape3DDynMeshVis.cpp
--- a/Engine/Plugins/Experimental/Avalanche/Source/AvalancheShapesEditor/Private/Visualizers/AvaShape3DDynMeshVis.cpp
+++ b/Engine/Plugins/Experimental/Avalanche/Source/AvalancheShapesEditor/Private/Visualizers/AvaShape3DDynMeshVis.cpp
@@ -48,7 +48,7 @@ bool FAvaShape3DDynamicMeshVisualizer::VisProxyHandleClick(FEditorViewportClient
 	// size
 	if (InVisProxy->IsA(HAvaShapeSizeHitProxy::StaticGetType()))
 	{
-		if (!DynMesh->AllowsSizeEditing())
+		if (!DynMesh->GetAllowEditSize())
 		{
 			return Super::VisProxyHandleClick(InViewportClient, InVisProxy, InClick);
 		}
@@ -166,7 +166,7 @@ void FAvaShape3DDynamicMeshVisualizer::DrawVisualizationNotEditing(const UActorC
 		return;
 	}
 
-	if (DynMesh->AllowsSizeEditing())
+	if (DynMesh->GetAllowEditSize())
 	{
 		for (const AvaAlignment Alignment : SupportedAlignments)
 		{
@@ -187,7 +187,7 @@ void FAvaShape3DDynamicMeshVisualizer::DrawVisualizationEditing(const UActorComp
 		return;
 	}
 
-	if (DynMesh->AllowsSizeEditing())
+	if (DynMesh->GetAllowEditSize())
 	{
 		for (const AvaAlignment Alignment : SupportedAlignments)
 		{
