[SequenceNavigator]
- Change menu names from "Sequencer.NavigationTool" to "SequenceNavigator"
- Change namespces from "UE::Sequencer::NavigationTool" to "UE::SequenceNavigator"
- Remove unused FNavigationToolInstance class
- Remove unneeded instances of IsValid()

diff --git a/Engine/Plugins/Experimental/Avalanche/Source/AvalancheSequencer/Private/AvaSequencer.cpp b/Engine/Plugins/Experimental/Avalanche/Source/AvalancheSequencer/Private/AvaSequencer.cpp
--- a/Engine/Plugins/Experimental/Avalanche/Source/AvalancheSequencer/Private/AvaSequencer.cpp
+++ b/Engine/Plugins/Experimental/Avalanche/Source/AvalancheSequencer/Private/AvaSequencer.cpp
@@ -487,7 +487,7 @@ void FAvaSequencer::GetSelectedObjects(const TArray<FGuid>& InObjectGuids
 		UObject* const BoundObject = BoundObjects[0].Get();
 
 		// Skip invalid or objects already processed
-		if (!IsValid(BoundObject) || ProcessedObjects.Contains(BoundObject))
+		if (!BoundObject || ProcessedObjects.Contains(BoundObject))
 		{
 			continue;
 		}
@@ -637,7 +637,7 @@ UAvaSequence* FAvaSequencer::CreateSequence() const
 	}
 
 	UObject* const Outer = SequenceProvider->ToUObject();
-	if (!IsValid(Outer))
+	if (!Outer)
 	{
 		return nullptr;
 	}
@@ -988,7 +988,7 @@ void FAvaSequencer::ApplyCustomPresetToSelection(FName InPresetName)
 
 void FAvaSequencer::ApplyPresetToSelection(const FAvaSequencePreset& InPreset)
 {
-	using namespace UE::Sequencer::NavigationTool;
+	using namespace UE::SequenceNavigator;
 
 	const TSharedPtr<ISequencer> Sequencer = SequencerWeak.Pin();
 	if (!Sequencer.IsValid())
@@ -1058,7 +1058,7 @@ void FAvaSequencer::DuplicateSequences_Execute()
 	}
 
 	UObject* const Outer = SequenceProvider->ToUObject();
-	if (!IsValid(Outer))
+	if (!Outer)
 	{
 		return;
 	}
@@ -1142,7 +1142,7 @@ void FAvaSequencer::DeleteSequences_Execute()
 	}
 
 	UObject* const Outer = SequenceProvider->ToUObject();
-	if (!IsValid(Outer))
+	if (!Outer)
 	{
 		return;
 	}
@@ -1353,7 +1353,7 @@ void FAvaSequencer::DuplicateSequences()
 	}
 
 	UObject* const Outer = SequenceProvider->ToUObject();
-	if (!IsValid(Outer))
+	if (!Outer)
 	{
 		return;
 	}
@@ -1364,10 +1364,11 @@ void FAvaSequencer::DuplicateSequences()
 
 	for (UAvaSequence* const TemplateSequence : SelectedSequences)
 	{
-		check(IsValid(TemplateSequence));
-
-		UAvaSequence* const DupedSequence = DuplicateObject<UAvaSequence>(TemplateSequence, Outer);
-		SequenceProvider->AddSequence(DupedSequence);
+		if (TemplateSequence)
+		{
+			UAvaSequence* const DupedSequence = DuplicateObject<UAvaSequence>(TemplateSequence, Outer);
+			SequenceProvider->AddSequence(DupedSequence);
+		}
 	}
 }
 
@@ -1399,8 +1400,7 @@ void FAvaSequencer::ExportSequences()
 
 			for (const TWeakObjectPtr<> WeakObject : WeakSelectedObjects)
 			{
-				UObject* const SyncObject = WeakObject.Get();
-				if (IsValid(SyncObject))
+				if (UObject* const SyncObject = WeakObject.Get())
 				{
 					SyncObjects.Add(SyncObject);
 				}
@@ -1480,7 +1480,7 @@ void FAvaSequencer::DeleteSequences(const TSet<UAvaSequence*>& InSequences)
 	}
 
 	UObject* const Outer = SequenceProvider->ToUObject();
-	if (!IsValid(Outer))
+	if (!Outer)
 	{
 		return;
 	}
@@ -1495,7 +1495,7 @@ void FAvaSequencer::DeleteSequences(const TSet<UAvaSequence*>& InSequences)
 	// Remove the sequences from the list (not marked as garbage yet)
 	for (UAvaSequence* const Sequence : InSequences)
 	{
-		check(IsValid(Sequence));
+		check(Sequence);
 
 		Sequence->Modify();
 		SequenceProvider->RemoveSequence(Sequence);
@@ -1551,7 +1551,7 @@ void FAvaSequencer::PlaySequences()
 	for (UAvaSequence* const Sequence : SelectedSequences)
 	{
 		IAvaSequencePlaybackObject* const PlaybackObject = Provider.GetPlaybackObject();
-		if (IsValid(Sequence) && PlaybackObject)
+		if (Sequence && PlaybackObject)
 		{
 			PlaybackObject->PlaySequence(Sequence, PlaySettings);
 		}
@@ -1571,7 +1571,7 @@ void FAvaSequencer::ContinueSequences()
 	for (UAvaSequence* const Sequence : SelectedSequences)
 	{
 		IAvaSequencePlaybackObject* const PlaybackObject = Provider.GetPlaybackObject();
-		if (IsValid(Sequence) && PlaybackObject)
+		if (Sequence && PlaybackObject)
 		{
 			PlaybackObject->ContinueSequence(Sequence);
 		}
@@ -1590,7 +1590,7 @@ void FAvaSequencer::StopSequences()
 	for (UAvaSequence* const Sequence : SelectedSequences)
 	{
 		IAvaSequencePlaybackObject* const PlaybackObject = Provider.GetPlaybackObject();
-		if (IsValid(Sequence) && PlaybackObject)
+		if (Sequence && PlaybackObject)
 		{
 			PlaybackObject->StopSequence(Sequence);
 		}
@@ -1599,7 +1599,7 @@ void FAvaSequencer::StopSequences()
 
 UObject* FAvaSequencer::FindObjectToPossess(UObject* InParentObject, const FMovieScenePossessable& InPossessable)
 {
-	if (!IsValid(InParentObject))
+	if (!InParentObject)
 	{
 		return nullptr;
 	}
@@ -1657,7 +1657,7 @@ void FAvaSequencer::FixBindingPaths()
 {
 	UAvaSequence* const Sequence   = GetViewedSequence();
 	UObject* const PlaybackContext = GetPlaybackContext();
-	if (!IsValid(Sequence) || !IsValid(PlaybackContext))
+	if (!Sequence || !PlaybackContext)
 	{
 		return;
 	}
@@ -1681,19 +1681,19 @@ void FAvaSequencer::FixBindingPaths()
 void FAvaSequencer::FixInvalidBindings()
 {
 	UObject* const PlaybackContext = GetPlaybackContext();
-	if (!IsValid(PlaybackContext))
+	if (!PlaybackContext)
 	{
 		return;
 	}
 
 	UAvaSequence* const Sequence = GetViewedSequence();
-	if (!IsValid(Sequence))
+	if (!Sequence)
 	{
 		return;
 	}
 
 	UMovieScene* const MovieScene = Sequence->GetMovieScene();
-	if (!IsValid(MovieScene))
+	if (!MovieScene)
 	{
 		return;
 	}
@@ -1982,7 +1982,7 @@ TArray<UAvaSequence*> FAvaSequencer::GetSequencesForObject(UObject* InObject) co
 {
 	TArray<UAvaSequence*> OutSequences;
 
-	if (!IsValid(InObject))
+	if (!InObject)
 	{
 		return OutSequences;
 	}
@@ -1995,7 +1995,7 @@ TArray<UAvaSequence*> FAvaSequencer::GetSequencesForObject(UObject* InObject) co
 
 	for (UAvaSequence* const Sequence : SequenceProvider->GetSequences())
 	{
-		if (!IsValid(Sequence))
+		if (!Sequence)
 		{
 			continue;
 		}
@@ -2022,7 +2022,7 @@ TSharedRef<SWidget> FAvaSequencer::CreateSequenceWidget()
 	const TSharedRef<ISequencer> Sequencer = GetSequencer();
 
 	UAvaSequencerSettings* const SequencerSettings = GetMutableDefault<UAvaSequencerSettings>();
-	check(IsValid(SequencerSettings));
+	check(SequencerSettings);
 	FSidebarState& SidebarState = SequencerSettings->GetSidebarState();
 
 	if (SidebarState.IsVisible())
@@ -2087,7 +2087,7 @@ void FAvaSequencer::OnEditorSelectionChanged(const FAvaEditorSelection& InEditor
 
 		for (UObject* const SelectedObject : InEditorSelection.GetSelectedObjects<UObject, EAvaSelectionSource::All>())
 		{
-			if (IsValid(SelectedObject))
+			if (SelectedObject)
 			{
 				FGuid SelectedObjectGuid = Sequencer->FindObjectId(*SelectedObject, Sequencer->GetFocusedTemplateID());
 				Sequencer->SelectObject(SelectedObjectGuid);
@@ -2114,7 +2114,7 @@ void FAvaSequencer::OnEditorSelectionChanged(const FAvaEditorSelection& InEditor
 
 void FAvaSequencer::NotifyOnSequenceTreeChanged()
 {
-	using namespace UE::Sequencer::NavigationTool;
+	using namespace UE::SequenceNavigator;
 
 	if (const TSharedPtr<INavigationTool> Tool = FNavigationToolExtender::FindNavigationTool(GetSequencer()))
 	{
@@ -2207,9 +2207,9 @@ void FAvaSequencer::PostRedo(const bool bInSuccess)
 	PostUndo(bInSuccess);
 }
 
-TArray<TSharedPtr<UE::Sequencer::NavigationTool::FNavigationToolAvaSequence>> FAvaSequencer::GetSelectedSequenceItems() const
+TArray<TSharedPtr<UE::SequenceNavigator::FNavigationToolAvaSequence>> FAvaSequencer::GetSelectedSequenceItems() const
 {
-	using namespace UE::Sequencer::NavigationTool;
+	using namespace UE::SequenceNavigator;
 
 	const TSharedPtr<ISequencer> Sequencer = SequencerWeak.Pin();
 	if (!Sequencer.IsValid())
@@ -2245,7 +2245,7 @@ TArray<TSharedPtr<UE::Sequencer::NavigationTool::FNavigationToolAvaSequence>> FA
 
 TArray<UAvaSequence*> FAvaSequencer::GetSelectedSequences() const
 {
-	using namespace UE::Sequencer::NavigationTool;
+	using namespace UE::SequenceNavigator;
 
 	const TArray<TSharedPtr<FNavigationToolAvaSequence>> SelectedItems = GetSelectedSequenceItems();
 	if (SelectedItems.IsEmpty())
@@ -2258,8 +2258,7 @@ TArray<UAvaSequence*> FAvaSequencer::GetSelectedSequences() const
 
 	for (const TSharedPtr<FNavigationToolAvaSequence>& SequenceItem : SelectedItems)
 	{
-		UAvaSequence* const AvaSequence = SequenceItem->GetAvaSequence();
-		if (IsValid(AvaSequence))
+		if (UAvaSequence* const AvaSequence = SequenceItem->GetAvaSequence())
 		{
 			SelectedAvaSequences.Add(AvaSequence);
 		}
@@ -2380,8 +2379,7 @@ TSharedRef<FAvaEaseCurveTool> FAvaSequencer::GetEaseCurveTool() const
 
 void FAvaSequencer::OnSidebarStateChanged(const FSidebarState& InNewState)
 {
-	UAvaSequencerSettings* const SequencerSettings = GetMutableDefault<UAvaSequencerSettings>();
-	if (IsValid(SequencerSettings))
+	if (UAvaSequencerSettings* const SequencerSettings = GetMutableDefault<UAvaSequencerSettings>())
 	{
 		SequencerSettings->SetSidebarState(InNewState);
 	}
@@ -2396,13 +2394,13 @@ void FAvaSequencer::ExtendSidebarSelectionMenu(FMenuBuilder& OutMenuBuilder)
 void FAvaSequencer::ExtendSidebarMarkedFramesMenu(FMenuBuilder& OutMenuBuilder)
 {
 	UAvaSequence* const Sequence = GetViewedSequence();
-	if (!IsValid(Sequence))
+	if (!Sequence)
 	{
 		return;
 	}
 
 	UMovieScene* const MovieScene = Sequence->GetMovieScene();
-	if (!IsValid(MovieScene))
+	if (!MovieScene)
 	{
 		return;
 	}
@@ -2447,7 +2445,7 @@ const TArray<TWeakObjectPtr<UAvaSequence>>& FAvaSequencer::GetRootSequences() co
 
 void FAvaSequencer::RegisterNavigationTool()
 {
-	using namespace UE::Sequencer::NavigationTool;
+	using namespace UE::SequenceNavigator;
 
 	if (!SequencerWeak.IsValid())
 	{
@@ -2468,7 +2466,7 @@ void FAvaSequencer::RegisterNavigationTool()
 
 void FAvaSequencer::UnregisterNavigationTool()
 {
-	using namespace UE::Sequencer::NavigationTool;
+	using namespace UE::SequenceNavigator;
 
 	if (!NavigationToolProvider.IsValid())
 	{
diff --git a/Engine/Plugins/Experimental/Avalanche/Source/AvalancheSequencer/Private/AvaSequencer.h b/Engine/Plugins/Experimental/Avalanche/Source/AvalancheSequencer/Private/AvaSequencer.h
--- a/Engine/Plugins/Experimental/Avalanche/Source/AvalancheSequencer/Private/AvaSequencer.h
+++ b/Engine/Plugins/Experimental/Avalanche/Source/AvalancheSequencer/Private/AvaSequencer.h
@@ -41,18 +41,18 @@ namespace UE
 	namespace Sequencer
 	{
 		class SOutlinerView;
-
-		namespace NavigationTool
-		{
-			class FNavigationToolAvaSequence;
-			class FNavigationToolProvider;
-		}
 	}
 
 	namespace AvaSequencer
 	{
 		class FAvaNavigationToolProvider;
 	}
+
+	namespace SequenceNavigator
+	{
+		class FNavigationToolAvaSequence;
+		class FNavigationToolProvider;
+	}
 }
 
 class FAvaSequencer : public IAvaSequencer, public FEditorUndoClient, public TSharedFromThis<FAvaSequencer>
@@ -262,7 +262,7 @@ public:
 	// End of FEditorUndoClient
 
 private:
-	TArray<TSharedPtr<UE::Sequencer::NavigationTool::FNavigationToolAvaSequence>> GetSelectedSequenceItems() const;
+	TArray<TSharedPtr<UE::SequenceNavigator::FNavigationToolAvaSequence>> GetSelectedSequenceItems() const;
 
 	TArray<UAvaSequence*> GetSelectedSequences() const;
 
diff --git a/Engine/Plugins/Experimental/Avalanche/Source/AvalancheSequencer/Private/NavigationTool/AvaNavigationToolProvider.cpp b/Engine/Plugins/Experimental/Avalanche/Source/AvalancheSequencer/Private/NavigationTool/AvaNavigationToolProvider.cpp
--- a/Engine/Plugins/Experimental/Avalanche/Source/AvalancheSequencer/Private/NavigationTool/AvaNavigationToolProvider.cpp
+++ b/Engine/Plugins/Experimental/Avalanche/Source/AvalancheSequencer/Private/NavigationTool/AvaNavigationToolProvider.cpp
@@ -71,7 +71,7 @@ FText FAvaNavigationToolProvider::GetDefaultColumnView() const
 	return MotionDesignColumnViewName;
 }
 
-FNavigationToolSaveState* FAvaNavigationToolProvider::GetSaveState(const UE::Sequencer::NavigationTool::INavigationTool& InTool) const
+FNavigationToolSaveState* FAvaNavigationToolProvider::GetSaveState(const UE::SequenceNavigator::INavigationTool& InTool) const
 {
 #if WITH_EDITOR
 	if (IAvaSceneInterface* const AvaScene = GetSceneInterface(InTool))
@@ -82,7 +82,7 @@ FNavigationToolSaveState* FAvaNavigationToolProvider::GetSaveState(const UE::Seq
 	return nullptr;
 }
 
-void FAvaNavigationToolProvider::SetSaveState(const UE::Sequencer::NavigationTool::INavigationTool& InTool, const FNavigationToolSaveState& InSaveState) const
+void FAvaNavigationToolProvider::SetSaveState(const UE::SequenceNavigator::INavigationTool& InTool, const FNavigationToolSaveState& InSaveState) const
 {
 #if WITH_EDITOR
 	if (IAvaSceneInterface* const AvaScene = GetSceneInterface(InTool))
@@ -148,7 +148,7 @@ void FAvaNavigationToolProvider::OnActivate()
 	ExtendToolToolBar();
 	ExtendToolItemContextMenu();
 
-	//DragDropInitializedDelegate = UE::Sequencer::NavigationTool::FNavigationToolExtender::OnItemDragDropOpInitialized()
+	//DragDropInitializedDelegate = UE::SequenceNavigator::FNavigationToolExtender::OnItemDragDropOpInitialized()
 	//	.AddSP(this, &FAvaNavigationToolProvider::ExtendToolItemDragDropOp);
 }
 
@@ -157,13 +157,13 @@ void FAvaNavigationToolProvider::OnDeactivate()
 	RemoveToolToolBarExtension();
 	RemoveToolItemContextMenuExtension();
 
-	//UE::Sequencer::NavigationTool::FNavigationToolExtender::OnItemDragDropOpInitialized()
+	//UE::SequenceNavigator::FNavigationToolExtender::OnItemDragDropOpInitialized()
 	//	.Remove(DragDropInitializedDelegate);
 }
 
-void FAvaNavigationToolProvider::OnExtendColumns(UE::Sequencer::NavigationTool::FNavigationToolColumnExtender& OutExtender)
+void FAvaNavigationToolProvider::OnExtendColumns(UE::SequenceNavigator::FNavigationToolColumnExtender& OutExtender)
 {
-	using namespace UE::Sequencer::NavigationTool;
+	using namespace UE::SequenceNavigator;
 
 	// Support built in columns
 	OutExtender.AddColumn<FNavigationToolPlayheadColumn>();
@@ -195,7 +195,7 @@ void FAvaNavigationToolProvider::OnExtendColumns(UE::Sequencer::NavigationTool::
 
 void FAvaNavigationToolProvider::OnExtendColumnViews(TSet<FNavigationToolColumnView>& OutColumnViews)
 {
-	using namespace UE::Sequencer::NavigationTool;
+	using namespace UE::SequenceNavigator;
 
 	FNavigationToolColumnView& MotionDesignColumnView = OutColumnViews
 		.FindOrAdd(FNavigationToolColumnView(MotionDesignColumnViewName));
@@ -211,12 +211,12 @@ void FAvaNavigationToolProvider::OnExtendColumnViews(TSet<FNavigationToolColumnV
 	FNavigationToolProvider::OnExtendColumnViews(OutColumnViews);
 }
 
-void FAvaNavigationToolProvider::OnExtendItemChildren(UE::Sequencer::NavigationTool::INavigationTool& InTool
-	, const UE::Sequencer::NavigationTool::FNavigationToolItemPtr& InParentItem
-	, TArray<UE::Sequencer::NavigationTool::FNavigationToolItemPtr>& OutChildren
+void FAvaNavigationToolProvider::OnExtendItemChildren(UE::SequenceNavigator::INavigationTool& InTool
+	, const UE::SequenceNavigator::FNavigationToolItemPtr& InParentItem
+	, TArray<UE::SequenceNavigator::FNavigationToolItemPtr>& OutChildren
 	, const bool bInRecursive)
 {
-	using namespace UE::Sequencer::NavigationTool;
+	using namespace UE::SequenceNavigator;
 
 	FNavigationToolProvider::OnExtendItemChildren(InTool, InParentItem, OutChildren, bInRecursive);
 
@@ -249,8 +249,7 @@ void FAvaNavigationToolProvider::OnExtendItemChildren(UE::Sequencer::NavigationT
 	// Add child sequence items
 	for (const TWeakObjectPtr<UAvaSequence>& RootSequenceWeak : RootWeakSequences)
 	{
-		UAvaSequence* const RootSequence = RootSequenceWeak.Get();
-		if (IsValid(RootSequence))
+		if (UAvaSequence* const RootSequence = RootSequenceWeak.Get())
 		{
 			const FNavigationToolItemPtr NewItem = InTool.FindOrAdd<FNavigationToolAvaSequence>(Provider
 				, InParentItem, RootSequence);
@@ -274,7 +273,7 @@ void FAvaNavigationToolProvider::OnExtendBuiltInFilters(TArray<FNavigationToolBu
 	FNavigationToolProvider::OnExtendBuiltInFilters(OutFilterParams);
 }
 
-IAvaSceneInterface* FAvaNavigationToolProvider::GetSceneInterface(const UE::Sequencer::NavigationTool::INavigationTool& InTool) const
+IAvaSceneInterface* FAvaNavigationToolProvider::GetSceneInterface(const UE::SequenceNavigator::INavigationTool& InTool) const
 {
 	if (const TSharedPtr<ISequencer> Sequencer = InTool.GetSequencer())
 	{
@@ -283,7 +282,7 @@ IAvaSceneInterface* FAvaNavigationToolProvider::GetSceneInterface(const UE::Sequ
 	return nullptr;
 }
 
-IAvaSequenceProvider* FAvaNavigationToolProvider::GetSequenceProvider(const UE::Sequencer::NavigationTool::INavigationTool& InTool) const
+IAvaSequenceProvider* FAvaNavigationToolProvider::GetSequenceProvider(const UE::SequenceNavigator::INavigationTool& InTool) const
 {
 	if (const TSharedPtr<ISequencer> Sequencer = InTool.GetSequencer())
 	{
@@ -294,8 +293,8 @@ IAvaSequenceProvider* FAvaNavigationToolProvider::GetSequenceProvider(const UE::
 
 void FAvaNavigationToolProvider::ExtendToolToolBar()
 {
-	UToolMenu* const ToolMenu = UToolMenus::Get()->ExtendMenu(UE::Sequencer::NavigationTool::GetToolBarMenuName());
-	if (!IsValid(ToolMenu))
+	UToolMenu* const ToolMenu = UToolMenus::Get()->ExtendMenu(UE::SequenceNavigator::GetToolBarMenuName());
+	if (!ToolMenu)
 	{
 		return;
 	}
@@ -345,8 +344,8 @@ void FAvaNavigationToolProvider::ExtendToolItemContextMenu()
 		return;
 	}
 
-	UToolMenu* const ExtendedMenu = ToolMenus->ExtendMenu(UE::Sequencer::NavigationTool::GetItemContextMenuName());
-	if (!IsValid(ExtendedMenu))
+	UToolMenu* const ExtendedMenu = ToolMenus->ExtendMenu(UE::SequenceNavigator::GetItemContextMenuName());
+	if (!ExtendedMenu)
 	{
 		return;
 	}
@@ -374,8 +373,7 @@ void FAvaNavigationToolProvider::ExtendToolItemContextMenu()
 
 void FAvaNavigationToolProvider::RemoveToolToolBarExtension()
 {
-	UToolMenu* const Menu = UToolMenus::Get()->FindMenu(UE::Sequencer::NavigationTool::GetToolBarMenuName());
-	if (IsValid(Menu))
+	if (UToolMenu* const Menu = UToolMenus::Get()->FindMenu(UE::SequenceNavigator::GetToolBarMenuName()))
 	{
 		Menu->RemoveSection(ToolbarSectionName);
 	}
@@ -383,21 +381,20 @@ void FAvaNavigationToolProvider::RemoveToolToolBarExtension()
 
 void FAvaNavigationToolProvider::RemoveToolItemContextMenuExtension()
 {
-	UToolMenu* const Menu = UToolMenus::Get()->FindMenu(UE::Sequencer::NavigationTool::GetItemContextMenuName());
-	if (IsValid(Menu))
+	if (UToolMenu* const Menu = UToolMenus::Get()->FindMenu(UE::SequenceNavigator::GetItemContextMenuName()))
 	{
 		Menu->RemoveSection(ContextMenuSectionName);
 	}
 }
 
-void FAvaNavigationToolProvider::ExtendToolItemDragDropOp(UE::Sequencer::NavigationTool::FNavigationToolItemDragDropOp& InDragDropOp)
+void FAvaNavigationToolProvider::ExtendToolItemDragDropOp(UE::SequenceNavigator::FNavigationToolItemDragDropOp& InDragDropOp)
 {
-	InDragDropOp.AddDropHandler<UE::Sequencer::NavigationTool::FNavigationToolAvaSequenceDropHandler>(WeakAvaSequencer);
+	InDragDropOp.AddDropHandler<UE::SequenceNavigator::FNavigationToolAvaSequenceDropHandler>(WeakAvaSequencer);
 }
 
 void FAvaNavigationToolProvider::OnSequenceAdded(UAvaSequence* const InAvaSequence)
 {
-	using namespace UE::Sequencer::NavigationTool;
+	using namespace UE::SequenceNavigator;
 
 	const TSharedPtr<FAvaSequencer> AvaSequencer = StaticCastSharedPtr<FAvaSequencer>(WeakAvaSequencer.Pin());
 	if (!AvaSequencer.IsValid())
@@ -423,7 +420,7 @@ void FAvaNavigationToolProvider::OnSequenceAdded(UAvaSequence* const InAvaSequen
 void FAvaNavigationToolProvider::GeneratePresetMenu(UToolMenu* const InToolMenu)
 {
 	const UAvaSequencerSettings* const SequencerSettings = GetDefault<UAvaSequencerSettings>();
-	if (!IsValid(SequencerSettings))
+	if (!SequencerSettings)
 	{
 		return;
 	}
@@ -475,10 +472,9 @@ void FAvaNavigationToolProvider::GeneratePresetMenu(UToolMenu* const InToolMenu)
 		, FSlateIcon(FAppStyle::GetAppStyleSetName(), "Icons.Settings")
 		, FExecuteAction::CreateLambda([]
 		{
-			ISettingsModule& SettingsModule = FModuleManager::LoadModuleChecked<ISettingsModule>(TEXT("Settings"));
-			const UAvaSequencerSettings* const AvaSequencerSettings = GetDefault<UAvaSequencerSettings>();
-			if (IsValid(AvaSequencerSettings))
+			if (const UAvaSequencerSettings* const AvaSequencerSettings = GetDefault<UAvaSequencerSettings>())
 			{
+				ISettingsModule& SettingsModule = FModuleManager::LoadModuleChecked<ISettingsModule>(TEXT("Settings"));
 				SettingsModule.ShowViewer(AvaSequencerSettings->GetContainerName(), AvaSequencerSettings->GetCategoryName(), AvaSequencerSettings->GetSectionName());
 			}
 		}));
@@ -486,7 +482,7 @@ void FAvaNavigationToolProvider::GeneratePresetMenu(UToolMenu* const InToolMenu)
 
 bool FAvaNavigationToolProvider::CanRelabelSequence() const
 {
-	using namespace UE::Sequencer::NavigationTool;
+	using namespace UE::SequenceNavigator;
 
 	const TSharedPtr<IAvaSequencer> AvaSequencer = WeakAvaSequencer.Pin();
 	if (!AvaSequencer.IsValid())
@@ -507,7 +503,7 @@ bool FAvaNavigationToolProvider::CanRelabelSequence() const
 
 void FAvaNavigationToolProvider::RelabelSequence()
 {
-	using namespace UE::Sequencer::NavigationTool;
+	using namespace UE::SequenceNavigator;
 
 	const TSharedPtr<IAvaSequencer> AvaSequencer = WeakAvaSequencer.Pin();
 	if (!AvaSequencer.IsValid())
@@ -539,7 +535,7 @@ void FAvaNavigationToolProvider::RelabelSequence()
 
 bool FAvaNavigationToolProvider::CanAddSequenceToSelection() const
 {
-	using namespace UE::Sequencer::NavigationTool;
+	using namespace UE::SequenceNavigator;
 
 	const TSharedPtr<IAvaSequencer> AvaSequencer = WeakAvaSequencer.Pin();
 	if (!AvaSequencer.IsValid())
@@ -560,7 +556,7 @@ bool FAvaNavigationToolProvider::CanAddSequenceToSelection() const
 
 void FAvaNavigationToolProvider::AddSequenceToSelection()
 {
-	using namespace UE::Sequencer::NavigationTool;
+	using namespace UE::SequenceNavigator;
 
 	const TSharedPtr<FAvaSequencer> AvaSequencer = WeakAvaSequencer.Pin();
 	if (!AvaSequencer.IsValid())
@@ -588,7 +584,7 @@ void FAvaNavigationToolProvider::AddSequenceToSelection()
 	}
 
 	UAvaSequence* const NewSequence = AvaSequencer->AddSequence(ParentSequence);
-	if (!IsValid(NewSequence))
+	if (!NewSequence)
 	{
 		return;
 	}
diff --git a/Engine/Plugins/Experimental/Avalanche/Source/AvalancheSequencer/Private/NavigationTool/AvaNavigationToolProvider.h b/Engine/Plugins/Experimental/Avalanche/Source/AvalancheSequencer/Private/NavigationTool/AvaNavigationToolProvider.h
--- a/Engine/Plugins/Experimental/Avalanche/Source/AvalancheSequencer/Private/NavigationTool/AvaNavigationToolProvider.h
+++ b/Engine/Plugins/Experimental/Avalanche/Source/AvalancheSequencer/Private/NavigationTool/AvaNavigationToolProvider.h
@@ -16,7 +16,7 @@ class UAvaSequence;
 class UObject;
 struct FNavigationToolSaveState;
 
-namespace UE::Sequencer::NavigationTool
+namespace UE::SequenceNavigator
 {
 	class FNavigationToolItem;
 	class FNavigationToolItemDragDropOp;
@@ -26,7 +26,7 @@ namespace UE::Sequencer::NavigationTool
 namespace UE::AvaSequencer
 {
 
-class FAvaNavigationToolProvider : public UE::Sequencer::NavigationTool::FNavigationToolProvider
+class FAvaNavigationToolProvider : public UE::SequenceNavigator::FNavigationToolProvider
 {
 public:
 	static const FName Identifier;
@@ -44,19 +44,19 @@ public:
 
 	virtual FText GetDefaultColumnView() const override;
 
-	virtual FNavigationToolSaveState* GetSaveState(const UE::Sequencer::NavigationTool::INavigationTool& InTool) const override;
-	virtual void SetSaveState(const UE::Sequencer::NavigationTool::INavigationTool& InTool, const FNavigationToolSaveState& InSaveState) const override;
+	virtual FNavigationToolSaveState* GetSaveState(const UE::SequenceNavigator::INavigationTool& InTool) const override;
+	virtual void SetSaveState(const UE::SequenceNavigator::INavigationTool& InTool, const FNavigationToolSaveState& InSaveState) const override;
 
 	virtual void BindCommands(const TSharedRef<FUICommandList>& InCommandList) override;
 
 	virtual void OnActivate() override;
 	virtual void OnDeactivate() override;
 
-	virtual void OnExtendColumns(UE::Sequencer::NavigationTool::FNavigationToolColumnExtender& OutExtender) override;
+	virtual void OnExtendColumns(UE::SequenceNavigator::FNavigationToolColumnExtender& OutExtender) override;
 	virtual void OnExtendColumnViews(TSet<FNavigationToolColumnView>& OutColumnViews) override;
-	virtual void OnExtendItemChildren(UE::Sequencer::NavigationTool::INavigationTool& InTool
-		, const UE::Sequencer::NavigationTool::FNavigationToolItemPtr& InParentItem
-		, TArray<UE::Sequencer::NavigationTool::FNavigationToolItemPtr>& OutChildren
+	virtual void OnExtendItemChildren(UE::SequenceNavigator::INavigationTool& InTool
+		, const UE::SequenceNavigator::FNavigationToolItemPtr& InParentItem
+		, TArray<UE::SequenceNavigator::FNavigationToolItemPtr>& OutChildren
 		, const bool bInRecursive) override;
 	virtual void OnExtendBuiltInFilters(TArray<FNavigationToolBuiltInFilterParams>& OutFilterParams) override;
 
@@ -66,8 +66,8 @@ private:
 	static const FName ToolbarSectionName;
 	static const FName ContextMenuSectionName;
 
-	IAvaSceneInterface* GetSceneInterface(const UE::Sequencer::NavigationTool::INavigationTool& InTool) const;
-	IAvaSequenceProvider* GetSequenceProvider(const UE::Sequencer::NavigationTool::INavigationTool& InTool) const;
+	IAvaSceneInterface* GetSceneInterface(const UE::SequenceNavigator::INavigationTool& InTool) const;
+	IAvaSequenceProvider* GetSequenceProvider(const UE::SequenceNavigator::INavigationTool& InTool) const;
 
 	void ExtendToolToolBar();
 	void ExtendToolItemContextMenu();
@@ -75,7 +75,7 @@ private:
 	void RemoveToolToolBarExtension();
 	void RemoveToolItemContextMenuExtension();
 
-	void ExtendToolItemDragDropOp(UE::Sequencer::NavigationTool::FNavigationToolItemDragDropOp& InDragDropOp);
+	void ExtendToolItemDragDropOp(UE::SequenceNavigator::FNavigationToolItemDragDropOp& InDragDropOp);
 
 	void OnSequenceAdded(UAvaSequence* const InAvaSequence);
 
diff --git a/Engine/Plugins/Experimental/Avalanche/Source/AvalancheSequencer/Private/NavigationTool/Columns/AvaNavigationToolStatusColumn.cpp b/Engine/Plugins/Experimental/Avalanche/Source/AvalancheSequencer/Private/NavigationTool/Columns/AvaNavigationToolStatusColumn.cpp
--- a/Engine/Plugins/Experimental/Avalanche/Source/AvalancheSequencer/Private/NavigationTool/Columns/AvaNavigationToolStatusColumn.cpp
+++ b/Engine/Plugins/Experimental/Avalanche/Source/AvalancheSequencer/Private/NavigationTool/Columns/AvaNavigationToolStatusColumn.cpp
@@ -10,7 +10,7 @@
 
 #define LOCTEXT_NAMESPACE "AvaNavigationToolStatusColumn"
 
-namespace UE::Sequencer::NavigationTool
+namespace UE::SequenceNavigator
 {
 
 FText FAvaNavigationToolStatusColumn::GetColumnDisplayNameText() const
@@ -61,6 +61,6 @@ TSharedRef<SWidget> FAvaNavigationToolStatusColumn::ConstructRowWidget(const FNa
 	return SNew(SAvaNavigationToolStatus, InItem, InView, InRow, Sequencer.ToSharedRef());
 }
 
-} // UE::Sequencer::NavigationTool
+} // UE::SequenceNavigator
 
 #undef LOCTEXT_NAMESPACE
diff --git a/Engine/Plugins/Experimental/Avalanche/Source/AvalancheSequencer/Private/NavigationTool/Columns/AvaNavigationToolStatusColumn.h b/Engine/Plugins/Experimental/Avalanche/Source/AvalancheSequencer/Private/NavigationTool/Columns/AvaNavigationToolStatusColumn.h
--- a/Engine/Plugins/Experimental/Avalanche/Source/AvalancheSequencer/Private/NavigationTool/Columns/AvaNavigationToolStatusColumn.h
+++ b/Engine/Plugins/Experimental/Avalanche/Source/AvalancheSequencer/Private/NavigationTool/Columns/AvaNavigationToolStatusColumn.h
@@ -4,7 +4,7 @@
 
 #include "Columns/NavigationToolColumn.h"
 
-namespace UE::Sequencer::NavigationTool
+namespace UE::SequenceNavigator
 {
 
 class FAvaNavigationToolStatusColumn : public FNavigationToolColumn
@@ -28,4 +28,4 @@ protected:
 	//~ End INavigationToolColumn
 };
 
-} // UE::Sequencer::NavigationTool
+} // UE::SequenceNavigator
diff --git a/Engine/Plugins/Experimental/Avalanche/Source/AvalancheSequencer/Private/NavigationTool/Columns/SAvaNavigationToolStatus.cpp b/Engine/Plugins/Experimental/Avalanche/Source/AvalancheSequencer/Private/NavigationTool/Columns/SAvaNavigationToolStatus.cpp
--- a/Engine/Plugins/Experimental/Avalanche/Source/AvalancheSequencer/Private/NavigationTool/Columns/SAvaNavigationToolStatus.cpp
+++ b/Engine/Plugins/Experimental/Avalanche/Source/AvalancheSequencer/Private/NavigationTool/Columns/SAvaNavigationToolStatus.cpp
@@ -15,7 +15,7 @@
 
 #define LOCTEXT_NAMESPACE "SAvaNavigationToolStatus"
 
-namespace UE::Sequencer::NavigationTool
+namespace UE::SequenceNavigator
 {
 
 void SAvaNavigationToolStatus::Construct(const FArguments& InArgs
@@ -81,7 +81,7 @@ void SAvaNavigationToolStatus::Tick(const FGeometry& InAllottedGeometry, const d
 	}
 
 	UAvaSequence* const Sequence = AvaSequenceItem->GetAvaSequence();
-	if (!ensure(IsValid(Sequence)))
+	if (!Sequence)
 	{
 		return;
 	}
@@ -100,7 +100,7 @@ void SAvaNavigationToolStatus::Tick(const FGeometry& InAllottedGeometry, const d
 
 	UAvaSequencePlayer* const Player = PlaybackObject->GetSequencePlayer(Sequence);
 
-	bSequenceInProgress = IsValid(Player);
+	bSequenceInProgress = Player != nullptr;
 	if (bSequenceInProgress)
 	{
 		const EMovieScenePlayerStatus::Type Status = Player->GetPlaybackStatus();
@@ -141,6 +141,6 @@ TOptional<float> SAvaNavigationToolStatus::GetProgressPercent() const
 	return Progress;
 }
 
-} // UE::Sequencer::NavigationTool
+} // UE::SequenceNavigator
 
 #undef LOCTEXT_NAMESPACE
diff --git a/Engine/Plugins/Experimental/Avalanche/Source/AvalancheSequencer/Private/NavigationTool/DragDrop/NavigationToolAvaSequenceDropHandler.cpp b/Engine/Plugins/Experimental/Avalanche/Source/AvalancheSequencer/Private/NavigationTool/DragDrop/NavigationToolAvaSequenceDropHandler.cpp
--- a/Engine/Plugins/Experimental/Avalanche/Source/AvalancheSequencer/Private/NavigationTool/DragDrop/NavigationToolAvaSequenceDropHandler.cpp
+++ b/Engine/Plugins/Experimental/Avalanche/Source/AvalancheSequencer/Private/NavigationTool/DragDrop/NavigationToolAvaSequenceDropHandler.cpp
@@ -14,7 +14,7 @@
 
 #define LOCTEXT_NAMESPACE "NavigationToolAvaSequenceDropHandler"
 
-namespace UE::Sequencer::NavigationTool
+namespace UE::SequenceNavigator
 {
 
 FNavigationToolAvaSequenceDropHandler::FNavigationToolAvaSequenceDropHandler(const TWeakPtr<IAvaSequencer>& InWeakAvaSequencer)
@@ -57,12 +57,12 @@ TOptional<EItemDropZone> FNavigationToolAvaSequenceDropHandler::CanDrop(const EI
 
 		// Only allow moving valid sequences and sequences that aren't the target
 		UAvaSequence* const AvaSequenceToMove = AvaSequenceItemToMove->GetAvaSequence();
-		if (!IsValid(AvaSequenceToMove) || AvaSequenceToMove == TargetAvaSequence)
+		if (!AvaSequenceToMove || AvaSequenceToMove == TargetAvaSequence)
 		{
 			return TOptional<EItemDropZone>();
 		}
 
-		if (!bTargetIsRoot && IsValid(TargetAvaSequence))
+		if (!bTargetIsRoot && TargetAvaSequence)
 		{
 			if (TargetAvaSequence->GetChildren().Contains(AvaSequenceToMove))
 			{
@@ -165,7 +165,7 @@ void FNavigationToolAvaSequenceDropHandler::DuplicateItems(const TArray<FNavigat
 	}
 
 	UObject* const Outer = SequenceProvider->ToUObject();
-	if (!IsValid(Outer))
+	if (!Outer)
 	{
 		return;
 	}
@@ -182,7 +182,7 @@ void FNavigationToolAvaSequenceDropHandler::DuplicateItems(const TArray<FNavigat
 		}
 
 		UAvaSequence* const TemplateSequence = AvaSequenceItem->GetAvaSequence();
-		if (!IsValid(TemplateSequence))
+		if (!TemplateSequence)
 		{
 			continue;
 		}
@@ -202,14 +202,13 @@ void FNavigationToolAvaSequenceDropHandler::DuplicateItems(const TArray<FNavigat
 
 	for (const UAvaSequence* const TemplateSequence : SequencesToDuplicate)
 	{
-		UAvaSequence* const DupedSequence = DuplicateObject<UAvaSequence>(TemplateSequence, Outer);
-		if (IsValid(DupedSequence))
+		if (UAvaSequence* const DupedSequence = DuplicateObject<UAvaSequence>(TemplateSequence, Outer))
 		{
 			SequenceProvider->AddSequence(DupedSequence);
 		}
 	}
 }
 
-} // UE::Sequencer::NavigationTool
+} // UE::SequenceNavigator
 
 #undef LOCTEXT_NAMESPACE
diff --git a/Engine/Plugins/Experimental/Avalanche/Source/AvalancheSequencer/Private/NavigationTool/NavigationToolAvaSequence.cpp b/Engine/Plugins/Experimental/Avalanche/Source/AvalancheSequencer/Private/NavigationTool/NavigationToolAvaSequence.cpp
--- a/Engine/Plugins/Experimental/Avalanche/Source/AvalancheSequencer/Private/NavigationTool/NavigationToolAvaSequence.cpp
+++ b/Engine/Plugins/Experimental/Avalanche/Source/AvalancheSequencer/Private/NavigationTool/NavigationToolAvaSequence.cpp
@@ -15,7 +15,7 @@
 
 #define LOCTEXT_NAMESPACE "NavigationToolAvaSequence"
 
-namespace UE::Sequencer::NavigationTool
+namespace UE::SequenceNavigator
 {
 
 using namespace UE::Sequencer;
@@ -56,8 +56,7 @@ bool FNavigationToolAvaSequence::AddChild(const FNavigationToolAddItemParams& In
 	}
 
 	// Remove the sequence from any parent it's a child of
-	UAvaSequence* const AvaSequenceToMoveParent = AvaSequenceToAdd->GetParent();
-	if (IsValid(AvaSequenceToMoveParent))
+	if (UAvaSequence* const AvaSequenceToMoveParent = AvaSequenceToAdd->GetParent())
 	{
 		AvaSequenceToMoveParent->RemoveChild(AvaSequenceToAdd);
 	}
@@ -105,12 +104,10 @@ bool FNavigationToolAvaSequence::RemoveChild(const FNavigationToolRemoveItemPara
 
 	if (const FNavigationToolAvaSequence* const AvaSequenceItemToRemove = InRemoveItemParams.Item->CastTo<FNavigationToolAvaSequence>())
 	{
-		UAvaSequence* const AvaSequenceToRemove = AvaSequenceItemToRemove->GetAvaSequence();
-		if (IsValid(AvaSequenceToRemove))
+		if (UAvaSequence* const AvaSequenceToRemove = AvaSequenceItemToRemove->GetAvaSequence())
 		{
 			// Remove the sequence from any parent it's a child of
-			UAvaSequence* const AvaSequenceToRemoveParent = AvaSequenceToRemove->GetParent();
-			if (IsValid(AvaSequenceToRemoveParent))
+			if (UAvaSequence* const AvaSequenceToRemoveParent = AvaSequenceToRemove->GetParent())
 			{
 				AvaSequenceToRemoveParent->RemoveChild(AvaSequenceToRemove);
 			}
@@ -125,7 +122,7 @@ void FNavigationToolAvaSequence::FindChildren(TArray<FNavigationToolItemPtr>& Ou
 	FNavigationToolSequence::FindChildren(OutChildren, bRecursive);
 
 	UAvaSequence* const AvaSequence = GetAvaSequence();
-	if (!IsValid(AvaSequence))
+	if (!AvaSequence)
 	{
 		return;
 	}
@@ -173,13 +170,13 @@ void FNavigationToolAvaSequence::GetItemProxies(TArray<TSharedPtr<FNavigationToo
 
 bool FNavigationToolAvaSequence::CanRename() const
 {
-	return IsValid(GetAvaSequence());
+	return GetAvaSequence() != nullptr;
 }
 
 bool FNavigationToolAvaSequence::Rename(const FString& InName)
 {
 	UAvaSequence* const UnderlyingSequence = GetAvaSequence();
-	if (!IsValid(UnderlyingSequence))
+	if (!UnderlyingSequence)
 	{
 		return false;
 	}
@@ -203,7 +200,7 @@ bool FNavigationToolAvaSequence::Rename(const FString& InName)
 
 bool FNavigationToolAvaSequence::CanDelete() const
 {
-	return IsValid(GetAvaSequence());
+	return GetAvaSequence() != nullptr;
 }
 
 bool FNavigationToolAvaSequence::Delete()
@@ -215,7 +212,7 @@ bool FNavigationToolAvaSequence::Delete()
 	}
 
 	UAvaSequence* const AvaSequence = GetAvaSequence();
-	if (!IsValid(AvaSequence))
+	if (!AvaSequence)
 	{
 		return false;
 	}
@@ -233,8 +230,7 @@ bool FNavigationToolAvaSequence::Delete()
 
 FText FNavigationToolAvaSequence::GetDisplayName() const
 {
-	UAvaSequence* const AvaSequence = GetAvaSequence();
-	if (IsValid(AvaSequence))
+	if (const UAvaSequence* const AvaSequence = GetAvaSequence())
 	{
 		return AvaSequence->GetDisplayName();
 	}
@@ -255,7 +251,7 @@ FSlateColor FNavigationToolAvaSequence::GetIconColor() const
 	}
 
 	const UMovieSceneSequence* const RootSequence = Sequencer->GetRootMovieSceneSequence();
-	if (!IsValid(RootSequence))
+	if (!RootSequence)
 	{
 		return FStyleColors::Foreground;
 	}
@@ -276,7 +272,7 @@ void FNavigationToolAvaSequence::OnSelect()
 void FNavigationToolAvaSequence::OnDoubleClick()
 {
 	UAvaSequence* const AvaSequence = GetAvaSequence();
-	if (!IsValid(AvaSequence))
+	if (!AvaSequence)
 	{
 		return;
 	}
@@ -316,6 +312,6 @@ void FNavigationToolAvaSequence::SetColor(const TOptional<FColor>& InColor)
 	Tool.SetItemColor(SharedThis(this), InColor.Get(FColor()));
 }
 
-} // UE::Sequencer::NavigationTool
+} // UE::SequenceNavigator
 
 #undef LOCTEXT_NAMESPACE
diff --git a/Engine/Plugins/Experimental/SequenceNavigator/Source/SequenceNavigator/Private/Columns/NavigationToolColorColumn.cpp b/Engine/Plugins/Experimental/SequenceNavigator/Source/SequenceNavigator/Private/Columns/NavigationToolColorColumn.cpp
--- a/Engine/Plugins/Experimental/SequenceNavigator/Source/SequenceNavigator/Private/Columns/NavigationToolColorColumn.cpp
+++ b/Engine/Plugins/Experimental/SequenceNavigator/Source/SequenceNavigator/Private/Columns/NavigationToolColorColumn.cpp
@@ -7,7 +7,7 @@
 
 #define LOCTEXT_NAMESPACE "NavigationToolColorColumn"
 
-namespace UE::Sequencer::NavigationTool
+namespace UE::SequenceNavigator
 {
 
 FText FNavigationToolColorColumn::GetColumnDisplayNameText() const
@@ -41,6 +41,6 @@ TSharedRef<SWidget> FNavigationToolColorColumn::ConstructRowWidget(const FNaviga
 	return SNew(SNavigationToolColor, InItem, InView, InRow);
 }
 
-}
+} // namespace UE::SequenceNavigator
 
 #undef LOCTEXT_NAMESPACE
diff --git a/Engine/Plugins/Experimental/SequenceNavigator/Source/SequenceNavigator/Private/Columns/NavigationToolColumn.cpp b/Engine/Plugins/Experimental/SequenceNavigator/Source/SequenceNavigator/Private/Columns/NavigationToolColumn.cpp
--- a/Engine/Plugins/Experimental/SequenceNavigator/Source/SequenceNavigator/Private/Columns/NavigationToolColumn.cpp
+++ b/Engine/Plugins/Experimental/SequenceNavigator/Source/SequenceNavigator/Private/Columns/NavigationToolColumn.cpp
@@ -2,7 +2,7 @@
 
 #include "Columns/NavigationToolColumn.h"
 
-namespace UE::Sequencer::NavigationTool
+namespace UE::SequenceNavigator
 {
 
 }
diff --git a/Engine/Plugins/Experimental/SequenceNavigator/Source/SequenceNavigator/Private/Columns/NavigationToolColumnExtender.cpp b/Engine/Plugins/Experimental/SequenceNavigator/Source/SequenceNavigator/Private/Columns/NavigationToolColumnExtender.cpp
--- a/Engine/Plugins/Experimental/SequenceNavigator/Source/SequenceNavigator/Private/Columns/NavigationToolColumnExtender.cpp
+++ b/Engine/Plugins/Experimental/SequenceNavigator/Source/SequenceNavigator/Private/Columns/NavigationToolColumnExtender.cpp
@@ -3,7 +3,7 @@
 #include "Columns/NavigationToolColumnExtender.h"
 #include "Columns/NavigationToolColumn.h"
 
-namespace UE::Sequencer::NavigationTool
+namespace UE::SequenceNavigator
 {
 
 void FNavigationToolColumnExtender::AddColumn(const TSharedPtr<FNavigationToolColumn>& InColumn
@@ -40,4 +40,4 @@ int32 FNavigationToolColumnExtender::FindColumnIndex(const FName InColumnId) con
 	return INDEX_NONE;
 }
 
-}
+} // namespace UE::SequenceNavigator
diff --git a/Engine/Plugins/Experimental/SequenceNavigator/Source/SequenceNavigator/Private/Filters/Filters/NavigationToolBuiltInFilter.cpp b/Engine/Plugins/Experimental/SequenceNavigator/Source/SequenceNavigator/Private/Filters/Filters/NavigationToolBuiltInFilter.cpp
--- a/Engine/Plugins/Experimental/SequenceNavigator/Source/SequenceNavigator/Private/Filters/Filters/NavigationToolBuiltInFilter.cpp
+++ b/Engine/Plugins/Experimental/SequenceNavigator/Source/SequenceNavigator/Private/Filters/Filters/NavigationToolBuiltInFilter.cpp
@@ -4,7 +4,7 @@
 #include "Items/INavigationToolItem.h"
 #include "Items/NavigationToolSequence.h"
 
-namespace UE::Sequencer::NavigationTool
+namespace UE::SequenceNavigator
 {
 
 FNavigationToolBuiltInFilter::FNavigationToolBuiltInFilter(const FNavigationToolBuiltInFilterParams& InFilterParams)
@@ -58,7 +58,7 @@ bool FNavigationToolBuiltInFilter::PassesFilter(const FNavigationToolFilterType
 	auto ItemObjectPassesFilter = [this](const FNavigationToolItemPtr& InItem)
 		{
 			UObject* const ItemObject = InItem->GetItemObject();
-			return IsValid(ItemObject) && FilterParams.IsValidObjectClass(ItemObject->GetClass());
+			return ItemObject && FilterParams.IsValidObjectClass(ItemObject->GetClass());
 		};
 
 	auto ItemTypePassesFilter = [this](const FNavigationToolItemPtr& InItem)
@@ -120,4 +120,4 @@ const FNavigationToolBuiltInFilterParams& FNavigationToolBuiltInFilter::GetFilte
 	return FilterParams;
 }
 
-}
+} // namespace UE::SequenceNavigator
diff --git a/Engine/Plugins/Experimental/SequenceNavigator/Source/SequenceNavigator/Private/Filters/Filters/NavigationToolBuiltInFilter.h b/Engine/Plugins/Experimental/SequenceNavigator/Source/SequenceNavigator/Private/Filters/Filters/NavigationToolBuiltInFilter.h
--- a/Engine/Plugins/Experimental/SequenceNavigator/Source/SequenceNavigator/Private/Filters/Filters/NavigationToolBuiltInFilter.h
+++ b/Engine/Plugins/Experimental/SequenceNavigator/Source/SequenceNavigator/Private/Filters/Filters/NavigationToolBuiltInFilter.h
@@ -6,7 +6,7 @@
 #include "Filters/NavigationToolFilterData.h"
 #include "NavigationToolDefines.h"
 
-namespace UE::Sequencer::NavigationTool
+namespace UE::SequenceNavigator
 {
 
 class FNavigationToolBuiltInFilter : public FFilterBase<FNavigationToolItemPtr>
@@ -49,4 +49,4 @@ protected:
 	FNavigationToolBuiltInFilterParams FilterParams;
 };
 
-} // namespace UE::Sequencer::NavigationTool
+} // namespace UE::SequenceNavigator
diff --git a/Engine/Plugins/Experimental/SequenceNavigator/Source/SequenceNavigator/Private/Filters/Filters/NavigationToolBuiltInFilterParams.cpp b/Engine/Plugins/Experimental/SequenceNavigator/Source/SequenceNavigator/Private/Filters/Filters/NavigationToolBuiltInFilterParams.cpp
--- a/Engine/Plugins/Experimental/SequenceNavigator/Source/SequenceNavigator/Private/Filters/Filters/NavigationToolBuiltInFilterParams.cpp
+++ b/Engine/Plugins/Experimental/SequenceNavigator/Source/SequenceNavigator/Private/Filters/Filters/NavigationToolBuiltInFilterParams.cpp
@@ -17,7 +17,7 @@
 #define LOCTEXT_NAMESPACE "NavigationToolBuiltInFilterParams"
 
 FNavigationToolBuiltInFilterParams::FNavigationToolBuiltInFilterParams(const FName InFilterId
-	, const TSet<UE::Sequencer::NavigationTool::FNavigationToolItemTypeId>& InItemClasses
+	, const TSet<UE::SequenceNavigator::FNavigationToolItemTypeId>& InItemClasses
 	, const TArray<TSubclassOf<UObject>>& InFilterClasses
 	, const ENavigationToolFilterMode InFilterMode
 	, const FSlateBrush* const InIconBrush
@@ -68,7 +68,7 @@ FText FNavigationToolBuiltInFilterParams::GetTooltipText() const
 
 	for (TSubclassOf<UObject> FilterClass : ObjectClasses)
 	{
-		if (IsValid(FilterClass))
+		if (FilterClass)
 		{
 			if (!SupportedClasses.IsEmpty())
 			{
@@ -96,7 +96,7 @@ const FSlateBrush* FNavigationToolBuiltInFilterParams::GetIconBrush() const
 
 	for (const TSubclassOf<UObject>& FilterClass : ObjectClasses)
 	{
-		if (IsValid(FilterClass))
+		if (FilterClass)
 		{
 			const_cast<FNavigationToolBuiltInFilterParams*>(this)->IconBrush = FSlateIconFinder::FindIconForClass(FilterClass).GetIcon();
 			return IconBrush;
@@ -126,7 +126,7 @@ void FNavigationToolBuiltInFilterParams::SetOverrideIconColor(const FSlateColor&
 	OverrideIcon.TintColor = InNewIconColor;
 }
 
-bool FNavigationToolBuiltInFilterParams::IsValidItemClass(const UE::Sequencer::NavigationTool::FNavigationToolItemTypeId InClassTypeId) const
+bool FNavigationToolBuiltInFilterParams::IsValidItemClass(const UE::SequenceNavigator::FNavigationToolItemTypeId InClassTypeId) const
 {
 	if (InClassTypeId.IsValid())
 	{
@@ -161,7 +161,7 @@ void FNavigationToolBuiltInFilterParams::SetFilterText(const FText& InText)
 	FilterText = InText;
 }
 
-bool FNavigationToolBuiltInFilterParams::PassesFilterText(const UE::Sequencer::NavigationTool::FNavigationToolItemPtr& InItem) const
+bool FNavigationToolBuiltInFilterParams::PassesFilterText(const UE::SequenceNavigator::FNavigationToolItemPtr& InItem) const
 {
 	// Since it needs to be a const function we have to create the TextFilter here instead of being a variable
 	if (FilterText.IsEmptyOrWhitespace())
@@ -179,7 +179,7 @@ bool FNavigationToolBuiltInFilterParams::PassesFilterText(const UE::Sequencer::N
 FNavigationToolBuiltInFilterParams FNavigationToolBuiltInFilterParams::CreateSequenceFilter()
 {
 	return FNavigationToolBuiltInFilterParams(TEXT("Level Sequence")
-		, { UE::Sequencer::NavigationTool::FNavigationToolSequence::GetStaticTypeId() }
+		, { UE::SequenceNavigator::FNavigationToolSequence::GetStaticTypeId() }
 		, { ULevelSequence::StaticClass(), UMovieSceneSequence::StaticClass() }
 		, ENavigationToolFilterMode::MatchesType
 		, FSlateIconFinder::FindIconForClass(ALevelSequenceActor::StaticClass()).GetIcon()
@@ -191,7 +191,7 @@ FNavigationToolBuiltInFilterParams FNavigationToolBuiltInFilterParams::CreateSeq
 FNavigationToolBuiltInFilterParams FNavigationToolBuiltInFilterParams::CreateTrackFilter()
 {
 	return FNavigationToolBuiltInFilterParams(TEXT("Track")
-		, { UE::Sequencer::NavigationTool::FNavigationToolTrack::GetStaticTypeId() }
+		, { UE::SequenceNavigator::FNavigationToolTrack::GetStaticTypeId() }
 		, { UMovieSceneTrack::StaticClass() }
 		, ENavigationToolFilterMode::MatchesType
 		, FSlateIconFinder::FindIconForClass(UMovieSceneTrack::StaticClass()).GetIcon()
@@ -203,9 +203,9 @@ FNavigationToolBuiltInFilterParams FNavigationToolBuiltInFilterParams::CreateTra
 FNavigationToolBuiltInFilterParams FNavigationToolBuiltInFilterParams::CreateBindingFilter()
 {
 	return FNavigationToolBuiltInFilterParams(TEXT("Binding")
-		, { UE::Sequencer::NavigationTool::FNavigationToolBinding::GetStaticTypeId()
-			, UE::Sequencer::NavigationTool::FNavigationToolActor::GetStaticTypeId()
-			, UE::Sequencer::NavigationTool::FNavigationToolComponent::GetStaticTypeId() }
+		, { UE::SequenceNavigator::FNavigationToolBinding::GetStaticTypeId()
+			, UE::SequenceNavigator::FNavigationToolActor::GetStaticTypeId()
+			, UE::SequenceNavigator::FNavigationToolComponent::GetStaticTypeId() }
 		, { }
 		, ENavigationToolFilterMode::MatchesType
 		, nullptr
@@ -217,7 +217,7 @@ FNavigationToolBuiltInFilterParams FNavigationToolBuiltInFilterParams::CreateBin
 FNavigationToolBuiltInFilterParams FNavigationToolBuiltInFilterParams::CreateMarkerFilter()
 {
 	return FNavigationToolBuiltInFilterParams(TEXT("Marker")
-		, { UE::Sequencer::NavigationTool::FNavigationToolMarker::GetStaticTypeId() }
+		, { UE::SequenceNavigator::FNavigationToolMarker::GetStaticTypeId() }
 		, { }
 		, ENavigationToolFilterMode::MatchesType
 		, nullptr
diff --git a/Engine/Plugins/Experimental/SequenceNavigator/Source/SequenceNavigator/Private/Filters/Filters/NavigationToolFilter_Dirty.cpp b/Engine/Plugins/Experimental/SequenceNavigator/Source/SequenceNavigator/Private/Filters/Filters/NavigationToolFilter_Dirty.cpp
--- a/Engine/Plugins/Experimental/SequenceNavigator/Source/SequenceNavigator/Private/Filters/Filters/NavigationToolFilter_Dirty.cpp
+++ b/Engine/Plugins/Experimental/SequenceNavigator/Source/SequenceNavigator/Private/Filters/Filters/NavigationToolFilter_Dirty.cpp
@@ -8,7 +8,7 @@
 
 #define LOCTEXT_NAMESPACE "NavigationToolFilter_Dirty"
 
-namespace UE::Sequencer::NavigationTool
+namespace UE::SequenceNavigator
 {
 
 using namespace UE::Sequencer;
@@ -48,19 +48,19 @@ bool FNavigationToolFilter_Dirty::PassesFilter(FNavigationToolFilterType InItem)
 	if (const FNavigationToolSequence* const SequenceItem = InItem->CastTo<FNavigationToolSequence>())
 	{
 		const UMovieSceneSequence* const Sequence = SequenceItem->GetSequence();
-		return IsValid(Sequence) && IsObjectPackageDirty(Sequence);
+		return Sequence && IsObjectPackageDirty(Sequence);
 	}
 
 	if (const FNavigationToolTrack* const TrackItem = InItem->CastTo<FNavigationToolTrack>())
 	{
 		const UMovieSceneTrack* const Track = TrackItem->GetTrack();
-		return IsValid(Track) && IsObjectPackageDirty(Track);
+		return Track && IsObjectPackageDirty(Track);
 	}
 
 	if (const FNavigationToolBinding* const BindingItem = InItem->CastTo<FNavigationToolBinding>())
 	{
 		const UObject* const BoundObject = BindingItem->GetCachedBoundObject();
-		return IsValid(BoundObject) && IsObjectPackageDirty(BoundObject);
+		return BoundObject && IsObjectPackageDirty(BoundObject);
 	}
 
 	return false;
@@ -68,14 +68,14 @@ bool FNavigationToolFilter_Dirty::PassesFilter(FNavigationToolFilterType InItem)
 
 bool FNavigationToolFilter_Dirty::IsObjectPackageDirty(const UObject* const InObject)
 {
-	if (IsValid(InObject))
+	if (InObject)
 	{
 		const UPackage* const Package = InObject->GetPackage();
-		return IsValid(Package) && Package->IsDirty();
+		return Package && Package->IsDirty();
 	}
 	return false;
 }
 
-}
+} // namespace UE::SequenceNavigator
 
 #undef LOCTEXT_NAMESPACE
diff --git a/Engine/Plugins/Experimental/SequenceNavigator/Source/SequenceNavigator/Private/Filters/Filters/NavigationToolFilter_Marks.cpp b/Engine/Plugins/Experimental/SequenceNavigator/Source/SequenceNavigator/Private/Filters/Filters/NavigationToolFilter_Marks.cpp
--- a/Engine/Plugins/Experimental/SequenceNavigator/Source/SequenceNavigator/Private/Filters/Filters/NavigationToolFilter_Marks.cpp
+++ b/Engine/Plugins/Experimental/SequenceNavigator/Source/SequenceNavigator/Private/Filters/Filters/NavigationToolFilter_Marks.cpp
@@ -7,7 +7,7 @@
 
 #define LOCTEXT_NAMESPACE "NavigationToolFilter_Marks"
 
-namespace UE::Sequencer::NavigationTool
+namespace UE::SequenceNavigator
 {
 
 using namespace UE::Sequencer;
@@ -51,13 +51,13 @@ bool FNavigationToolFilter_Marks::PassesFilter(FNavigationToolFilterType InItem)
 	}
 
 	UMovieSceneSequence* const Sequence = SequenceItem->GetSequence();
-	if (!IsValid(Sequence))
+	if (!Sequence)
 	{
 		return false;
 	}
 
 	UMovieScene* const MovieScene = Sequence->GetMovieScene();
-	if (!IsValid(MovieScene))
+	if (!MovieScene)
 	{
 		return false;
 	}
@@ -65,6 +65,6 @@ bool FNavigationToolFilter_Marks::PassesFilter(FNavigationToolFilterType InItem)
 	return MovieScene->GetMarkedFrames().Num() > 0;
 }
 
-}
+} // namespace UE::SequenceNavigator
 
 #undef LOCTEXT_NAMESPACE
diff --git a/Engine/Plugins/Experimental/SequenceNavigator/Source/SequenceNavigator/Private/Filters/Filters/NavigationToolFilter_Playhead.cpp b/Engine/Plugins/Experimental/SequenceNavigator/Source/SequenceNavigator/Private/Filters/Filters/NavigationToolFilter_Playhead.cpp
--- a/Engine/Plugins/Experimental/SequenceNavigator/Source/SequenceNavigator/Private/Filters/Filters/NavigationToolFilter_Playhead.cpp
+++ b/Engine/Plugins/Experimental/SequenceNavigator/Source/SequenceNavigator/Private/Filters/Filters/NavigationToolFilter_Playhead.cpp
@@ -10,7 +10,7 @@
 
 #define LOCTEXT_NAMESPACE "NavigationToolFilter_Playhead"
 
-namespace UE::Sequencer::NavigationTool
+namespace UE::SequenceNavigator
 {
 
 using namespace UE::Sequencer;
@@ -123,13 +123,13 @@ void FNavigationToolFilter_Playhead::OnEndScrubbingEvent()
 FTimerManager* FNavigationToolFilter_Playhead::GetTimerManager() const
 {
 	UObject* const ContextObject = FilterInterface.GetSequencer().GetPlaybackContext();
-	if (!IsValid(ContextObject))
+	if (!ContextObject)
 	{
 		return nullptr;
 	}
 
 	UWorld* const World = ContextObject->GetWorld();
-	if (!IsValid(World))
+	if (!World)
 	{
 		return nullptr;
 	}
@@ -161,6 +161,6 @@ void FNavigationToolFilter_Playhead::StopRefreshTimer()
 	}
 }
 
-}
+} // namespace UE::SequenceNavigator
 
 #undef LOCTEXT_NAMESPACE
diff --git a/Engine/Plugins/Experimental/SequenceNavigator/Source/SequenceNavigator/Private/Filters/Filters/NavigationToolFilter_Text.cpp b/Engine/Plugins/Experimental/SequenceNavigator/Source/SequenceNavigator/Private/Filters/Filters/NavigationToolFilter_Text.cpp
--- a/Engine/Plugins/Experimental/SequenceNavigator/Source/SequenceNavigator/Private/Filters/Filters/NavigationToolFilter_Text.cpp
+++ b/Engine/Plugins/Experimental/SequenceNavigator/Source/SequenceNavigator/Private/Filters/Filters/NavigationToolFilter_Text.cpp
@@ -12,7 +12,7 @@
 
 #define LOCTEXT_NAMESPACE "NavigationToolFilter_Text"
 
-namespace UE::Sequencer::NavigationTool
+namespace UE::SequenceNavigator
 {
 
 using namespace UE::Sequencer;
@@ -34,7 +34,7 @@ FNavigationToolFilter_Text::FNavigationToolFilter_Text(INavigationToolFilterBar&
 	for (TObjectIterator<UNavigationToolFilterTextExpressionExtension> ExtensionIt(RF_NoFlags); ExtensionIt; ++ExtensionIt)
 	{
 		const UNavigationToolFilterTextExpressionExtension* const PotentialExtension = *ExtensionIt;
-		if (IsValid(PotentialExtension)
+		if (PotentialExtension
 			&& PotentialExtension->HasAnyFlags(RF_ClassDefaultObject)
 			&& !PotentialExtension->GetClass()->HasAnyClassFlags(CLASS_Deprecated | CLASS_Abstract))
 		{
@@ -211,6 +211,6 @@ bool FNavigationToolFilter_Text::IsTokenValueValid(const FExpressionToken& InTok
 	return true;
 }
 
-}
+} // namespace UE::SequenceNavigator
 
 #undef LOCTEXT_NAMESPACE
diff --git a/Engine/Plugins/Experimental/SequenceNavigator/Source/SequenceNavigator/Private/Filters/Filters/NavigationToolFilter_Unbound.cpp b/Engine/Plugins/Experimental/SequenceNavigator/Source/SequenceNavigator/Private/Filters/Filters/NavigationToolFilter_Unbound.cpp
--- a/Engine/Plugins/Experimental/SequenceNavigator/Source/SequenceNavigator/Private/Filters/Filters/NavigationToolFilter_Unbound.cpp
+++ b/Engine/Plugins/Experimental/SequenceNavigator/Source/SequenceNavigator/Private/Filters/Filters/NavigationToolFilter_Unbound.cpp
@@ -6,7 +6,7 @@
 
 #define LOCTEXT_NAMESPACE "NavigationToolFilter_Unbound"
 
-namespace UE::Sequencer::NavigationTool
+namespace UE::SequenceNavigator
 {
 
 using namespace UE::Sequencer;
@@ -50,9 +50,9 @@ bool FNavigationToolFilter_Unbound::PassesFilter(FNavigationToolFilterType InIte
 	}
 
 	const UObject* const BoundObject = BindingItem->GetCachedBoundObject();
-	return !IsValid(BoundObject);
+	return !BoundObject;
 }
 
-}
+} // namespace UE::SequenceNavigator
 
 #undef LOCTEXT_NAMESPACE
diff --git a/Engine/Plugins/Experimental/SequenceNavigator/Source/SequenceNavigator/Private/Filters/NavigationToolFilterBar.cpp b/Engine/Plugins/Experimental/SequenceNavigator/Source/SequenceNavigator/Private/Filters/NavigationToolFilterBar.cpp
--- a/Engine/Plugins/Experimental/SequenceNavigator/Source/SequenceNavigator/Private/Filters/NavigationToolFilterBar.cpp
+++ b/Engine/Plugins/Experimental/SequenceNavigator/Source/SequenceNavigator/Private/Filters/NavigationToolFilterBar.cpp
@@ -29,7 +29,7 @@
 
 #define LOCTEXT_NAMESPACE "NavigationToolFilterBar"
 
-namespace UE::Sequencer::NavigationTool
+namespace UE::SequenceNavigator
 {
 
 using namespace UE::Sequencer;
@@ -113,7 +113,7 @@ void FNavigationToolFilterBar::CreateDefaultFilters()
 	for (TObjectIterator<UNavigationToolFilterExtension> ExtensionIt(RF_NoFlags); ExtensionIt; ++ExtensionIt)
 	{
 		const UNavigationToolFilterExtension* const PotentialExtension = *ExtensionIt;
-		if (IsValid(PotentialExtension)
+		if (PotentialExtension
 			&& PotentialExtension->HasAnyFlags(RF_ClassDefaultObject)
 			&& !PotentialExtension->GetClass()->HasAnyClassFlags(CLASS_Deprecated | CLASS_Abstract))
 		{
@@ -188,7 +188,7 @@ void FNavigationToolFilterBar::BindCommands(const TSharedPtr<FUICommandList>& In
 	}
 
 	UMovieSceneSequence* const FocusedSequence = Sequencer->GetFocusedMovieSceneSequence();
-	if (!IsValid(FocusedSequence))
+	if (!FocusedSequence)
 	{
 		return;
 	}
@@ -206,7 +206,7 @@ void FNavigationToolFilterBar::BindCommands(const TSharedPtr<FUICommandList>& In
 void FNavigationToolFilterBar::CreateCustomTextFiltersFromConfig()
 {
 	UNavigationToolSettings* const Settings = GetMutableDefault<UNavigationToolSettings>();
-	if (!ensure(IsValid(Settings)))
+	if (!Settings)
 	{
 		return;
 	}
@@ -244,7 +244,7 @@ FName FNavigationToolFilterBar::GetIdentifier() const
 	}
 
 	USequencerSettings* const SequencerSettings = Sequencer->GetSequencerSettings();
-	if (!IsValid(SequencerSettings))
+	if (!SequencerSettings)
 	{
 		return DefaultIdentifier;
 	}
@@ -492,7 +492,7 @@ bool FNavigationToolFilterBar::HasAnyFilterActive(const bool bCheckTextFilter
 bool FNavigationToolFilterBar::IsFilterActive(const TSharedRef<FNavigationToolFilter> InFilter) const
 {
 	UNavigationToolSettings* const Settings = GetMutableDefault<UNavigationToolSettings>();
-	if (!ensure(IsValid(Settings)))
+	if (!Settings)
 	{
 		return false;
 	}
@@ -1807,6 +1807,6 @@ void FNavigationToolFilterBar::CreateNewTextFilter()
 	CreateWindow_AddCustomTextFilter(DefaultNewCustomTextFilterData(NewFilterString));
 }
 
-}
+} // namespace UE::SequenceNavigator
 
 #undef LOCTEXT_NAMESPACE
diff --git a/Engine/Plugins/Experimental/SequenceNavigator/Source/SequenceNavigator/Private/Filters/TextExpressions/NavigationToolFilterTextExpressionContext.cpp b/Engine/Plugins/Experimental/SequenceNavigator/Source/SequenceNavigator/Private/Filters/TextExpressions/NavigationToolFilterTextExpressionContext.cpp
--- a/Engine/Plugins/Experimental/SequenceNavigator/Source/SequenceNavigator/Private/Filters/TextExpressions/NavigationToolFilterTextExpressionContext.cpp
+++ b/Engine/Plugins/Experimental/SequenceNavigator/Source/SequenceNavigator/Private/Filters/TextExpressions/NavigationToolFilterTextExpressionContext.cpp
@@ -4,7 +4,7 @@
 #include "Filters/INavigationToolFilterBar.h"
 #include "Items/INavigationToolItem.h"
 
-namespace UE::Sequencer::NavigationTool
+namespace UE::SequenceNavigator
 {
 
 using namespace UE::Sequencer;
@@ -61,7 +61,7 @@ UMovieSceneSequence* FNavigationToolFilterTextExpressionContext::GetFocusedMovie
 UMovieScene* FNavigationToolFilterTextExpressionContext::GetFocusedGetMovieScene() const
 {
 	const UMovieSceneSequence* const FocusedMovieSceneSequence = GetFocusedMovieSceneSequence();
-	return IsValid(FocusedMovieSceneSequence) ? FocusedMovieSceneSequence->GetMovieScene() : nullptr;
+	return FocusedMovieSceneSequence ? FocusedMovieSceneSequence->GetMovieScene() : nullptr;
 }
 
-}
+} // namespace UE::SequenceNavigator
diff --git a/Engine/Plugins/Experimental/SequenceNavigator/Source/SequenceNavigator/Private/Items/NavigationToolActor.cpp b/Engine/Plugins/Experimental/SequenceNavigator/Source/SequenceNavigator/Private/Items/NavigationToolActor.cpp
--- a/Engine/Plugins/Experimental/SequenceNavigator/Source/SequenceNavigator/Private/Items/NavigationToolActor.cpp
+++ b/Engine/Plugins/Experimental/SequenceNavigator/Source/SequenceNavigator/Private/Items/NavigationToolActor.cpp
@@ -16,7 +16,7 @@
 
 #define LOCTEXT_NAMESPACE "NavigationToolActor"
 
-namespace UE::Sequencer::NavigationTool
+namespace UE::SequenceNavigator
 {
 
 FNavigationToolActor::FNavigationToolActor(INavigationTool& InTool
@@ -32,7 +32,7 @@ void FNavigationToolActor::FindChildren(TArray<FNavigationToolItemPtr>& OutChild
 	Super::FindChildren(OutChildren, bInRecursive);
 
 	const AActor* const ThisActor = GetActor();
-	if (!IsValid(ThisActor))
+	if (!ThisActor)
 	{
 		return;
 	}
@@ -44,13 +44,13 @@ void FNavigationToolActor::FindChildren(TArray<FNavigationToolItemPtr>& OutChild
 	}
 
 	UMovieSceneSequence* const Sequence = ParentSequenceItem->GetSequence();
-	if (!IsValid(Sequence))
+	if (!Sequence)
 	{
 		return;
 	}
 
 	UMovieScene* const MovieScene = Sequence->GetMovieScene();
-	if (!IsValid(MovieScene))
+	if (!MovieScene)
 	{
 		return;
 	}
@@ -74,7 +74,7 @@ void FNavigationToolActor::FindChildren(TArray<FNavigationToolItemPtr>& OutChild
 			const TWeakObjectPtr<> CurrentBoundObject = WeakBoundObjects.IsEmpty() ? nullptr : WeakBoundObjects[0];
 
 			const UActorComponent* const CurrentBoundActorComponent = Cast<UActorComponent>(CurrentBoundObject);
-			if (IsValid(CurrentBoundActorComponent) && CurrentBoundActorComponent->GetOwner() == ThisActor)
+			if (CurrentBoundActorComponent && CurrentBoundActorComponent->GetOwner() == ThisActor)
 			{
 				const FNavigationToolItemPtr NewItem = Tool.FindOrAdd<FNavigationToolComponent>(GetProvider()
 					, SharedThisRef, ParentSequenceItem, CurrentBinding);
@@ -93,7 +93,7 @@ bool FNavigationToolActor::IsAllowedInTool() const
 	FNavigationTool& ToolPrivate = static_cast<FNavigationTool&>(Tool);
 
 	const AActor* const UnderlyingActor = GetActor();
-	if (!IsValid(UnderlyingActor))
+	if (!UnderlyingActor)
 	{
 		// Always allow unbound binding items
 		return true;
@@ -111,8 +111,7 @@ ENavigationToolItemViewMode FNavigationToolActor::GetSupportedViewModes(const IN
 
 TArray<FName> FNavigationToolActor::GetTags() const
 {
-	const AActor* const UnderlyingActor = GetActor();
-	if (IsValid(UnderlyingActor))
+	if (const AActor* const UnderlyingActor = GetActor())
 	{
 		return UnderlyingActor->Tags;
 	}
@@ -121,8 +120,7 @@ TArray<FName> FNavigationToolActor::GetTags() const
 
 bool FNavigationToolActor::GetVisibility() const
 {
-	const AActor* const UnderlyingActor = GetActor();
-	if (IsValid(UnderlyingActor))
+	if (const AActor* const UnderlyingActor = GetActor())
 	{
 		return !UnderlyingActor->IsTemporarilyHiddenInEditor(true);
 	}
@@ -131,8 +129,7 @@ bool FNavigationToolActor::GetVisibility() const
 
 void FNavigationToolActor::OnVisibilityChanged(const bool bInNewVisibility)
 {
-	AActor* const UnderlyingActor = GetActor();
-	if (IsValid(UnderlyingActor))
+	if (AActor* const UnderlyingActor = GetActor())
 	{
 		UnderlyingActor->SetIsTemporarilyHiddenInEditor(!bInNewVisibility);
 	}
@@ -140,13 +137,13 @@ void FNavigationToolActor::OnVisibilityChanged(const bool bInNewVisibility)
 
 bool FNavigationToolActor::CanRename() const
 {
-	return Super::CanRename() && IsValid(GetActor());
+	return Super::CanRename() && GetActor() != nullptr;
 }
 
 bool FNavigationToolActor::Rename(const FString& InName)
 {
 	AActor* const UnderlyingActor = GetActor();
-	const bool bIsActorLabelEditable = IsValid(UnderlyingActor) && UnderlyingActor->IsActorLabelEditable();
+	const bool bIsActorLabelEditable = UnderlyingActor && UnderlyingActor->IsActorLabelEditable();
 
 	bool bRenamed = false;
 	
@@ -169,6 +166,6 @@ AActor* FNavigationToolActor::GetActor() const
 	return WeakBoundObject.IsValid() ? Cast<AActor>(GetCachedBoundObject()) : nullptr;
 }
 
-}
+} // namespace UE::SequenceNavigator
 
 #undef LOCTEXT_NAMESPACE
diff --git a/Engine/Plugins/Experimental/SequenceNavigator/Source/SequenceNavigator/Private/Items/NavigationToolBinding.cpp b/Engine/Plugins/Experimental/SequenceNavigator/Source/SequenceNavigator/Private/Items/NavigationToolBinding.cpp
--- a/Engine/Plugins/Experimental/SequenceNavigator/Source/SequenceNavigator/Private/Items/NavigationToolBinding.cpp
+++ b/Engine/Plugins/Experimental/SequenceNavigator/Source/SequenceNavigator/Private/Items/NavigationToolBinding.cpp
@@ -19,7 +19,7 @@
 
 #define LOCTEXT_NAMESPACE "NavigationToolBinding"
 
-namespace UE::Sequencer::NavigationTool
+namespace UE::SequenceNavigator
 {
 
 using namespace UE::Sequencer;
@@ -66,8 +66,7 @@ FText FNavigationToolBinding::GetDisplayName() const
 
 FText FNavigationToolBinding::GetClassName() const
 {
-	const UClass* const BoundObjectClass = GetBoundObjectClass();
-	if (IsValid(BoundObjectClass))
+	if (const UClass* const BoundObjectClass = GetBoundObjectClass())
 	{
 		return BoundObjectClass->GetDisplayNameText();
 	}
@@ -81,8 +80,7 @@ FSlateIcon FNavigationToolBinding::GetIcon() const
 
 FText FNavigationToolBinding::GetIconTooltipText() const
 {
-	const UObject* const UnderlyingObject = GetCachedBoundObject();
-	if (IsValid(UnderlyingObject))
+	if (const UObject* const UnderlyingObject = GetCachedBoundObject())
 	{
 		return FText::Format(LOCTEXT("BoundObjectToolTip", " Class: %s (BindingID: %s)")
 			, UnderlyingObject->GetClass()->GetDisplayNameText(), FText::FromString(LexToString(Binding.GetObjectGuid())));
@@ -100,7 +98,7 @@ FLinearColor FNavigationToolBinding::GetItemBrushColor() const
 	}
 
 	UMovieSceneSequence* const Sequence = GetSequence();
-	if (!IsValid(Sequence))
+	if (!Sequence)
 	{
 		return FLinearColor::Red;
 	}
@@ -381,6 +379,6 @@ TViewModelPtr<FObjectBindingModel> FNavigationToolBinding::GetViewModel() const
 	return StorageExtension->FindModelForObjectBinding(ObjectGuid);
 }
 
-}
+} // namespace UE::SequenceNavigator
 
 #undef LOCTEXT_NAMESPACE
diff --git a/Engine/Plugins/Experimental/SequenceNavigator/Source/SequenceNavigator/Private/Items/NavigationToolComponent.cpp b/Engine/Plugins/Experimental/SequenceNavigator/Source/SequenceNavigator/Private/Items/NavigationToolComponent.cpp
--- a/Engine/Plugins/Experimental/SequenceNavigator/Source/SequenceNavigator/Private/Items/NavigationToolComponent.cpp
+++ b/Engine/Plugins/Experimental/SequenceNavigator/Source/SequenceNavigator/Private/Items/NavigationToolComponent.cpp
@@ -11,7 +11,7 @@
 
 #define LOCTEXT_NAMESPACE "NavigationToolComponent"
 
-namespace UE::Sequencer::NavigationTool
+namespace UE::SequenceNavigator
 {
 
 FNavigationToolComponent::FNavigationToolComponent(INavigationTool& InTool
@@ -45,7 +45,7 @@ bool FNavigationToolComponent::IsAllowedInTool() const
 	FNavigationTool& ToolPrivate = static_cast<FNavigationTool&>(Tool);
 
 	const UActorComponent* const UnderlyingComponent = GetComponent();
-	if (!IsValid(UnderlyingComponent))
+	if (!UnderlyingComponent)
 	{
 		// Always allow unbound binding items
 		return true;
@@ -74,8 +74,7 @@ FLinearColor FNavigationToolComponent::GetItemBrushColor() const
 
 TArray<FName> FNavigationToolComponent::GetTags() const
 {
-	UActorComponent* const UnderlyingComponent = GetComponent();
-	if (IsValid(UnderlyingComponent))
+	if (const UActorComponent* const UnderlyingComponent = GetComponent())
 	{
 		return UnderlyingComponent->ComponentTags;
 	}
@@ -84,8 +83,7 @@ TArray<FName> FNavigationToolComponent::GetTags() const
 
 bool FNavigationToolComponent::GetVisibility() const
 {
-	const USceneComponent* const UnderlyingComponent = Cast<USceneComponent>(GetComponent());
-	if (IsValid(UnderlyingComponent))
+	if (const USceneComponent* const UnderlyingComponent = Cast<USceneComponent>(GetComponent()))
 	{
 		return UnderlyingComponent->IsVisibleInEditor();
 	}
@@ -94,8 +92,7 @@ bool FNavigationToolComponent::GetVisibility() const
 
 void FNavigationToolComponent::OnVisibilityChanged(const bool bInNewVisibility)
 {
-	USceneComponent* const UnderlyingComponent = Cast<USceneComponent>(GetComponent());
-	if (IsValid(UnderlyingComponent))
+	if (USceneComponent* const UnderlyingComponent = Cast<USceneComponent>(GetComponent()))
 	{
 		UnderlyingComponent->SetVisibility(bInNewVisibility);
 	}
@@ -103,13 +100,13 @@ void FNavigationToolComponent::OnVisibilityChanged(const bool bInNewVisibility)
 
 bool FNavigationToolComponent::CanRename() const
 {
-	return Super::CanRename() && IsValid(GetComponent());
+	return Super::CanRename() && GetComponent() != nullptr;
 }
 
 bool FNavigationToolComponent::Rename(const FString& InName)
 {
 	UActorComponent* const UnderlyingComponent = GetComponent();
-	if (!IsValid(UnderlyingComponent))
+	if (!UnderlyingComponent)
 	{
 		return false;
 	}
@@ -136,6 +133,6 @@ UActorComponent* FNavigationToolComponent::GetComponent() const
 	return WeakBoundObject.IsValid() ? Cast<UActorComponent>(GetCachedBoundObject()) : nullptr;
 }
 
-}
+} // namespace UE::SequenceNavigator
 
 #undef LOCTEXT_NAMESPACE
diff --git a/Engine/Plugins/Experimental/SequenceNavigator/Source/SequenceNavigator/Private/Items/NavigationToolItemId.cpp b/Engine/Plugins/Experimental/SequenceNavigator/Source/SequenceNavigator/Private/Items/NavigationToolItemId.cpp
--- a/Engine/Plugins/Experimental/SequenceNavigator/Source/SequenceNavigator/Private/Items/NavigationToolItemId.cpp
+++ b/Engine/Plugins/Experimental/SequenceNavigator/Source/SequenceNavigator/Private/Items/NavigationToolItemId.cpp
@@ -11,7 +11,7 @@
 #include "MVVM/ViewModels/SequenceModel.h"
 #include "UObject/SoftObjectPath.h"
 
-namespace UE::Sequencer::NavigationTool
+namespace UE::SequenceNavigator
 {
 
 FNavigationToolItemId FNavigationToolItemId::RootId = FNavigationToolItemId(TEXT("Root"));
@@ -58,8 +58,10 @@ FNavigationToolItemId::FNavigationToolItemId(const FNavigationToolItemPtr& InPar
 	CalculateTypeHash();
 }
 
-FNavigationToolItemId::FNavigationToolItemId(const FViewModelPtr& InViewModel)
+FNavigationToolItemId::FNavigationToolItemId(const UE::Sequencer::FViewModelPtr& InViewModel)
 {
+	using namespace UE::Sequencer;
+
 	check(InViewModel.IsValid());
 
 	const TViewModelPtr<FSectionModel> SectionModel = InViewModel->FindAncestorOfType<FSectionModel>(true);
@@ -83,8 +85,7 @@ FNavigationToolItemId::FNavigationToolItemId(const FViewModelPtr& InViewModel)
 
 	if (const TViewModelPtr<ITrackExtension> TrackExtension = InViewModel.ImplicitCast())
 	{
-		const UMovieSceneTrack* const Track = TrackExtension->GetTrack();
-		if (IsValid(Track))
+		if (const UMovieSceneTrack* const Track = TrackExtension->GetTrack())
 		{
 			ReferenceId = GetObjectPath(Track);
 		}
@@ -114,7 +115,7 @@ void FNavigationToolItemId::ConstructId(const FNavigationToolItemPtr& InParentIt
 
 	AddSeparatedSegment(Id, GetObjectPath(InSequence));
 
-	if (IsValid(InSection))
+	if (InSection)
 	{
 		AddSeparatedSegment(Id, GetObjectPath(InSection));
 		AddSeparatedSegment(Id, FString::FromInt(InSectionIndex));
@@ -187,4 +188,4 @@ void FNavigationToolItemId::CalculateTypeHash()
 	bHasCachedHash = true;
 }
 
-}
+} // namespace UE::SequenceNavigator
diff --git a/Engine/Plugins/Experimental/SequenceNavigator/Source/SequenceNavigator/Private/Items/NavigationToolItemUtils.cpp b/Engine/Plugins/Experimental/SequenceNavigator/Source/SequenceNavigator/Private/Items/NavigationToolItemUtils.cpp
--- a/Engine/Plugins/Experimental/SequenceNavigator/Source/SequenceNavigator/Private/Items/NavigationToolItemUtils.cpp
+++ b/Engine/Plugins/Experimental/SequenceNavigator/Source/SequenceNavigator/Private/Items/NavigationToolItemUtils.cpp
@@ -15,7 +15,7 @@
 
 #define LOCTEXT_NAMESPACE "NavigationToolItemUtils"
 
-namespace UE::Sequencer::NavigationTool
+namespace UE::SequenceNavigator
 {
 
 bool CompareToolItemOrder(const FNavigationToolItemPtr& InA, const FNavigationToolItemPtr& InB)
@@ -110,7 +110,7 @@ void RemoveSequenceDisplayNameParentPrefix(FText& InOutDisplayName
 	, const TSharedRef<FNavigationToolSequence>& InSequenceItem)
 {
 	UMovieSceneSequence* const Sequence = InSequenceItem->GetSequence();
-	if (!IsValid(Sequence))
+	if (!Sequence)
 	{
 		return;
 	}
@@ -129,15 +129,14 @@ void RemoveSequenceDisplayNameParentPrefix(FText& InOutDisplayName
 	}
 
 	UMovieSceneSequence* const ParentSequence = ParentSequenceItem->GetSequence();
-	if (!IsValid(ParentSequence))
+	if (!ParentSequence)
 	{
 		return;
 	}
 
 	const UMovieSceneCinematicShotSection* const ShotSection = Cast<UMovieSceneCinematicShotSection>(InSequenceItem->GetSubSection());
 
-	FText NewSequenceName = IsValid(ShotSection)
-		? FText::FromString(ShotSection->GetShotDisplayName()) : Sequence->GetDisplayName();
+	FText NewSequenceName = ShotSection ? FText::FromString(ShotSection->GetShotDisplayName()) : Sequence->GetDisplayName();
 	if (NewSequenceName.IsEmpty())
 	{
 		return;
@@ -176,6 +175,6 @@ void AppendSequenceDisplayNameDirtyStatus(FText& InOutDisplayName
 	}
 }
 
-}
+} // namespace UE::SequenceNavigator
 
 #undef LOCTEXT_NAMESPACE
diff --git a/Engine/Plugins/Experimental/SequenceNavigator/Source/SequenceNavigator/Private/Items/NavigationToolMarker.cpp b/Engine/Plugins/Experimental/SequenceNavigator/Source/SequenceNavigator/Private/Items/NavigationToolMarker.cpp
--- a/Engine/Plugins/Experimental/SequenceNavigator/Source/SequenceNavigator/Private/Items/NavigationToolMarker.cpp
+++ b/Engine/Plugins/Experimental/SequenceNavigator/Source/SequenceNavigator/Private/Items/NavigationToolMarker.cpp
@@ -15,7 +15,7 @@
 
 #define LOCTEXT_NAMESPACE "NavigationToolMarker"
 
-namespace UE::Sequencer::NavigationTool
+namespace UE::SequenceNavigator
 {
 
 FNavigationToolMarker::FNavigationToolMarker(INavigationTool& InTool
@@ -76,7 +76,7 @@ FSlateColor FNavigationToolMarker::GetIconColor() const
 	}
 
 	USequencerSettings* const SequencerSettings = Sequencer->GetSequencerSettings();
-	if (!IsValid(SequencerSettings))
+	if (!SequencerSettings)
 	{
 		return FStyleColors::Foreground;
 	}
@@ -111,7 +111,7 @@ void FNavigationToolMarker::OnSelect()
 void FNavigationToolMarker::OnDoubleClick()
 {
 	UMovieSceneSequence* const Sequence = GetParentSequence();
-	if (!IsValid(Sequence))
+	if (!Sequence)
 	{
 		return;
 	}
@@ -325,6 +325,6 @@ UMovieScene* FNavigationToolMarker::GetParentMovieScene() const
 	return nullptr;
 }
 
-}
+} // namespace UE::SequenceNavigator
 
 #undef LOCTEXT_NAMESPACE
diff --git a/Engine/Plugins/Experimental/SequenceNavigator/Source/SequenceNavigator/Private/Items/NavigationToolSequence.cpp b/Engine/Plugins/Experimental/SequenceNavigator/Source/SequenceNavigator/Private/Items/NavigationToolSequence.cpp
--- a/Engine/Plugins/Experimental/SequenceNavigator/Source/SequenceNavigator/Private/Items/NavigationToolSequence.cpp
+++ b/Engine/Plugins/Experimental/SequenceNavigator/Source/SequenceNavigator/Private/Items/NavigationToolSequence.cpp
@@ -26,7 +26,9 @@
 
 #define LOCTEXT_NAMESPACE "NavigationToolSequence"
 
-namespace UE::Sequencer::NavigationTool
+using namespace UE::Sequencer;
+
+namespace UE::SequenceNavigator
 {
 
 FNavigationToolSequence::FNavigationToolSequence(INavigationTool& InTool
@@ -62,13 +64,13 @@ void FNavigationToolSequence::FindChildren(TArray<FNavigationToolItemPtr>& OutCh
 	Super::FindChildren(OutChildren, bInRecursive);
 
 	UMovieSceneSequence* const Sequence = GetSequence();
-	if (!IsValid(Sequence))
+	if (!Sequence)
 	{
 		return;
 	}
 
 	UMovieScene* const MovieScene = Sequence->GetMovieScene();
-	if (!IsValid(MovieScene))
+	if (!MovieScene)
 	{
 		return;
 	}
@@ -80,13 +82,12 @@ void FNavigationToolSequence::FindChildren(TArray<FNavigationToolItemPtr>& OutCh
 	{
 		FNavigationToolItemPtr NewItem;
 
-		UMovieSceneSubTrack* const SubTrack = Cast<UMovieSceneSubTrack>(Track);
-		if (IsValid(SubTrack))
+		if (UMovieSceneSubTrack* const SubTrack = Cast<UMovieSceneSubTrack>(Track))
 		{
 			NewItem = Tool.FindOrAdd<FNavigationToolSubTrack>(GetProvider()
 				, SharedThisRef, SubTrack, Sequence, WeakSubSection, SubSectionIndex);
 		}
-		else if (IsValid(Track))
+		else if (Track)
 		{
 			NewItem = Tool.FindOrAdd<FNavigationToolTrack>(GetProvider()
 				, SharedThisRef, Track, Sequence, WeakSubSection, SubSectionIndex);
@@ -162,15 +163,13 @@ FText FNavigationToolSequence::GetDisplayName() const
 {
 	FText NewDisplayName = FText::GetEmpty();
 
-	const UMovieSceneCinematicShotSection* const ShotSection = Cast<UMovieSceneCinematicShotSection>(WeakSubSection.Get());
-	if (IsValid(ShotSection))
+	if (const UMovieSceneCinematicShotSection* const ShotSection = Cast<UMovieSceneCinematicShotSection>(WeakSubSection.Get()))
 	{
 		NewDisplayName = FText::FromString(ShotSection->GetShotDisplayName());
 	}
 	else
 	{
-		const UMovieSceneSequence* const Sequence = GetSequence();
-		if (IsValid(Sequence))
+		if (const UMovieSceneSequence* const Sequence = GetSequence())
 		{
 			NewDisplayName = Sequence->GetDisplayName();
 		}
@@ -190,8 +189,7 @@ FText FNavigationToolSequence::GetDisplayName() const
 		RemoveSequenceDisplayNameParentPrefix(NewDisplayName, ThisSequenceRef);
 	}
 
-	UMovieSceneSequence* const Sequence = ThisSequenceRef->GetSequence();
-	if (IsValid(Sequence))
+	if (UMovieSceneSequence* const Sequence = ThisSequenceRef->GetSequence())
 	{
 		AppendSequenceDisplayNameDirtyStatus(NewDisplayName, *Sequence);
 	}
@@ -201,8 +199,7 @@ FText FNavigationToolSequence::GetDisplayName() const
 
 FText FNavigationToolSequence::GetClassName() const
 {
-	UMovieSceneSequence* const Sequence = GetSequence();
-	if (IsValid(Sequence))
+	if (const UMovieSceneSequence* const Sequence = GetSequence())
 	{
 		return FText::FromString(Sequence->GetClass()->GetName());
 	}
@@ -240,7 +237,7 @@ void FNavigationToolSequence::OnSelect()
 void FNavigationToolSequence::OnDoubleClick()
 {
 	UMovieSceneSequence* const Sequence = GetSequence();
-	if (!IsValid(Sequence))
+	if (!Sequence)
 	{
 		return;
 	}
@@ -252,8 +249,7 @@ void FNavigationToolSequence::OnDoubleClick()
 		// Instead of focusing the sequence in the Sequencer, focus the parent sequence and select sequence
 		if (const FNavigationToolSequence* const ParentSequenceItem = FindAncestorOfType<FNavigationToolSequence>())
 		{
-			UMovieSceneSequence* const ParentSequence = ParentSequenceItem->GetSequence();
-			if (IsValid(ParentSequence))
+			if (UMovieSceneSequence* const ParentSequence = ParentSequenceItem->GetSequence())
 			{
 				FocusSequence(Tool, *ParentSequence, *this);
 			}
@@ -286,13 +282,13 @@ bool FNavigationToolSequence::CanRename() const
 	// Disable for now
 	return false;
 	/*const UMovieSceneSequence* const Sequence = GetSequence();
-	if (!IsValid(Sequence))
+	if (!Sequence)
 	{
 		return false;
 	}
 
 	const UMovieScene* const MovieScene = Sequence->GetMovieScene();
-	if (!IsValid(MovieScene))
+	if (!MovieScene)
 	{
 		return false;
 	}
@@ -303,7 +299,7 @@ bool FNavigationToolSequence::CanRename() const
 bool FNavigationToolSequence::Rename(const FString& InName)
 {
 	UMovieSceneSequence* const Sequence = GetSequence();
-	if (!IsValid(Sequence))
+	if (!Sequence)
 	{
 		return false;
 	}
@@ -488,7 +484,7 @@ void FNavigationToolSequence::SetColor(const TOptional<FColor>& InColor)
 EItemRevisionControlState FNavigationToolSequence::GetRevisionControlState() const
 {
 	UMovieSceneSequence* const Sequence = GetSequence();
-	if (!IsValid(Sequence))
+	if (!Sequence)
 	{
 		return EItemRevisionControlState::None;
 	}
@@ -504,8 +500,7 @@ EItemRevisionControlState FNavigationToolSequence::GetRevisionControlState() con
 
 const FSlateBrush* FNavigationToolSequence::GetRevisionControlStatusIcon() const
 {
-	UMovieSceneSequence* const Sequence = GetSequence();
-	if (IsValid(Sequence))
+	if (const UMovieSceneSequence* const Sequence = GetSequence())
 	{
 		return FindSourceControlStatusBrush(Sequence->GetPackage());
 	}
@@ -514,8 +509,7 @@ const FSlateBrush* FNavigationToolSequence::GetRevisionControlStatusIcon() const
 
 FText FNavigationToolSequence::GetRevisionControlStatusText() const
 {
-	UMovieSceneSequence* const Sequence = GetSequence();
-	if (IsValid(Sequence))
+	if (const UMovieSceneSequence* const Sequence = GetSequence())
 	{
 		return FindSourceControlStatusText(Sequence->GetPackage());
 	}
@@ -532,18 +526,15 @@ EItemContainsPlayhead FNavigationToolSequence::ContainsPlayhead() const
 
 	const FQualifiedFrameTime PlayheadTime = Sequencer->GetLocalTime();
 
-	UMovieSceneSubSection* const SubSection = GetSubSection();
-	if (IsValid(SubSection))
+	if (UMovieSceneSubSection* const SubSection = GetSubSection())
 	{
 		return SubSection->GetTrueRange().Contains(PlayheadTime.Time.FrameNumber)
 			? EItemContainsPlayhead::ContainsPlayhead : EItemContainsPlayhead::None;
 	}
 
-	UMovieSceneSequence* const RootSequence = Sequencer->GetRootMovieSceneSequence();
-	if (IsValid(RootSequence))
+	if (UMovieSceneSequence* const RootSequence = Sequencer->GetRootMovieSceneSequence())
 	{
-		UMovieScene* const RootSequenceMovieScene = RootSequence->GetMovieScene();
-		if (IsValid(RootSequenceMovieScene))
+		if (UMovieScene* const RootSequenceMovieScene = RootSequence->GetMovieScene())
 		{
 			return RootSequenceMovieScene->GetPlaybackRange().Contains(PlayheadTime.Time.FrameNumber)
 				? EItemContainsPlayhead::ContainsPlayhead : EItemContainsPlayhead::None;
@@ -556,13 +547,13 @@ EItemContainsPlayhead FNavigationToolSequence::ContainsPlayhead() const
 FFrameNumber FNavigationToolSequence::GetInTime() const
 {
 	const UMovieSceneSequence* const Sequence = WeakSequence.Get();
-	if (!IsValid(Sequence))
+	if (!Sequence)
 	{
 		return 0;
 	}
 
 	const UMovieSceneSubSection* const SubSection = WeakSubSection.Get();
-	if (!IsValid(SubSection))
+	if (!SubSection)
 	{
 		return 0;
 	}
@@ -574,13 +565,13 @@ FFrameNumber FNavigationToolSequence::GetInTime() const
 void FNavigationToolSequence::SetInTime(const FFrameNumber& InTime)
 {
 	const UMovieSceneSequence* const Sequence = WeakSequence.Get();
-	if (!IsValid(Sequence))
+	if (!Sequence)
 	{
 		return;
 	}
 
 	UMovieSceneSubSection* const SubSection = WeakSubSection.Get();
-	if (!IsValid(SubSection))
+	if (!SubSection)
 	{
 		return;
 	}
@@ -597,13 +588,13 @@ void FNavigationToolSequence::SetInTime(const FFrameNumber& InTime)
 FFrameNumber FNavigationToolSequence::GetOutTime() const
 {
 	const UMovieSceneSequence* const Sequence = WeakSequence.Get();
-	if (!IsValid(Sequence))
+	if (!Sequence)
 	{
 		return 0;
 	}
 
 	const UMovieSceneSubSection* const SubSection = WeakSubSection.Get();
-	if (!IsValid(SubSection))
+	if (!SubSection)
 	{
 		return 0;
 	}
@@ -615,13 +606,13 @@ FFrameNumber FNavigationToolSequence::GetOutTime() const
 void FNavigationToolSequence::SetOutTime(const FFrameNumber& InTime)
 {
 	const UMovieSceneSequence* const Sequence = WeakSequence.Get();
-	if (!IsValid(Sequence))
+	if (!Sequence)
 	{
 		return;
 	}
 
 	UMovieSceneSubSection* const SubSection = WeakSubSection.Get();
-	if (!IsValid(SubSection))
+	if (!SubSection)
 	{
 		return;
 	}
@@ -652,8 +643,7 @@ int32 FNavigationToolSequence::GetSubSectionIndex() const
 
 UMovieScene* FNavigationToolSequence::GetSequenceMovieScene() const
 {
-	UMovieSceneSequence* const Sequence = GetSequence();
-	if (IsValid(Sequence))
+	if (const UMovieSceneSequence* const Sequence = GetSequence())
 	{
 		return Sequence->GetMovieScene();
 	}
@@ -663,7 +653,7 @@ UMovieScene* FNavigationToolSequence::GetSequenceMovieScene() const
 TArray<FMovieSceneBinding> FNavigationToolSequence::GetSortedBindings() const
 {
 	UMovieScene* const MovieScene = GetSequenceMovieScene();
-	if (!IsValid(MovieScene))
+	if (!MovieScene)
 	{
 		return {};
 	}
@@ -718,6 +708,6 @@ TViewModelPtr<FSectionModel> FNavigationToolSequence::GetViewModel() const
 	return StorageExtension->FindModelForSection(WeakSubSection.Get());
 }
 
-}
+} // namespace UE::SequenceNavigator
 
 #undef LOCTEXT_NAMESPACE
diff --git a/Engine/Plugins/Experimental/SequenceNavigator/Source/SequenceNavigator/Private/Items/NavigationToolSubTrack.cpp b/Engine/Plugins/Experimental/SequenceNavigator/Source/SequenceNavigator/Private/Items/NavigationToolSubTrack.cpp
--- a/Engine/Plugins/Experimental/SequenceNavigator/Source/SequenceNavigator/Private/Items/NavigationToolSubTrack.cpp
+++ b/Engine/Plugins/Experimental/SequenceNavigator/Source/SequenceNavigator/Private/Items/NavigationToolSubTrack.cpp
@@ -13,7 +13,7 @@
 
 #define LOCTEXT_NAMESPACE "NavigationToolSubTrack"
 
-namespace UE::Sequencer::NavigationTool
+namespace UE::SequenceNavigator
 {
 
 using namespace UE::Sequencer;
@@ -38,19 +38,16 @@ void FNavigationToolSubTrack::FindChildren(TArray<FNavigationToolItemPtr>& OutCh
 {
 	Super::FindChildren(OutChildren, bInRecursive);
 
-	UMovieSceneSubTrack* const SubTrack = GetSubTrack();
-	if (IsValid(SubTrack))
+	if (UMovieSceneSubTrack* const SubTrack = GetSubTrack())
 	{
 		const TSharedRef<FNavigationToolSubTrack> SharedThisRef = SharedThis(this);
 
 		const TArray<UMovieSceneSection*>& AllSections = SubTrack->GetAllSections();
 		for (int32 Index = 0; Index < AllSections.Num(); ++Index)
 		{
-			UMovieSceneSubSection* const SubSection = Cast<UMovieSceneSubSection>(AllSections[Index]);
-			if (IsValid(SubSection))
+			if (UMovieSceneSubSection* const SubSection = Cast<UMovieSceneSubSection>(AllSections[Index]))
 			{
-				UMovieSceneSequence* const Sequence = SubSection->GetSequence();
-				if (IsValid(Sequence))
+				if (UMovieSceneSequence* const Sequence = SubSection->GetSequence())
 				{
 					const FNavigationToolItemPtr NewItem = Tool.FindOrAdd<FNavigationToolSequence>(GetProvider(), SharedThisRef, Sequence, SubSection, Index);
 					OutChildren.Add(NewItem);
@@ -81,8 +78,7 @@ UMovieSceneSubTrack* FNavigationToolSubTrack::GetSubTrack() const
 
 EItemSequenceInactiveState FNavigationToolSubTrack::GetInactiveState() const
 {
-	UMovieSceneSubSection* const ThisSubSection = Cast<UMovieSceneSubSection>(WeakSection.Get());
-	if (IsValid(ThisSubSection))
+	if (UMovieSceneSubSection* const ThisSubSection = Cast<UMovieSceneSubSection>(WeakSection.Get()))
 	{
 		return ThisSubSection->IsActive() ? EItemSequenceInactiveState::None : EItemSequenceInactiveState::Inactive;
 	}
@@ -104,8 +100,7 @@ void FNavigationToolSubTrack::SetIsInactive(const bool bInIsInactive)
 {
 	const bool bNewActiveState = !bInIsInactive;
 
-	UMovieSceneSubSection* const ThisSubSection = Cast<UMovieSceneSubSection>(WeakSection.Get());
-	if (IsValid(ThisSubSection))
+	if (UMovieSceneSubSection* const ThisSubSection = Cast<UMovieSceneSubSection>(WeakSection.Get()))
 	{
 		if (ThisSubSection->IsActive() != bNewActiveState)
 		{
@@ -143,8 +138,7 @@ void FNavigationToolSubTrack::SetMarkerVisibility(const bool bInVisible)
 		return;
 	}
 
-	UMovieSceneSubSection* const ThisSubSection = Cast<UMovieSceneSubSection>(WeakSection.Get());
-	if (IsValid(ThisSubSection))
+	if (UMovieSceneSubSection* const ThisSubSection = Cast<UMovieSceneSubSection>(WeakSection.Get()))
 	{
 		UMovieSceneSequence* const Sequence = ThisSubSection->GetSequence();
 		const bool bIsVisible = IsGloballyMarkedFramesForSequence(Sequence);
@@ -161,6 +155,6 @@ void FNavigationToolSubTrack::SetMarkerVisibility(const bool bInVisible)
 	}
 }
 
-}
+} // namespace UE::SequenceNavigator
 
 #undef LOCTEXT_NAMESPACE
diff --git a/Engine/Plugins/Experimental/SequenceNavigator/Source/SequenceNavigator/Private/Items/NavigationToolTrack.cpp b/Engine/Plugins/Experimental/SequenceNavigator/Source/SequenceNavigator/Private/Items/NavigationToolTrack.cpp
--- a/Engine/Plugins/Experimental/SequenceNavigator/Source/SequenceNavigator/Private/Items/NavigationToolTrack.cpp
+++ b/Engine/Plugins/Experimental/SequenceNavigator/Source/SequenceNavigator/Private/Items/NavigationToolTrack.cpp
@@ -18,7 +18,7 @@
 
 #define LOCTEXT_NAMESPACE "NavigationToolTrack"
 
-namespace UE::Sequencer::NavigationTool
+namespace UE::SequenceNavigator
 {
 
 FNavigationToolTrack::FNavigationToolTrack(INavigationTool& InTool
@@ -67,8 +67,7 @@ void FNavigationToolTrack::FindChildren(TArray<FNavigationToolItemPtr>& OutChild
 
 FText FNavigationToolTrack::GetDisplayName() const
 {
-	UMovieSceneTrack* const Track = GetTrack();
-	if (IsValid(Track))
+	if (UMovieSceneTrack* const Track = GetTrack())
 	{
 		return Track->GetDisplayName();
 	}
@@ -77,8 +76,7 @@ FText FNavigationToolTrack::GetDisplayName() const
 
 FText FNavigationToolTrack::GetClassName() const
 {
-	UMovieSceneTrack* const Track = GetTrack();
-	if (IsValid(Track))
+	if (UMovieSceneTrack* const Track = GetTrack())
 	{
 		return FText::FromString(Track->GetClass()->GetName());
 	}
@@ -88,7 +86,7 @@ FText FNavigationToolTrack::GetClassName() const
 const FSlateBrush* FNavigationToolTrack::GetDefaultIconBrush() const
 {
 	UMovieSceneTrack* const Track = GetTrack();
-	if (!IsValid(Track))
+	if (!Track)
 	{
 		return nullptr;
 	}
@@ -114,8 +112,7 @@ FSlateIcon FNavigationToolTrack::GetIcon() const
 
 FText FNavigationToolTrack::GetIconTooltipText() const
 {
-	const UMovieSceneTrack* const Track = GetTrack();
-	if (IsValid(Track))
+	if (const UMovieSceneTrack* const Track = GetTrack())
 	{
 		return Track->GetClass()->GetDisplayNameText();
 	}
@@ -125,13 +122,12 @@ FText FNavigationToolTrack::GetIconTooltipText() const
 bool FNavigationToolTrack::IsSelected(const FNavigationToolScopedSelection& InSelection) const
 {
 	UMovieSceneTrack* const Track = GetTrack();
-	return IsValid(Track) && InSelection.IsSelected(Track);
+	return Track && InSelection.IsSelected(Track);
 }
 
 void FNavigationToolTrack::Select(FNavigationToolScopedSelection& InSelection) const
 {
-	UMovieSceneTrack* const Track = GetTrack();
-	if (IsValid(Track))
+	if (UMovieSceneTrack* const Track = GetTrack())
 	{
 		InSelection.Select(Track);
 	}
@@ -150,7 +146,7 @@ void FNavigationToolTrack::OnDoubleClick()
 	}
 
 	UMovieSceneTrack* const Track = GetTrack();
-	if (!IsValid(Track))
+	if (!Track)
 	{
 		return;
 	}
@@ -283,8 +279,10 @@ int32 FNavigationToolTrack::GetSectionIndex() const
 	return SectionIndex;
 }
 
-TViewModelPtr<FTrackModel> FNavigationToolTrack::GetViewModel() const
+UE::Sequencer::TViewModelPtr<UE::Sequencer::FTrackModel> FNavigationToolTrack::GetViewModel() const
 {
+	using namespace UE::Sequencer;
+
 	UMovieSceneTrack* const Track = GetTrack();
 	if (!IsValid(Track))
 	{
@@ -318,6 +316,6 @@ TViewModelPtr<FTrackModel> FNavigationToolTrack::GetViewModel() const
 	return StorageExtension->FindModelForTrack(Track);
 }
 
-}
+} // namespace UE::SequenceNavigator
 
 #undef LOCTEXT_NAMESPACE
diff --git a/Engine/Plugins/Experimental/SequenceNavigator/Source/SequenceNavigator/Private/Menus/NavigationToolFilterBarContextMenu.cpp b/Engine/Plugins/Experimental/SequenceNavigator/Source/SequenceNavigator/Private/Menus/NavigationToolFilterBarContextMenu.cpp
--- a/Engine/Plugins/Experimental/SequenceNavigator/Source/SequenceNavigator/Private/Menus/NavigationToolFilterBarContextMenu.cpp
+++ b/Engine/Plugins/Experimental/SequenceNavigator/Source/SequenceNavigator/Private/Menus/NavigationToolFilterBarContextMenu.cpp
@@ -9,12 +9,12 @@
 
 #define LOCTEXT_NAMESPACE "NavigationToolFilterBarContextMenu"
 
-namespace UE::Sequencer::NavigationTool
+namespace UE::SequenceNavigator
 {
 
 TSharedRef<SWidget> FNavigationToolFilterBarContextMenu::CreateMenu(const TSharedRef<INavigationToolFilterBar>& InFilterBar)
 {
-	const FName FilterMenuName = TEXT("Sequencer.NavigationTool.FilterBarContextMenu");
+	const FName FilterMenuName = TEXT("SequenceNavigator.FilterBarContextMenu");
 	if (!UToolMenus::Get()->IsMenuRegistered(FilterMenuName))
 	{
 		UToolMenu* const Menu = UToolMenus::Get()->RegisterMenu(FilterMenuName);
@@ -37,13 +37,13 @@ TSharedRef<SWidget> FNavigationToolFilterBarContextMenu::CreateMenu(const TShare
 
 void FNavigationToolFilterBarContextMenu::PopulateMenu(UToolMenu* const InMenu)
 {
-	if (!IsValid(InMenu))
+	if (!InMenu)
 	{
 		return;
 	}
 
 	UNavigationToolFilterBarContext* const Context = InMenu->FindContext<UNavigationToolFilterBarContext>();
-	if (!IsValid(Context))
+	if (!Context)
 	{
 		return;
 	}
@@ -142,6 +142,6 @@ void FNavigationToolFilterBarContextMenu::OnResetFilters()
 	FilterBar->EnableCustomTextFilters(false, {});
 }
 
-}
+} // namespace UE::SequenceNavigator
 
 #undef LOCTEXT_NAMESPACE
diff --git a/Engine/Plugins/Experimental/SequenceNavigator/Source/SequenceNavigator/Private/Menus/NavigationToolFilterContextMenu.cpp b/Engine/Plugins/Experimental/SequenceNavigator/Source/SequenceNavigator/Private/Menus/NavigationToolFilterContextMenu.cpp
--- a/Engine/Plugins/Experimental/SequenceNavigator/Source/SequenceNavigator/Private/Menus/NavigationToolFilterContextMenu.cpp
+++ b/Engine/Plugins/Experimental/SequenceNavigator/Source/SequenceNavigator/Private/Menus/NavigationToolFilterContextMenu.cpp
@@ -13,7 +13,7 @@
 
 #define LOCTEXT_NAMESPACE "NavigationToolFilterContextMenu"
 
-namespace UE::Sequencer::NavigationTool
+namespace UE::SequenceNavigator
 {
 
 TSharedRef<SWidget> FNavigationToolFilterContextMenu::CreateMenuWidget(const TSharedRef<FNavigationToolFilter>& InFilter)
@@ -21,9 +21,9 @@ TSharedRef<SWidget> FNavigationToolFilterContextMenu::CreateMenuWidget(const TSh
 	WeakFilter = InFilter;
 
 	UToolMenus* const ToolMenus = UToolMenus::Get();
-	check(IsValid(ToolMenus));
+	check(ToolMenus);
 
-	const FName FilterMenuName = TEXT("Sequencer.NavigationTool.FilterContextMenu");
+	const FName FilterMenuName = TEXT("SequenceNavigator.FilterContextMenu");
 	if (!ToolMenus->IsMenuRegistered(FilterMenuName))
 	{
 		UToolMenu* const Menu = ToolMenus->RegisterMenu(FilterMenuName);
@@ -48,7 +48,7 @@ TSharedRef<SWidget> FNavigationToolFilterContextMenu::CreateMenuWidget(const TSh
 
 void FNavigationToolFilterContextMenu::PopulateMenu(UToolMenu* const InMenu)
 {
-	if (!IsValid(InMenu))
+	if (!InMenu)
 	{
 		return;
 	}
@@ -241,6 +241,6 @@ const TSharedPtr<FNavigationToolFilter> FNavigationToolFilterContextMenu::GetFil
 	return WeakFilter.IsValid() ? WeakFilter.Pin() : nullptr;
 }
 
-}
+} // namespace UE::SequenceNavigator
 
 #undef LOCTEXT_NAMESPACE
diff --git a/Engine/Plugins/Experimental/SequenceNavigator/Source/SequenceNavigator/Private/Menus/NavigationToolFilterMenu.cpp b/Engine/Plugins/Experimental/SequenceNavigator/Source/SequenceNavigator/Private/Menus/NavigationToolFilterMenu.cpp
--- a/Engine/Plugins/Experimental/SequenceNavigator/Source/SequenceNavigator/Private/Menus/NavigationToolFilterMenu.cpp
+++ b/Engine/Plugins/Experimental/SequenceNavigator/Source/SequenceNavigator/Private/Menus/NavigationToolFilterMenu.cpp
@@ -12,12 +12,12 @@
 
 #define LOCTEXT_NAMESPACE "NavigationToolFilterMenu"
 
-namespace UE::Sequencer::NavigationTool
+namespace UE::SequenceNavigator
 {
 
 TSharedRef<SWidget> FNavigationToolFilterMenu::CreateMenu(const TSharedRef<INavigationToolFilterBar>& InFilterBar)
 {
-	const FName FilterMenuName = TEXT("Sequencer.NavigationTool.FilterMenu");
+	const FName FilterMenuName = TEXT("SequenceNavigator.FilterMenu");
 	if (!UToolMenus::Get()->IsMenuRegistered(FilterMenuName))
 	{
 		UToolMenu* const Menu = UToolMenus::Get()->RegisterMenu(FilterMenuName);
@@ -41,13 +41,13 @@ TSharedRef<SWidget> FNavigationToolFilterMenu::CreateMenu(const TSharedRef<INavi
 
 void FNavigationToolFilterMenu::PopulateMenu(UToolMenu* const InMenu)
 {
-	if (!IsValid(InMenu))
+	if (!InMenu)
 	{
 		return;
 	}
 
 	UNavigationToolFilterBarContext* const Context = InMenu->FindContext<UNavigationToolFilterBarContext>();
-	if (!IsValid(Context))
+	if (!Context)
 	{
 		return;
 	}
@@ -115,7 +115,7 @@ void FNavigationToolFilterMenu::PopulateCommonFilterSections(UToolMenu& InMenu)
 	}
 
 	const UNavigationToolSettings* const Settings = GetDefault<UNavigationToolSettings>();
-	if (!IsValid(Settings))
+	if (!Settings)
 	{
 		return;
 	}
@@ -153,7 +153,7 @@ void FNavigationToolFilterMenu::PopulateCommonFilterSections(UToolMenu& InMenu)
 
 void FNavigationToolFilterMenu::FillCustomTextFiltersMenu(UToolMenu* const InMenu)
 {
-	if (!IsValid(InMenu))
+	if (!InMenu)
 	{
 		return;
 	}
@@ -472,6 +472,6 @@ void FNavigationToolFilterMenu::ResetFilters()
 	FilterBar->ResetFilters();
 }
 
-}
+} // namespace UE::SequenceNavigator
 
 #undef LOCTEXT_NAMESPACE
diff --git a/Engine/Plugins/Experimental/SequenceNavigator/Source/SequenceNavigator/Private/Menus/NavigationToolItemContextMenu.cpp b/Engine/Plugins/Experimental/SequenceNavigator/Source/SequenceNavigator/Private/Menus/NavigationToolItemContextMenu.cpp
--- a/Engine/Plugins/Experimental/SequenceNavigator/Source/SequenceNavigator/Private/Menus/NavigationToolItemContextMenu.cpp
+++ b/Engine/Plugins/Experimental/SequenceNavigator/Source/SequenceNavigator/Private/Menus/NavigationToolItemContextMenu.cpp
@@ -11,12 +11,12 @@
 
 #define LOCTEXT_NAMESPACE "NavigationToolItemContextMenu"
 
-namespace UE::Sequencer::NavigationTool
+namespace UE::SequenceNavigator
 {
 
 FName FNavigationToolItemContextMenu::GetMenuName()
 {
-	static const FName MenuName(TEXT("Sequencer.NavigationTool.ItemContextMenu"));
+	static const FName MenuName(TEXT("SequenceNavigator.ItemContextMenu"));
 	return MenuName;
 }
 
@@ -49,7 +49,7 @@ void FNavigationToolItemContextMenu::PopulateMenu(UToolMenu* const InMenu)
 	UToolMenu& MenuRef = *InMenu;
 
 	UNavigationToolItemMenuContext* const ItemsContext = MenuRef.FindContext<UNavigationToolItemMenuContext>();
-	if (!IsValid(ItemsContext))
+	if (!ItemsContext)
 	{
 		return;
 	}
@@ -120,6 +120,6 @@ void FNavigationToolItemContextMenu::CreateToolSection(UToolMenu& InMenu)
 	FocusInContentBrowserEntry.Icon = FSlateIcon(FAppStyle::GetAppStyleSetName(), TEXT("ContentBrowser.TabIcon"));
 }
 
-}
+} // namespace UE::SequenceNavigator
 
 #undef LOCTEXT_NAMESPACE
diff --git a/Engine/Plugins/Experimental/SequenceNavigator/Source/SequenceNavigator/Private/Menus/NavigationToolToolbarMenu.cpp b/Engine/Plugins/Experimental/SequenceNavigator/Source/SequenceNavigator/Private/Menus/NavigationToolToolbarMenu.cpp
--- a/Engine/Plugins/Experimental/SequenceNavigator/Source/SequenceNavigator/Private/Menus/NavigationToolToolbarMenu.cpp
+++ b/Engine/Plugins/Experimental/SequenceNavigator/Source/SequenceNavigator/Private/Menus/NavigationToolToolbarMenu.cpp
@@ -17,12 +17,12 @@
 
 #define LOCTEXT_NAMESPACE "NavigationToolToolbarMenu"
 
-namespace UE::Sequencer::NavigationTool
+namespace UE::SequenceNavigator
 {
 
 FName FNavigationToolToolbarMenu::GetMenuName()
 {
-	static const FName ToolBarName(TEXT("Sequencer.NavigationTool.ToolBar"));
+	static const FName ToolBarName(TEXT("SequenceNavigator.ToolBar"));
 	return ToolBarName;
 }
 
@@ -41,8 +41,7 @@ TSharedRef<SWidget> FNavigationToolToolbarMenu::CreateToolbar(const TSharedRef<F
 		ToolBar->bShouldCloseWindowAfterMenuSelection = false;
 		ToolBar->AddDynamicSection(TEXT("Main"), FNewToolMenuDelegate::CreateLambda([this](UToolMenu* const InMenu)
 			{
-				UNavigationToolViewMenuContext* const Context = InMenu->FindContext<UNavigationToolViewMenuContext>();
-				if (IsValid(Context))
+				if (UNavigationToolViewMenuContext* const Context = InMenu->FindContext<UNavigationToolViewMenuContext>())
 				{
 					Context->OnPopulateMenu.ExecuteIfBound(InMenu);
 				}
@@ -59,13 +58,13 @@ TSharedRef<SWidget> FNavigationToolToolbarMenu::CreateToolbar(const TSharedRef<F
 
 void FNavigationToolToolbarMenu::PopulateToolBar(UToolMenu* const InMenu)
 {
-	if (!IsValid(InMenu))
+	if (!InMenu)
 	{
 		return;
 	}
 
 	UNavigationToolViewMenuContext* const ToolBarContext = InMenu->FindContext<UNavigationToolViewMenuContext>();
-	if (!IsValid(ToolBarContext))
+	if (!ToolBarContext)
 	{
 		return;
 	}
@@ -108,7 +107,7 @@ void FNavigationToolToolbarMenu::PopulateToolBar(UToolMenu* const InMenu)
 
 void FNavigationToolToolbarMenu::CreateSettingsMenu(UToolMenu* const InMenu)
 {
-	if (!IsValid(InMenu))
+	if (!InMenu)
 	{
 		return;
 	}
@@ -145,13 +144,13 @@ void FNavigationToolToolbarMenu::CreateSettingsMenu(UToolMenu* const InMenu)
 
 void FNavigationToolToolbarMenu::CreateItemViewOptionsMenu(UToolMenu* const InMenu)
 {
-	if (!IsValid(InMenu))
+	if (!InMenu)
 	{
 		return;
 	}
 
 	UNavigationToolViewMenuContext* const ToolBarContext = InMenu->FindContext<UNavigationToolViewMenuContext>();
-	if (!IsValid(ToolBarContext))
+	if (!ToolBarContext)
 	{
 		return;
 	}
@@ -425,6 +424,6 @@ void FNavigationToolToolbarMenu::OnToggleGlobalFilter(const TSharedRef<FNavigati
 	InTool->Refresh();
 }
 
-}
+} // namespace UE::SequenceNavigator
 
 #undef LOCTEXT_NAMESPACE
diff --git a/Engine/Plugins/Experimental/SequenceNavigator/Source/SequenceNavigator/Private/Menus/NavigationToolViewMenuContext.h b/Engine/Plugins/Experimental/SequenceNavigator/Source/SequenceNavigator/Private/Menus/NavigationToolViewMenuContext.h
--- a/Engine/Plugins/Experimental/SequenceNavigator/Source/SequenceNavigator/Private/Menus/NavigationToolViewMenuContext.h
+++ b/Engine/Plugins/Experimental/SequenceNavigator/Source/SequenceNavigator/Private/Menus/NavigationToolViewMenuContext.h
@@ -16,24 +16,24 @@ class UNavigationToolViewMenuContext : public UObject
 {
 	GENERATED_BODY()
 
-	friend UE::Sequencer::NavigationTool::FNavigationToolView;
+	friend UE::SequenceNavigator::FNavigationToolView;
 
 public:
-	void Init(const TWeakPtr<UE::Sequencer::NavigationTool::FNavigationToolView>& InToolView)
+	void Init(const TWeakPtr<UE::SequenceNavigator::FNavigationToolView>& InToolView)
 	{
 		WeakToolView = InToolView;
 	}
 
-	TSharedPtr<UE::Sequencer::NavigationTool::FNavigationToolView> GetToolView() const
+	TSharedPtr<UE::SequenceNavigator::FNavigationToolView> GetToolView() const
 	{
 		return WeakToolView.Pin();
 	}
 
-	TSharedPtr<UE::Sequencer::NavigationTool::FNavigationTool> GetTool() const
+	TSharedPtr<UE::SequenceNavigator::FNavigationTool> GetTool() const
 	{
-		if (const TSharedPtr<UE::Sequencer::NavigationTool::FNavigationToolView> ToolView = GetToolView())
+		if (const TSharedPtr<UE::SequenceNavigator::FNavigationToolView> ToolView = GetToolView())
 		{
-			return StaticCastSharedPtr<UE::Sequencer::NavigationTool::FNavigationTool>(ToolView->GetOwnerTool());
+			return StaticCastSharedPtr<UE::SequenceNavigator::FNavigationTool>(ToolView->GetOwnerTool());
 		}
 		return nullptr;
 	}
@@ -41,5 +41,5 @@ public:
 	FOnPopulateNavigationToolViewToolbarMenu OnPopulateMenu;
 
 private:
-	TWeakPtr<UE::Sequencer::NavigationTool::FNavigationToolView> WeakToolView;
+	TWeakPtr<UE::SequenceNavigator::FNavigationToolView> WeakToolView;
 };
