Motion Design: Fix on undo related to guides

diff --git a/Engine/Plugins/Experimental/Avalanche/Source/AvalancheLevelViewport/Private/SAvaLevelViewport.cpp b/Engine/Plugins/Experimental/Avalanche/Source/AvalancheLevelViewport/Private/SAvaLevelViewport.cpp
--- a/Engine/Plugins/Experimental/Avalanche/Source/AvalancheLevelViewport/Private/SAvaLevelViewport.cpp
+++ b/Engine/Plugins/Experimental/Avalanche/Source/AvalancheLevelViewport/Private/SAvaLevelViewport.cpp
@@ -555,7 +555,7 @@ void SAvaLevelViewport::SaveGuides()
 		return;
 	}
 
-	FScopedTransaction SaveTransaction(UE::AvaLevelViewport::Private::GuideTransactionTitle);
+	FScopedTransaction SaveTransaction(UE::AvaLevelViewport::Private::GuideTransactionTitle, !GIsTransacting);
 	DataSubsystem->ModifyDataSource();
 
 	TArray<FAvaViewportGuideInfo> GuideInfos;
