Motion Design: fixed localization for outliner filters

diff --git a/Engine/Plugins/Experimental/Avalanche/Source/AvalancheOutliner/Private/Filters/AvaOutlinerItemTypeFilter.cpp b/Engine/Plugins/Experimental/Avalanche/Source/AvalancheOutliner/Private/Filters/AvaOutlinerItemTypeFilter.cpp
--- a/Engine/Plugins/Experimental/Avalanche/Source/AvalancheOutliner/Private/Filters/AvaOutlinerItemTypeFilter.cpp
+++ b/Engine/Plugins/Experimental/Avalanche/Source/AvalancheOutliner/Private/Filters/AvaOutlinerItemTypeFilter.cpp
@@ -7,6 +7,8 @@
 #include "Item/IAvaOutlinerItem.h"
 #include "Styling/SlateIconFinder.h"
 
+#define LOCTEXT_NAMESPACE "AvaOutlinerItemTypeFilter"
+
 FAvaOutlinerItemTypeFilterData::FAvaOutlinerItemTypeFilterData(const TArray<TSubclassOf<UObject>>& InFilterClasses
 		, EAvaOutlinerTypeFilterMode InMode
 		, const FSlateBrush* InIconBrush
@@ -35,21 +37,18 @@ FText FAvaOutlinerItemTypeFilterData::GetTooltipText() const
 		return TooltipText;
 	}
 
-	FString SupportedClasses;
+	TArray<FText> SupportedClasses;
+	SupportedClasses.Reserve(FilterClasses.Num());
 
 	for (TSubclassOf<UObject> FilterClass : FilterClasses)
 	{
-		if (IsValid(FilterClass))
+		if (FilterClass)
 		{
-			if (!SupportedClasses.IsEmpty())
-			{
-				SupportedClasses += TEXT(", ");
-			}
-			SupportedClasses += FilterClass->GetDisplayNameText().ToString();
+			SupportedClasses.Add(FilterClass->GetDisplayNameText());
 		}
 	}
 
-	const_cast<FAvaOutlinerItemTypeFilterData*>(this)->TooltipText = FText::FromString(MoveTemp(SupportedClasses));
+	const_cast<FAvaOutlinerItemTypeFilterData*>(this)->TooltipText = FText::Join(LOCTEXT("CommaDelimiter", ", "), SupportedClasses);
 	return TooltipText;
 }
 
@@ -188,3 +187,5 @@ FText FAvaOutlinerItemTypeFilter::GetTooltipText() const
 {
 	return FilterData.GetTooltipText();
 }
+
+#undef LOCTEXT_NAMESPACE
