[Motion Design] Rundown/Playable: Support for ignoring RC controllers to not interfere with Scene State tasks (that are also using RC controllers)

diff --git a/Engine/Plugins/Experimental/Avalanche/Source/AvalancheMedia/Internal/Playable/AvaPlayableRemoteControl.h b/Engine/Plugins/Experimental/Avalanche/Source/AvalancheMedia/Internal/Playable/AvaPlayableRemoteControl.h
--- a/Engine/Plugins/Experimental/Avalanche/Source/AvalancheMedia/Internal/Playable/AvaPlayableRemoteControl.h
+++ b/Engine/Plugins/Experimental/Avalanche/Source/AvalancheMedia/Internal/Playable/AvaPlayableRemoteControl.h
@@ -32,6 +32,12 @@ enum class EAvaPlayableRemoteControlResult : int8
 
 namespace UE::AvaPlayableRemoteControl
 {
+	/** Returns true if the binding is valid and access to read the property. */
+	AVALANCHEMEDIA_API bool HasReadAccess(const TSharedPtr<const FRemoteControlEntity>& InRemoteControlEntity, FString& OutErrorText);
+
+	/** Returns true if the binding is valid and access to write the property. */
+	AVALANCHEMEDIA_API bool HasWriteAccess(const TSharedPtr<const FRemoteControlEntity>& InRemoteControlEntity, FString& OutErrorText);
+
 	/** Gets the value (as bytes) of the given property (remote control entity). Returns true if the operation was successful, false otherwise. */
 	AVALANCHEMEDIA_API EAvaPlayableRemoteControlResult GetValueOfEntity(const TSharedPtr<const FRemoteControlEntity>& InRemoteControlEntity, TArray<uint8>& OutValue);
 	AVALANCHEMEDIA_API EAvaPlayableRemoteControlResult GetValueOfEntity(const TSharedPtr<const FRemoteControlEntity>& InRemoteControlEntity, FString& OutValue);
@@ -47,6 +53,7 @@ namespace UE::AvaPlayableRemoteControl
 	/** Set the value (as bytes) of the given controller. Returns true if the operation was successful, false otherwise. */
 	AVALANCHEMEDIA_API EAvaPlayableRemoteControlResult SetValueOfController(URCVirtualPropertyBase* InController, const TArrayView<const uint8>& InValue);
 	AVALANCHEMEDIA_API EAvaPlayableRemoteControlResult SetValueOfController(URCVirtualPropertyBase* InController, const FString& InValue);
+	AVALANCHEMEDIA_API EAvaPlayableRemoteControlResult SetValueOfController(URCVirtualPropertyBase* InController, const FString& InValue, bool bInBehaviorsEnabled);
 
 	/** Retrieve the entity ids that are controlled by the given controller. */
 	AVALANCHEMEDIA_API bool GetEntitiesControlledByController(const URemoteControlPreset* InRemoteControlPreset, const URCVirtualPropertyBase* InVirtualProperty, TSet<FGuid>& OutEntityIds);
diff --git a/Engine/Plugins/Experimental/Avalanche/Source/AvalancheMedia/Internal/Rundown/AvaRundownManagedInstance.h b/Engine/Plugins/Experimental/Avalanche/Source/AvalancheMedia/Internal/Rundown/AvaRundownManagedInstance.h
--- a/Engine/Plugins/Experimental/Avalanche/Source/AvalancheMedia/Internal/Rundown/AvaRundownManagedInstance.h
+++ b/Engine/Plugins/Experimental/Avalanche/Source/AvalancheMedia/Internal/Rundown/AvaRundownManagedInstance.h
@@ -3,6 +3,7 @@
 #pragma once
 
 #include "AvaType.h"
+#include "Playable/AvaPlayableRemoteControlPresetInfo.h"
 #include "Playable/AvaPlayableRemoteControlValues.h"
 #include "UObject/GCObject.h"
 #include "UObject/SoftObjectPath.h"
@@ -10,6 +11,7 @@
 class FAvaRundownManagedInstanceCache;
 class IAvaSceneInterface;
 class URemoteControlPreset;
+struct FAvaPlayableRemoteControlPresetInfo;
 
 class AVALANCHEMEDIA_API FAvaRundownManagedInstance : public IAvaTypeCastable, public FGCObject, public TSharedFromThis<FAvaRundownManagedInstance>
 {
@@ -28,12 +30,20 @@ public:
 	/** Access a copy of the Remote Control values from the source (a.k.a. template). */
 	const FAvaPlayableRemoteControlValues& GetDefaultRemoteControlValues() const { return DefaultRemoteControlValues; }
 
+	/** Access cached preset info for this instance. */
+	const FAvaPlayableRemoteControlPresetInfo& GetRemoteControlPresetInfo() const
+	{
+		return RemoteControlPresetInfo && !RemoteControlPresetInfo->IsDirty() ?
+			*RemoteControlPresetInfo : GetRemoteControlPresetInfoFromCache(GetRemoteControlPreset());
+	}
+
 	/** Get the Ava scene interface from the asset. */
 	virtual IAvaSceneInterface* GetSceneInterface() const { return nullptr; }
 
 	const FSoftObjectPath& GetSourceAssetPath() const { return SourceAssetPath; }
 	
 protected:
+	FAvaPlayableRemoteControlPresetInfo& GetRemoteControlPresetInfoFromCache(URemoteControlPreset* InRemoteControlPreset) const;
 	void InvalidateFromCache();
 
 	void RegisterSourceRemoteControlPresetDelegates(URemoteControlPreset* InSourceRemoteControlPreset);
@@ -52,4 +62,5 @@ protected:
 	FAvaRundownManagedInstanceCache* ParentCache;
 	FSoftObjectPath SourceAssetPath;
 	FAvaPlayableRemoteControlValues DefaultRemoteControlValues;
+	mutable TSharedPtr<FAvaPlayableRemoteControlPresetInfo> RemoteControlPresetInfo;
 };
\ No newline at end of file
diff --git a/Engine/Plugins/Experimental/Avalanche/Source/AvalancheMedia/Private/AvaMediaModule.cpp b/Engine/Plugins/Experimental/Avalanche/Source/AvalancheMedia/Private/AvaMediaModule.cpp
--- a/Engine/Plugins/Experimental/Avalanche/Source/AvalancheMedia/Private/AvaMediaModule.cpp
+++ b/Engine/Plugins/Experimental/Avalanche/Source/AvalancheMedia/Private/AvaMediaModule.cpp
@@ -13,6 +13,7 @@
 #include "Misc/CoreDelegates.h"
 #include "ModularFeature/AvaMediaSynchronizedEventsFeature.h"
 #include "Playable/AvaPlayableGroupSceneViewExtension.h"
+#include "Playable/AvaPlayableRemoteControlPresetInfoCache.h"
 #include "Playback/AvaPlaybackClientDelegates.h"
 #include "Playback/AvaPlaybackClientDummy.h"
 #include "SceneViewExtension.h"
@@ -186,6 +187,8 @@ void FAvaMediaModule::StartupModule()
 	{
 		AvaMediaSync->SetFeatureEnabled(false);
 	}
+
+	RemoteControlPresetInfoCache = MakeShared<FAvaPlayableRemoteControlPresetInfoCache>();
 }
 
 void FAvaMediaModule::ShutdownModule()
@@ -871,4 +874,9 @@ IAvaBroadcastDeviceProviderProxyManager& FAvaMediaModule::GetDeviceProviderProxy
 	return DeviceProviderProxyManager;
 }
 
+IAvaPlayableRemoteControlPresetInfoCache& FAvaMediaModule::GetPlayableRemoteControlPresetInfoCache() const
+{
+	return *RemoteControlPresetInfoCache;
+}
+
 IMPLEMENT_MODULE(FAvaMediaModule, AvalancheMedia)
\ No newline at end of file
diff --git a/Engine/Plugins/Experimental/Avalanche/Source/AvalancheMedia/Private/AvaMediaModule.h b/Engine/Plugins/Experimental/Avalanche/Source/AvalancheMedia/Private/AvaMediaModule.h
--- a/Engine/Plugins/Experimental/Avalanche/Source/AvalancheMedia/Private/AvaMediaModule.h
+++ b/Engine/Plugins/Experimental/Avalanche/Source/AvalancheMedia/Private/AvaMediaModule.h
@@ -17,6 +17,7 @@
 #include "Rundown/AvaRundownServer.h"
 
 class FAvaPlayableGroupSceneViewExtension;
+class FAvaPlayableRemoteControlPresetInfoCache;
 
 class FAvaMediaModule : public IAvaMediaModule
 {
@@ -68,6 +69,7 @@ public:
 	virtual TSharedPtr<IAvaRundownServer> GetRundownServer() const override { return RundownServer; }
 	virtual TSharedPtr<IAvaRundownServer> MakeDetachedRundownServer(const FString& InServerName) override;
 	virtual IAvaBroadcastDeviceProviderProxyManager& GetDeviceProviderProxyManager() override;
+	virtual IAvaPlayableRemoteControlPresetInfoCache& GetPlayableRemoteControlPresetInfoCache() const override;
 	//~ End IAvaMediaModule
 	
 private:
@@ -160,6 +162,7 @@ private:
 	TSharedPtr<FAvaPlaybackServerProcess> LocalPlaybackServerProcess;
 	TSharedPtr<FAvaPlaybackManager> LocalPlaybackManager;
 	TSharedPtr<FAvaRundownManagedInstanceCache> ManagedInstanceCache;
+	TSharedPtr<FAvaPlayableRemoteControlPresetInfoCache> RemoteControlPresetInfoCache;
 
 	TSharedPtr<FAvaPlaybackHttpServer> AvaPlaybackHttpPlaybackServer;
 
diff --git a/Engine/Plugins/Experimental/Avalanche/Source/AvalancheMedia/Private/AvaMediaSettings.cpp b/Engine/Plugins/Experimental/Avalanche/Source/AvalancheMedia/Private/AvaMediaSettings.cpp
--- a/Engine/Plugins/Experimental/Avalanche/Source/AvalancheMedia/Private/AvaMediaSettings.cpp
+++ b/Engine/Plugins/Experimental/Avalanche/Source/AvalancheMedia/Private/AvaMediaSettings.cpp
@@ -13,6 +13,9 @@ UAvaMediaSettings::UAvaMediaSettings()
 	CategoryName = TEXT("Motion Design");
 	SectionName  = TEXT("Playback & Broadcast");
 	PlayableSettings.SynchronizedEventsFeature.Implementation = SynchronizedEventsFeatureSelection_Default.ToString();
+	// The choice of going with a trailing "__" as the default ignored postfix is inspired by a Python naming convention to
+	// indicate ignored/hidden functions not meant to be called directly by users.
+	PlayableSettings.IgnoredControllerPostfix.Add(TEXT("__"));
 }
 
 UAvaMediaSettings* UAvaMediaSettings::GetSingletonInstance()
diff --git a/Engine/Plugins/Experimental/Avalanche/Source/AvalancheMedia/Private/Playable/AvaPlayable.cpp b/Engine/Plugins/Experimental/Avalanche/Source/AvalancheMedia/Private/Playable/AvaPlayable.cpp
--- a/Engine/Plugins/Experimental/Avalanche/Source/AvalancheMedia/Private/Playable/AvaPlayable.cpp
+++ b/Engine/Plugins/Experimental/Avalanche/Source/AvalancheMedia/Private/Playable/AvaPlayable.cpp
@@ -12,6 +12,7 @@
 #include "IAvaSequenceProvider.h"
 #include "Playable/AvaPlayableGroup.h"
 #include "Playable/AvaPlayableRemoteControl.h"
+#include "Playable/AvaPlayableRemoteControlPresetInfo.h"
 #include "Playable/AvaPlayableRemoteControlValues.h"
 #include "Playable/Playables/AvaPlayableLevelStreaming.h"
 #include "Playable/Playables/AvaPlayableRemoteProxy.h"
@@ -30,13 +31,26 @@ static TAutoConsoleVariable<bool> CVarUpdateRCControllerValues(
 	TEXT("Set the controller values without running the behaviors."),
 	ECVF_Default);
 
-// Some "special" controllers, such as texture controller require the behavior to be executed.
+// Some "Special" controllers, such as texture controller require the behavior to be executed.
 static TAutoConsoleVariable<bool> CVarExecuteSpecialRCControllerBehavior(
 	TEXT("MotionDesignPlayable.ExecuteSpecialRCControllerBehavior"),
 	true,
 	TEXT("Execute the special controller behaviors (such as texture controllers)."),
 	ECVF_Default);
 
+// In the pursuit of determinism, we want to restore the state of playables by restoring
+// all the entity values, and if possible the controllers, as more and more controllers
+// are needed to affect the playable state (texture controllers, scene state, etc).
+// We still don't have a guarantee that behaviors are deterministic, but at least, we can know
+// if they are overlapping. Controllers with overlapping controlled states can't be executed, but
+// we can attempt to execute the non-overlapping ones. This is not yet good enough to be enabled by default,
+// but it might be able to solve issues in specific cases. Thus why it is available as a cvar for now.
+static TAutoConsoleVariable<bool> CVarExecuteNonOverlappingRCControllerBehavior(
+	TEXT("MotionDesignPlayable.ExecuteNonOverlappingRCControllerBehavior"),
+	false,
+	TEXT("Execute all non-overlapping controller behaviors."),
+	ECVF_Default);
+
 namespace UE::AvaPlayable::Private
 {
 	bool ShouldCreateLocalPlayable(const FName& InChannelName, const UAvaBroadcast& InBroadcast)
@@ -294,17 +308,22 @@ EAvaPlayableCommandResult UAvaPlayable::UpdateRemoteControlCommand(const TShared
 	using namespace UE::AvaPlayable::Private;
 	UE_LOG(LogAvaPlayable, Verbose, TEXT("%s Playable {%s} -> Executing RC Update."), *GetBriefFrameInfo(), *GetPrettyPlayableInfo(this));
 
-	TSet<FGuid> AppliedEntities;
+	TSet<FGuid> AppliedOrIgnoredEntities;
 	
 	if (CVarUpdateRCControllerValues.GetValueOnGameThread())
 	{
+		using namespace UE::AvaPlayableRemoteControl;
+
 		// Apply special controllers that don't work just with entity values.
 		const bool bExecuteSpecialControllerBehavior = CVarExecuteSpecialRCControllerBehavior.GetValueOnGameThread();
 		const bool bExecuteAllControllerBehaviors = EnumHasAnyFlags(InFlags, EAvaPlayableRCUpdateFlags::ExecuteControllerBehaviors);
+		const bool bExecuteNonOverlappingControllers = CVarExecuteNonOverlappingRCControllerBehavior.GetValueOnGameThread();
 
 		TSet<FGuid> ModifiedControllers;
 
-		// Behaviors are executed only on the controller values that changed. 
+		// For RC updates, the behaviors are executed only on the controller values that changed.
+		// The expectation is that only a few controllers are changed in an update and there are no collisions
+		// in the underlying entities (order independent update).
 		if (bExecuteAllControllerBehaviors)
 		{
 			for (const TPair<FGuid, FAvaPlayableRemoteControlValue>& Controller : InRemoteControlValues->ControllerValues)
@@ -323,29 +342,45 @@ EAvaPlayableCommandResult UAvaPlayable::UpdateRemoteControlCommand(const TShared
 				}
 			}
 		}
+
+		TSharedPtr<FAvaPlayableRemoteControlPresetInfo> RCPInfo =
+			IAvaPlayableRemoteControlPresetInfoCache::Get().GetRemoteControlPresetInfo(GetSourceAssetPath(), RemoteControlPreset);
 		
 		TArray<URCVirtualPropertyBase*> Controllers = RemoteControlPreset->GetControllers();
 		for (URCVirtualPropertyBase* Controller : Controllers)
 		{
+			if (!Controller)
+			{
+				continue;
+			}
+
+			// Skip the ignored controllers
+			if (FAvaPlayableRemoteControlValues::ShouldIgnoreController(Controller))
+			{
+				GetEntitiesControlledByController(RemoteControlPreset, Controller, AppliedOrIgnoredEntities);
+				continue;
+			}
+			
 			if (const FAvaPlayableRemoteControlValue* ControllerValue = InRemoteControlValues->GetControllerValue(Controller->Id))
 			{
-				using namespace UE::AvaPlayableRemoteControl;
 				using namespace UE::RCCustomControllers;
-
-				// Texture Controller Require the Bind Behavior to be executed to setup the texture in external mode.
-				if ((bExecuteSpecialControllerBehavior && GetCustomControllerTypeName(Controller) == CustomTextureControllerName)
-					|| (bExecuteAllControllerBehaviors && (ModifiedControllers.Contains(Controller->Id))))
+				
+				const bool bBehaviorsEnabled =
+					// Texture Controller Require the Bind Behavior to be executed to setup the texture in external mode.
+					(bExecuteSpecialControllerBehavior && GetCustomControllerTypeName(Controller) == CustomTextureControllerName)
+					// Modified controller's behavior are executed when updating RC values
+					|| (bExecuteAllControllerBehaviors && ModifiedControllers.Contains(Controller->Id))
+					// Non-overlapping controller's behavior can be enabled (experimental)
+					|| (bExecuteNonOverlappingControllers && RCPInfo.IsValid() && !RCPInfo->IsControllerOverlapping(Controller->Id)); 
+
+				SetValueOfController(Controller, ControllerValue->Value, bBehaviorsEnabled);
+
+				if (bBehaviorsEnabled)
 				{
-					SetValueOfController(Controller, ControllerValue->Value);
-					GetEntitiesControlledByController(RemoteControlPreset, Controller, AppliedEntities);
+					GetEntitiesControlledByController(RemoteControlPreset, Controller, AppliedOrIgnoredEntities);
 					UE_LOG(LogAvaPlayable, Verbose, TEXT("%s Playable {%s} -> Updating Controller %s (with behaviors)."),
 						*GetBriefFrameInfo(), *GetPrettyPlayableInfo(this), *Controller->DisplayName.ToString());
 				}
-				else
-				{
-					FScopedPushControllerBehavioursEnable PushBehaviorsEnable(Controller, /*bInBehaviorsEnabled*/ false);
-					SetValueOfController(Controller, ControllerValue->Value);
-				}
 			}
 		}
 	}
@@ -354,36 +389,7 @@ EAvaPlayableCommandResult UAvaPlayable::UpdateRemoteControlCommand(const TShared
 	// We assume the controller actions are already executed in the rundown's managed RCP
 	// during page edition and the resulting entity values are already captured.
 
-	// Todo: change this function to support filter predicate.
-	//InRemoteControlValues->ApplyEntityValuesToRemoteControlPreset(RemoteControlPreset);
-
-	for (const TWeakPtr<FRemoteControlEntity>& EntityWeakPtr : RemoteControlPreset->GetExposedEntities<FRemoteControlEntity>())
-	{
-		if (const TSharedPtr<FRemoteControlEntity> Entity = EntityWeakPtr.Pin())
-		{
-			// Skip properties that are already set by controllers at runtime.
-			if (AppliedEntities.Contains(Entity->GetId()))
-			{
-				continue;
-			}
-			
-			if (const FAvaPlayableRemoteControlValue* EntityValue = InRemoteControlValues->GetEntityValue(Entity->GetId()))
-			{
-				using namespace UE::AvaPlayableRemoteControl;
-				const EAvaPlayableRemoteControlResult Result = SetValueOfEntity(Entity, EntityValue->Value);
-				if (Failed(Result))
-				{
-					UE_LOG(LogAvaPlayable, Error, TEXT("Failed to set value of exposed entity \"%s\" (id:%s): %s."),
-						*Entity->GetLabel().ToString(), *Entity->GetId().ToString(), *EnumToString(Result));
-				}
-			}
-			else
-			{
-				UE_LOG(LogAvaPlayable, Error, TEXT("Exposed entity \"%s\" (id:%s): value not found in page."),
-					*Entity->GetLabel().ToString(), *Entity->GetId().ToString());
-			}
-		}
-	}
+	InRemoteControlValues->ApplyEntityValuesToRemoteControlPreset(RemoteControlPreset, AppliedOrIgnoredEntities);
 
 	LatestRemoteControlValues = InRemoteControlValues;
 	OnRemoteControlValuesApplied();
diff --git a/Engine/Plugins/Experimental/Avalanche/Source/AvalancheMedia/Private/Playable/AvaPlayableRemoteControl.cpp b/Engine/Plugins/Experimental/Avalanche/Source/AvalancheMedia/Private/Playable/AvaPlayableRemoteControl.cpp
--- a/Engine/Plugins/Experimental/Avalanche/Source/AvalancheMedia/Private/Playable/AvaPlayableRemoteControl.cpp
+++ b/Engine/Plugins/Experimental/Avalanche/Source/AvalancheMedia/Private/Playable/AvaPlayableRemoteControl.cpp
@@ -105,14 +105,57 @@ namespace UE::AvaPlayableRemoteControl::Private
 		}
 		return false;
 	}
+
+	inline bool HasAccess(const TSharedPtr<const FRemoteControlEntity>& InRemoteControlEntity, const ERCAccess InAccess, FString& OutErrorText)
+	{
+		using namespace UE::AvaPlayableRemoteControl::Private;
+
+		if (!InRemoteControlEntity.IsValid())
+		{
+			OutErrorText = TEXT("Entity is null"); 
+			return false;
+		}
+
+		TSharedPtr<const FRemoteControlProperty> Field = CastEntity<FRemoteControlProperty>(InRemoteControlEntity);
+		if (!Field.IsValid())
+		{
+			OutErrorText = FString::Printf(TEXT("Wrong Entity type \"%s\", expected \"%s\"."), *GetNameSafe(InRemoteControlEntity->GetStruct()), *GetNameSafe(FRemoteControlProperty::StaticStruct())); 
+			return false;
+		}
+		
+		if (UObject* FieldBoundObject = Field->GetBoundObject(); IsValid(FieldBoundObject))
+		{
+			FRCObjectReference DummyObjectRef;
+			if (InAccess == ERCAccess::READ_ACCESS)
+			{
+				return	ResolveObjectPropertyForReadOnly(FieldBoundObject, Field->FieldPathInfo, DummyObjectRef, OutErrorText);
+			}
+			return IRemoteControlModule::Get().ResolveObjectProperty(InAccess, FieldBoundObject, Field->FieldPathInfo, DummyObjectRef, &OutErrorText);
+		}
+		
+		OutErrorText = TEXT("Invalid bound object");
+		return false;
+	}
+}
+
+bool UE::AvaPlayableRemoteControl::HasReadAccess(const TSharedPtr<const FRemoteControlEntity>& InRemoteControlEntity, FString& OutErrorText)
+{
+	using namespace UE::AvaPlayableRemoteControl::Private;
+	return HasAccess(InRemoteControlEntity, ERCAccess::READ_ACCESS, OutErrorText);
+}
+
+bool UE::AvaPlayableRemoteControl::HasWriteAccess(const TSharedPtr<const FRemoteControlEntity>& InRemoteControlEntity, FString& OutErrorText)
+{
+	using namespace UE::AvaPlayableRemoteControl::Private;
+	return HasAccess(InRemoteControlEntity, ERCAccess::WRITE_ACCESS, OutErrorText);
 }
 
-EAvaPlayableRemoteControlResult UE::AvaPlayableRemoteControl::GetValueOfEntity(const TSharedPtr<const FRemoteControlEntity>& RemoteControlEntity, TArray<uint8>& OutValue)
+EAvaPlayableRemoteControlResult UE::AvaPlayableRemoteControl::GetValueOfEntity(const TSharedPtr<const FRemoteControlEntity>& InRemoteControlEntity, TArray<uint8>& OutValue)
 {
 	using namespace UE::AvaPlayableRemoteControl::Private;
 	OutValue.Reset();
 
-	TSharedPtr<const FRemoteControlProperty> Field = CastEntity<FRemoteControlProperty>(RemoteControlEntity);
+	TSharedPtr<const FRemoteControlProperty> Field = CastEntity<FRemoteControlProperty>(InRemoteControlEntity);
 	if (!Field.IsValid())
 	{
 		return EAvaPlayableRemoteControlResult::InvalidParameter;
@@ -136,11 +179,11 @@ EAvaPlayableRemoteControlResult UE::AvaPlayableRemoteControl::GetValueOfEntity(c
 	return Result;
 }
 
-EAvaPlayableRemoteControlResult UE::AvaPlayableRemoteControl::SetValueOfEntity(const TSharedPtr<FRemoteControlEntity>& RemoteControlEntity, const TArrayView<const uint8>& InValue)
+EAvaPlayableRemoteControlResult UE::AvaPlayableRemoteControl::SetValueOfEntity(const TSharedPtr<FRemoteControlEntity>& InRemoteControlEntity, const TArrayView<const uint8>& InValue)
 {
 	using namespace UE::AvaPlayableRemoteControl::Private;
 
-	TSharedPtr<FRemoteControlProperty> Field = CastEntity<FRemoteControlProperty>(RemoteControlEntity);
+	TSharedPtr<FRemoteControlProperty> Field = CastEntity<FRemoteControlProperty>(InRemoteControlEntity);
 	if (!Field.IsValid())
 	{
 		return EAvaPlayableRemoteControlResult::InvalidParameter;
@@ -272,6 +315,17 @@ EAvaPlayableRemoteControlResult UE::AvaPlayableRemoteControl::SetValueOfControll
 	return SetValueOfController(InController, AvaMediaSerializationUtils::JsonValueConversion::ValueToConstBytesView(InValue));
 }
 
+EAvaPlayableRemoteControlResult UE::AvaPlayableRemoteControl::SetValueOfController(URCVirtualPropertyBase* InController, const FString& InValue, bool bInBehaviorsEnabled)
+{
+	if (bInBehaviorsEnabled)
+	{
+		return SetValueOfController(InController, InValue);
+	}
+	
+	FScopedPushControllerBehavioursEnable PushBehavioursEnable(InController, false);
+	return SetValueOfController(InController, InValue);
+}
+
 bool UE::AvaPlayableRemoteControl::GetEntitiesControlledByController(const URemoteControlPreset* InRemoteControlPreset, const URCVirtualPropertyBase* InVirtualProperty, TSet<FGuid>& OutEntityIds)
 {
 	const URCController* Controller = Cast<URCController>(InVirtualProperty);
diff --git a/Engine/Plugins/Experimental/Avalanche/Source/AvalancheMedia/Private/Playable/AvaPlayableRemoteControlPresetInfo.cpp b/Engine/Plugins/Experimental/Avalanche/Source/AvalancheMedia/Private/Playable/AvaPlayableRemoteControlPresetInfo.cpp
new file mode 100644
--- /dev/null
+++ b/Engine/Plugins/Experimental/Avalanche/Source/AvalancheMedia/Private/Playable/AvaPlayableRemoteControlPresetInfo.cpp
@@ -0,0 +1,64 @@
+// Copyright Epic Games, Inc. All Rights Reserved.
+
+#include "Playable/AvaPlayableRemoteControlPresetInfo.h"
+
+#include "Playable/AvaPlayableRemoteControl.h"
+#include "Playable/AvaPlayableRemoteControlValues.h"
+#include "RCVirtualProperty.h"
+
+void FAvaPlayableRemoteControlPresetInfo::Refresh(const URemoteControlPreset* InRemoteControlPreset)
+{
+	EntitiesControlledByController.Reset();
+	OverlappingControllers.Reset();
+
+	if (IsValid(InRemoteControlPreset))
+	{
+		PresetId = InRemoteControlPreset->GetPresetId();
+		
+		// Set of entities for the current controller only.
+		TSet<FGuid> TmpControlledEntities;
+		TmpControlledEntities.Reserve(EntitiesControlledByController.GetAllocatedSize());
+		
+		const TArray<URCVirtualPropertyBase*> Controllers = InRemoteControlPreset->GetControllers();
+		for (const URCVirtualPropertyBase* Controller : Controllers)
+		{
+			if (!Controller)
+			{
+				continue;
+			}
+			
+			TmpControlledEntities.Reset();
+			
+			if (!UE::AvaPlayableRemoteControl::GetEntitiesControlledByController(InRemoteControlPreset, Controller, TmpControlledEntities))
+			{
+				UE_LOG(LogAvaPlayableRemoteControl, Warning, TEXT("Failed to get controlled entities for controller \"%s\" (id:%s)."),
+					*Controller->DisplayName.ToString(), *Controller->Id.ToString());
+			}
+
+			if (!TmpControlledEntities.IsEmpty())
+			{
+				// Merge in the controlled entity info map
+				for (const FGuid& EntityId : TmpControlledEntities)
+				{
+					auto& EntityInfo = EntitiesControlledByController.FindOrAdd(EntityId);
+
+					// Remark: we don't consider overlap with ignored controllers.
+					if (!FAvaPlayableRemoteControlValues::ShouldIgnoreController(Controller))
+					{
+						EntityInfo.ControlledBy.AddUnique(Controller->Id);	
+						if (EntityInfo.ControlledBy.Num() > 1)
+						{
+							OverlappingControllers.Append(EntityInfo.ControlledBy);
+						}
+					}
+				}
+			}
+		}
+	}
+	else
+	{
+		PresetId.Invalidate();
+	}
+
+	bDirty = false;
+}
\ No newline at end of file
diff --git a/Engine/Plugins/Experimental/Avalanche/Source/AvalancheMedia/Private/Playable/AvaPlayableRemoteControlPresetInfoCache.cpp b/Engine/Plugins/Experimental/Avalanche/Source/AvalancheMedia/Private/Playable/AvaPlayableRemoteControlPresetInfoCache.cpp
new file mode 100644
--- /dev/null
+++ b/Engine/Plugins/Experimental/Avalanche/Source/AvalancheMedia/Private/Playable/AvaPlayableRemoteControlPresetInfoCache.cpp
@@ -0,0 +1,144 @@
+// Copyright Epic Games, Inc. All Rights Reserved.
+
+#include "AvaPlayableRemoteControlPresetInfoCache.h"
+
+#include "AssetRegistry/AssetData.h"
+#include "IAvaMediaModule.h"
+#include "UObject/ObjectSaveContext.h"
+#include "UObject/Package.h"
+
+IAvaPlayableRemoteControlPresetInfoCache& IAvaPlayableRemoteControlPresetInfoCache::Get()
+{
+	return IAvaMediaModule::Get().GetPlayableRemoteControlPresetInfoCache();
+}
+
+TSharedPtr<FAvaPlayableRemoteControlPresetInfo> FAvaPlayableRemoteControlPresetInfoCache::GetRemoteControlPresetInfo(const FSoftObjectPath& InAssetPath, const URemoteControlPreset* InRemoteControlPreset)
+{
+	if (InAssetPath.IsNull())
+	{
+		return nullptr;
+	}
+
+	OrderQueue.Remove(InAssetPath); // Removes preserves the order. O(n)
+	OrderQueue.Add(InAssetPath); // Most recent is at the end of the array.
+
+	if (TSharedPtr<FAvaPlayableRemoteControlPresetInfo>* ExistingEntry = PresetInfoCache.Find(InAssetPath))
+	{
+		if (ExistingEntry->IsValid())
+		{
+			if ((*ExistingEntry)->IsDirty())
+			{
+				(*ExistingEntry)->Refresh(InRemoteControlPreset);
+			}
+			return *ExistingEntry;
+		}
+	}
+
+	TSharedPtr<FAvaPlayableRemoteControlPresetInfo> NewPresetInfo = MakeShared<FAvaPlayableRemoteControlPresetInfo>();
+	NewPresetInfo->Refresh(InRemoteControlPreset);
+
+	PresetInfoCache.Add(InAssetPath, NewPresetInfo);
+	TrimCache();
+	
+	return NewPresetInfo;
+}
+
+void FAvaPlayableRemoteControlPresetInfoCache::Flush(const FSoftObjectPath& InAssetPath)
+{
+	if (const TSharedPtr<FAvaPlayableRemoteControlPresetInfo>* ExistingEntry = PresetInfoCache.Find(InAssetPath))
+	{
+		if (ExistingEntry->GetSharedReferenceCount() <= 1)
+		{
+			RemoveEntry(InAssetPath);
+		}
+	}
+}
+
+void FAvaPlayableRemoteControlPresetInfoCache::Flush()
+{
+	RemoveEntries([](const FSoftObjectPath& InAssetPath, const TSharedPtr<FAvaPlayableRemoteControlPresetInfo>& InPresetInfo)
+	{
+		return InPresetInfo.GetSharedReferenceCount() <= 1;
+	});
+}
+
+int32 FAvaPlayableRemoteControlPresetInfoCache::GetMaximumCacheSize()
+{
+	return 100;	// This is not very critical, we just don't want an infinite size cache.
+}
+
+void FAvaPlayableRemoteControlPresetInfoCache::TrimCache()
+{
+	const int32 MaximumCacheSize = GetMaximumCacheSize(); 
+	if (MaximumCacheSize > 0)
+	{
+		while (OrderQueue.Num() > MaximumCacheSize)
+		{
+			// LRU: oldest is at the start of the array.
+			PresetInfoCache.Remove(OrderQueue[0]);
+			OrderQueue.RemoveAt(0);
+		}
+	}
+}
+
+void FAvaPlayableRemoteControlPresetInfoCache::RemoveEntry(const FSoftObjectPath& InAssetPath)
+{
+	OrderQueue.Remove(InAssetPath);
+	PresetInfoCache.Remove(InAssetPath);
+}
+
+void FAvaPlayableRemoteControlPresetInfoCache::RemoveEntries(TFunctionRef<bool(const FSoftObjectPath&, const TSharedPtr<FAvaPlayableRemoteControlPresetInfo>&)> InRemovePredicate)
+{
+	for (TMap<FSoftObjectPath, TSharedPtr<FAvaPlayableRemoteControlPresetInfo>>::TIterator EntryIt = PresetInfoCache.CreateIterator(); EntryIt; ++EntryIt)
+	{
+		if (InRemovePredicate(EntryIt.Key(), EntryIt.Value()))
+		{
+			OrderQueue.Remove(EntryIt.Key());
+			EntryIt.RemoveCurrent();
+		}
+	}
+}
+
+void FAvaPlayableRemoteControlPresetInfoCache::OnPackageSaved(const FString& InPackageFileName, UPackage* InPackage, FObjectPostSaveContext InObjectSaveContext)
+{
+	if (InObjectSaveContext.IsProceduralSave())
+	{
+		return;
+	}
+
+	OnPackageModified(InPackage->GetFName());
+}
+
+void FAvaPlayableRemoteControlPresetInfoCache::OnAvaSyncPackageModified(IAvaMediaSyncProvider* InAvaMediaSyncProvider, const FName& InPackageName)
+{
+	UE_LOG(LogAvaMedia, Verbose,
+		TEXT("A sync operation has touched the package \"%s\" on disk. Remote Control Preset Info Cache notified."),
+		*InPackageName.ToString());
+
+	OnPackageModified(InPackageName);
+}
+
+void FAvaPlayableRemoteControlPresetInfoCache::OnAssetRemoved(const FAssetData& InAssetData)
+{
+	// Remove entry for the given package.
+	RemoveEntries([PackageName = InAssetData.PackageName](const FSoftObjectPath& InAssetPath, const TSharedPtr<FAvaPlayableRemoteControlPresetInfo>& InPresetInfo)
+	{
+		return InAssetPath.GetLongPackageFName() == PackageName;
+	});
+}
+
+void FAvaPlayableRemoteControlPresetInfoCache::OnPackageModified(const FName& InPackageName)
+{
+	// Invalidate corresponding assets from that package.
+	for (const TPair<FSoftObjectPath, TSharedPtr<FAvaPlayableRemoteControlPresetInfo>>& Entry : PresetInfoCache)
+	{
+		if (Entry.Key.GetLongPackageFName() == InPackageName && Entry.Value.IsValid())
+		{
+			UE_LOG(LogAvaMedia, Verbose,
+				TEXT("Remote Control Preset Info Cache: Package \"%s\" being touched caused asset \"%s\" to be invalidated."),
+				*InPackageName.ToString(), *Entry.Key.ToString());
+
+			Entry.Value->MarkDirty();
+		}
+	}
+}
\ No newline at end of file
diff --git a/Engine/Plugins/Experimental/Avalanche/Source/AvalancheMedia/Private/Playable/AvaPlayableRemoteControlPresetInfoCache.h b/Engine/Plugins/Experimental/Avalanche/Source/AvalancheMedia/Private/Playable/AvaPlayableRemoteControlPresetInfoCache.h
new file mode 100644
--- /dev/null
+++ b/Engine/Plugins/Experimental/Avalanche/Source/AvalancheMedia/Private/Playable/AvaPlayableRemoteControlPresetInfoCache.h
@@ -0,0 +1,46 @@
+// Copyright Epic Games, Inc. All Rights Reserved.
+
+#pragma once
+
+#include "Playable/AvaPlayableRemoteControlPresetInfo.h"
+
+class IAvaMediaSyncProvider;
+class UPackage;
+
+/**
+ * Implementation of a RCP info cache.
+ * Similar to FAvaRundownManagedInstanceCache in terms of invalidation.
+ * Given this is relatively lightweight auxiliary data, we could mark as dirty instead of deleting.
+ */
+class FAvaPlayableRemoteControlPresetInfoCache final : public IAvaPlayableRemoteControlPresetInfoCache
+{
+public:
+	//~ Begin IAvaPlayableRemoteControlPresetInfoCache
+	virtual TSharedPtr<FAvaPlayableRemoteControlPresetInfo> GetRemoteControlPresetInfo(const FSoftObjectPath& InAssetPath, const URemoteControlPreset* InRemoteControlPreset) override;
+	virtual void Flush(const FSoftObjectPath& InAssetPath) override;
+	virtual void Flush() override;
+	//~ End IAvaPlayableRemoteControlPresetInfoCache
+
+	/**
+	 * Get the maximum size of the cache beyond witch it will start flushing elements.
+	 */
+	static int32 GetMaximumCacheSize();
+	
+	/**
+	 * Trim the cache elements that exceed the cache capacity according to LRU replacement policy.
+	 */
+	void TrimCache();
+
+private:
+	void RemoveEntry(const FSoftObjectPath& InAssetPath);
+	void RemoveEntries(TFunctionRef<bool(const FSoftObjectPath&, const TSharedPtr<FAvaPlayableRemoteControlPresetInfo>&)> InRemovePredicate);
+
+	void OnPackageSaved(const FString& InPackageFileName, UPackage* InPackage, FObjectPostSaveContext InObjectSaveContext);
+	void OnAvaSyncPackageModified(IAvaMediaSyncProvider* InAvaMediaSyncProvider, const FName& InPackageName);
+	void OnAssetRemoved(const FAssetData& InAssetData);
+	void OnPackageModified(const FName& InPackageName);
+
+private:
+	TMap<FSoftObjectPath, TSharedPtr<FAvaPlayableRemoteControlPresetInfo>> PresetInfoCache;
+	TArray<FSoftObjectPath> OrderQueue;
+};
\ No newline at end of file
diff --git a/Engine/Plugins/Experimental/Avalanche/Source/AvalancheMedia/Private/Playable/AvaPlayableRemoteControlValues.cpp b/Engine/Plugins/Experimental/Avalanche/Source/AvalancheMedia/Private/Playable/AvaPlayableRemoteControlValues.cpp
--- a/Engine/Plugins/Experimental/Avalanche/Source/AvalancheMedia/Private/Playable/AvaPlayableRemoteControlValues.cpp
+++ b/Engine/Plugins/Experimental/Avalanche/Source/AvalancheMedia/Private/Playable/AvaPlayableRemoteControlValues.cpp
@@ -3,11 +3,13 @@
 #include "Playable/AvaPlayableRemoteControlValues.h"
 
 #include "AvaMediaSerializationUtils.h"
+#include "IAvaMediaModule.h"
 #include "Dom/JsonObject.h"
 #include "Misc/PackageName.h"
 #include "Playable/AvaPlayableRemoteControl.h"
 #include "Playable/AvaPlayableRemoteControlValuesPrivate.h"
 #include "RCVirtualProperty.h"
+#include "Playable/AvaPlayableSettings.h"
 #include "Serialization/CustomVersion.h"
 #include "Serialization/JsonReader.h"
 #include "Serialization/JsonSerializer.h"
@@ -204,24 +206,6 @@ bool FAvaPlayableRemoteControlValue::Serialize(FArchive& Ar)
 	return true;
 }
 
-void FAvaPlayableRemoteControlValues::RefreshControlledEntities(const URemoteControlPreset* InRemoteControlPreset)
-{
-	EntitiesControlledByController.Reset();
-
-	if (IsValid(InRemoteControlPreset))
-	{
-		const TArray<URCVirtualPropertyBase*> Controllers = InRemoteControlPreset->GetControllers();
-		for (const URCVirtualPropertyBase* PropertyBase : Controllers)
-		{
-			if (!UE::AvaPlayableRemoteControl::GetEntitiesControlledByController(InRemoteControlPreset, PropertyBase, EntitiesControlledByController))
-			{
-				UE_LOG(LogAvaPlayableRemoteControl, Warning, TEXT("Failed to get controlled entities for controller \"%s\" (id:%s)."),
-					*PropertyBase->DisplayName.ToString(), *PropertyBase->Id.ToString());
-			}
-		}
-	}
-}
-
 void FAvaPlayableRemoteControlValues::CopyFrom(const URemoteControlPreset* InRemoteControlPreset, bool bInIsDefault)
 {
 	using namespace UE::AvaPlayableRemoteControl;
@@ -229,8 +213,6 @@ void FAvaPlayableRemoteControlValues::CopyFrom(const URemoteControlPreset* InRem
 	EntityValues.Reset();
 	ControllerValues.Reset();
 
-	RefreshControlledEntities(InRemoteControlPreset);	// will reset ids if invalid.
-	
 	if (IsValid(InRemoteControlPreset))
 	{
 		FString ValueAsString;
@@ -386,7 +368,7 @@ bool FAvaPlayableRemoteControlValues::SetControllerValue(const FGuid& InId, cons
 	return true;
 }
 
-void FAvaPlayableRemoteControlValues::ApplyEntityValuesToRemoteControlPreset(URemoteControlPreset* InRemoteControlPreset) const
+void FAvaPlayableRemoteControlValues::ApplyEntityValuesToRemoteControlPreset(URemoteControlPreset* InRemoteControlPreset, const TSet<FGuid>& InSkipEntities) const
 {
 	if (!InRemoteControlPreset)
 	{
@@ -397,6 +379,13 @@ void FAvaPlayableRemoteControlValues::ApplyEntityValuesToRemoteControlPreset(URe
 	{
 		if (const TSharedPtr<FRemoteControlEntity> Entity = EntityWeakPtr.Pin())
 		{
+			if (InSkipEntities.Contains(Entity->GetId()))
+			{
+				UE_LOG(LogAvaPlayableRemoteControl, Verbose, TEXT("Skipping exposed entity \"%s\" (id:%s)."),
+					*Entity->GetLabel().ToString(), *Entity->GetId().ToString());
+				continue;
+			}
+
 			if (const FAvaPlayableRemoteControlValue* Value = GetEntityValue(Entity->GetId()))
 			{
 				const EAvaPlayableRemoteControlResult Result = SetValueOfEntity(Entity, Value->Value); 
@@ -408,8 +397,17 @@ void FAvaPlayableRemoteControlValues::ApplyEntityValuesToRemoteControlPreset(URe
 			}
 			else
 			{
-				UE_LOG(LogAvaPlayableRemoteControl, Error, TEXT("Exposed entity \"%s\" (id:%s): value not found in page."),
-					*Entity->GetLabel().ToString(), *Entity->GetId().ToString());
+				FString AccessError;
+				if (!HasReadAccess(Entity, AccessError))
+				{
+					UE_LOG(LogAvaPlayableRemoteControl, Error, TEXT("Exposed entity \"%s\" (id:%s): value not found in page. Reason: %s."),
+						*Entity->GetLabel().ToString(), *Entity->GetId().ToString(), *AccessError);
+				}
+				else
+				{
+					UE_LOG(LogAvaPlayableRemoteControl, Error, TEXT("Exposed entity \"%s\" (id:%s): value not found in page."),
+						*Entity->GetLabel().ToString(), *Entity->GetId().ToString());
+				}
 			}
 		}
 	}
@@ -427,16 +425,7 @@ void FAvaPlayableRemoteControlValues::ApplyControllerValuesToRemoteControlPreset
 		if (const FAvaPlayableRemoteControlValue* Value = GetControllerValue(Controller->Id))
 		{
 			using namespace UE::AvaPlayableRemoteControl;
-			EAvaPlayableRemoteControlResult Result;
-			if (bInForceDisableBehaviors)
-			{
-				FScopedPushControllerBehavioursEnable PushBehavioursEnable(Controller, false);
-				Result = SetValueOfController(Controller, Value->Value);
-			}
-			else
-			{
-				Result = SetValueOfController(Controller, Value->Value);
-			}
+			EAvaPlayableRemoteControlResult Result = SetValueOfController(Controller, Value->Value, /*bInBehaviorsEnabled*/ !bInForceDisableBehaviors);
 			if (Failed(Result))
 			{
 				UE_LOG(LogAvaPlayableRemoteControl, Error, TEXT("Failed to set virtual value of controller \"%s\" (id:%s): %s."),
@@ -463,8 +452,7 @@ bool FAvaPlayableRemoteControlValues::Merge(const FAvaPlayableRemoteControlValue
 	const bool bHasIdCollisions = HasIdCollisions(InOtherValues);
 
 	ControllerValues.Append(InOtherValues.ControllerValues);
-	EntityValues.Append(InOtherValues.EntityValues);		
-	EntitiesControlledByController.Append(InOtherValues.EntitiesControlledByController);
+	EntityValues.Append(InOtherValues.EntityValues);
 
 	return !bHasIdCollisions;
 }
@@ -507,3 +495,22 @@ void FAvaPlayableRemoteControlValues::CollectReferencedAssetPaths(const TMap<FGu
 		}
 	}
 }
+
+bool FAvaPlayableRemoteControlValues::ShouldIgnoreController(const URCVirtualPropertyBase* InController)
+{
+	if (InController)
+	{
+		const FAvaPlayableSettings& PlayableSettings = IAvaMediaModule::Get().GetPlayableSettings();
+		// Current way to identify an ignored controller is by using a postfix in the name.
+		// Other options would be to check meta data. (todo)
+		const FString ControllerName = InController->DisplayName.ToString(); 
+		for (const FString& PostFixToIgnore : PlayableSettings.IgnoredControllerPostfix)
+		{
+			if (ControllerName.EndsWith(PostFixToIgnore))
+			{
+				return true;
+			}
+		}
+	}
+	return false;
+}
\ No newline at end of file
diff --git a/Engine/Plugins/Experimental/Avalanche/Source/AvalancheMedia/Private/Rundown/AvaRundownManagedInstance.cpp b/Engine/Plugins/Experimental/Avalanche/Source/AvalancheMedia/Private/Rundown/AvaRundownManagedInstance.cpp
--- a/Engine/Plugins/Experimental/Avalanche/Source/AvalancheMedia/Private/Rundown/AvaRundownManagedInstance.cpp
+++ b/Engine/Plugins/Experimental/Avalanche/Source/AvalancheMedia/Private/Rundown/AvaRundownManagedInstance.cpp
@@ -2,6 +2,7 @@
 
 #include "Rundown/AvaRundownManagedInstance.h"
 
+#include "Playable/AvaPlayableRemoteControlPresetInfoCache.h"
 #include "RemoteControlPreset.h"
 #include "Rundown/AvaRundownManagedInstanceCache.h"
 
@@ -24,6 +25,18 @@ void FAvaRundownManagedInstance::InvalidateFromCache()
 	}
 }
 
+FAvaPlayableRemoteControlPresetInfo& FAvaRundownManagedInstance::GetRemoteControlPresetInfoFromCache(URemoteControlPreset* InRemoteControlPreset) const
+{
+	RemoteControlPresetInfo = IAvaPlayableRemoteControlPresetInfoCache::Get().GetRemoteControlPresetInfo(GetSourceAssetPath(), InRemoteControlPreset);
+	if (RemoteControlPresetInfo.IsValid())
+	{
+		return *RemoteControlPresetInfo;
+	}
+
+	static FAvaPlayableRemoteControlPresetInfo EmptyPresetInfo;
+	return EmptyPresetInfo;
+}
+
 void FAvaRundownManagedInstance::RegisterSourceRemoteControlPresetDelegates(URemoteControlPreset* InSourceRemoteControlPreset)
 {
 	if (!::IsValid(InSourceRemoteControlPreset))
@@ -91,7 +104,14 @@ void FAvaRundownManagedInstance::OnSourceRemoteControlExposedPropertiesModified(
 	// can be updated.
 
 	// There is no event called when entities are unbound to actions, so rescan all actions just in case.
-	DefaultRemoteControlValues.RefreshControlledEntities(InSourcePreset);
+	if (RemoteControlPresetInfo.IsValid())
+	{
+		RemoteControlPresetInfo->Refresh(InSourcePreset);
+	}
+	else
+	{
+		GetRemoteControlPresetInfoFromCache(InSourcePreset);
+	}
 
 	for (const FGuid& PropertyId : InModifiedProperties)
 	{
@@ -126,6 +146,7 @@ void FAvaRundownManagedInstance::OnSourceRemoteControlControllerModified(URemote
 	// We don't want to apply those to the managed remote control because we don't want to overwrite the page's values.
 	// We probably need an extra event to propagate a change in default value so pages with no values for that property
 	// can be updated.
+	// Todo: Need to hookup some reimport notification for all templates/pages using this asset.
 	for (const FGuid& ControllerId : InModifiedControllerIds)
 	{
 		// Update existing values only.
diff --git a/Engine/Plugins/Experimental/Avalanche/Source/AvalancheMedia/Private/Rundown/AvaRundownPage.cpp b/Engine/Plugins/Experimental/Avalanche/Source/AvalancheMedia/Private/Rundown/AvaRundownPage.cpp
--- a/Engine/Plugins/Experimental/Avalanche/Source/AvalancheMedia/Private/Rundown/AvaRundownPage.cpp
+++ b/Engine/Plugins/Experimental/Avalanche/Source/AvalancheMedia/Private/Rundown/AvaRundownPage.cpp
@@ -168,11 +168,14 @@ bool FAvaRundownPage::UpdatePageSummary(const UAvaRundown* InRundown, const TArr
 	{
 		for (const URCVirtualPropertyBase* VirtualProperty : Preset->GetControllers())
 		{
-			FString StringValue = VirtualProperty->GetDisplayValueAsString();
-			const FText TextToAdd = FText::FromString(StringValue);
-			if (!TextToAdd.IsEmptyOrWhitespace())
+			if (!FAvaPlayableRemoteControlValues::ShouldIgnoreController(VirtualProperty))
 			{
-				OrderedDescription.FindOrAdd(VirtualProperty->DisplayIndex + DisplayIndexOffset) = TextToAdd;
+				FString StringValue = VirtualProperty->GetDisplayValueAsString();
+				const FText TextToAdd = FText::FromString(StringValue);
+				if (!TextToAdd.IsEmptyOrWhitespace())
+				{
+					OrderedDescription.FindOrAdd(VirtualProperty->DisplayIndex + DisplayIndexOffset) = TextToAdd;
+				}
 			}
 		}
 		DisplayIndexOffset += Preset->GetNumControllers();
diff --git a/Engine/Plugins/Experimental/Avalanche/Source/AvalancheMedia/Public/IAvaMediaModule.h b/Engine/Plugins/Experimental/Avalanche/Source/AvalancheMedia/Public/IAvaMediaModule.h
--- a/Engine/Plugins/Experimental/Avalanche/Source/AvalancheMedia/Public/IAvaMediaModule.h
+++ b/Engine/Plugins/Experimental/Avalanche/Source/AvalancheMedia/Public/IAvaMediaModule.h
@@ -15,6 +15,7 @@ class FName;
 class IAvaBroadcastDeviceProviderProxyManager;
 class IAvaBroadcastSettings;
 class IAvaMediaSyncProvider;
+class IAvaPlayableRemoteControlPresetInfoCache;
 class IAvaPlaybackClient;
 class IAvaPlaybackServer;
 class IAvaRundownServer;
@@ -227,4 +228,7 @@ public:
 	 * Access the device provider proxy manager.
 	 */
 	virtual IAvaBroadcastDeviceProviderProxyManager& GetDeviceProviderProxyManager() = 0;
+
+	/** Access the global remote control preset info cache. */
+	virtual IAvaPlayableRemoteControlPresetInfoCache& GetPlayableRemoteControlPresetInfoCache() const = 0;
 };
diff --git a/Engine/Plugins/Experimental/Avalanche/Source/AvalancheMedia/Public/Playable/AvaPlayableRemoteControlPresetInfo.h b/Engine/Plugins/Experimental/Avalanche/Source/AvalancheMedia/Public/Playable/AvaPlayableRemoteControlPresetInfo.h
new file mode 100644
--- /dev/null
+++ b/Engine/Plugins/Experimental/Avalanche/Source/AvalancheMedia/Public/Playable/AvaPlayableRemoteControlPresetInfo.h
@@ -0,0 +1,112 @@
+// Copyright Epic Games, Inc. All Rights Reserved.
+
+#pragma once
+
+#include "CoreTypes.h"
+#include "AvaPlayableRemoteControlPresetInfo.generated.h"
+
+class URemoteControlPreset;
+struct FSoftObjectPath;
+
+/**
+ * Container for the controlled entities extra information.
+ */
+USTRUCT()
+struct FAvaPlayableRemoteControlControlledEntityInfo
+{
+	GENERATED_BODY()
+
+	/** Controller ids that are controlling the entity. */
+	UPROPERTY()
+	TArray<FGuid> ControlledBy;
+};
+
+/**
+ * Container for the remote control preset extra information used for playable.
+ */
+USTRUCT()
+struct FAvaPlayableRemoteControlPresetInfo
+{
+	GENERATED_BODY()
+
+	FAvaPlayableRemoteControlPresetInfo() = default;
+
+	/**
+	 * Refreshes from given remote control preset.
+	 * @remark will reset bDirty flag.
+	 */
+	void Refresh(const URemoteControlPreset* InRemoteControlPreset);
+
+	/**
+	 * Call this whenever the source RCP is modified (or is likely to be) to
+	 * invalidate the information and trigger a refresh on the next access.
+	 */
+	void MarkDirty()
+	{
+		bDirty = true;
+	}
+
+	/** Indicates if the data needs to be refreshed. */
+	bool IsDirty() const
+	{
+		return bDirty;
+	}
+
+	/**
+	 * Returns true if the given controller is considered "overlapping", i.e. it's set of controlled entities
+	 * overlap with other controllers. An overlapping controller can't safely be updating its behaviors when the end
+	 * result needs to be deterministic.
+	 */
+	bool IsControllerOverlapping(const FGuid& InControllerId) const
+	{
+		return OverlappingControllers.Contains(InControllerId);
+	}
+
+	/** Id of the RemoteControlPreset */
+	UPROPERTY()
+	FGuid PresetId;
+	
+	/** Contains a set of entity guids that are bound to a controller action. */
+	UPROPERTY()
+	TMap<FGuid, FAvaPlayableRemoteControlControlledEntityInfo> EntitiesControlledByController;
+
+	/**
+	 * Contains a set of controllers that are "overlapping". 
+	 * The term "overlapping controllers" defines the scenario where multiple controllers influence
+	 * a shared subset of controlled entities.
+	 */ 
+	UPROPERTY()
+	TSet<FGuid> OverlappingControllers;
+
+protected:
+	/** When the source asset is modified (or likely to be), we mark as dirty to trigger a refresh on next access. */
+	bool bDirty = true;
+};
+
+/**
+ * Global Cache for remote control preset info 
+ */
+class IAvaPlayableRemoteControlPresetInfoCache
+{
+public:
+	static AVALANCHEMEDIA_API IAvaPlayableRemoteControlPresetInfoCache& Get();
+
+	/**
+	 * Request the RCP info cached for the given asset path.
+	 * If not available, will be created from the given RCP.
+	 */
+	virtual TSharedPtr<FAvaPlayableRemoteControlPresetInfo> GetRemoteControlPresetInfo(const FSoftObjectPath& InAssetPath, const URemoteControlPreset* InRemoteControlPreset) = 0;
+
+	/**
+	 *	Flush specified entry from the cache.
+	 */
+	virtual void Flush(const FSoftObjectPath& InAssetPath) = 0;
+
+	/**
+	 * Flush all unused entries from the cache.
+	 */
+	virtual void Flush() = 0;
+	
+protected:
+	virtual ~IAvaPlayableRemoteControlPresetInfoCache() = default;
+};
\ No newline at end of file
diff --git a/Engine/Plugins/Experimental/Avalanche/Source/AvalancheMedia/Public/Playable/AvaPlayableRemoteControlValues.h b/Engine/Plugins/Experimental/Avalanche/Source/AvalancheMedia/Public/Playable/AvaPlayableRemoteControlValues.h
--- a/Engine/Plugins/Experimental/Avalanche/Source/AvalancheMedia/Public/Playable/AvaPlayableRemoteControlValues.h
+++ b/Engine/Plugins/Experimental/Avalanche/Source/AvalancheMedia/Public/Playable/AvaPlayableRemoteControlValues.h
@@ -70,6 +70,7 @@ struct TStructOpsTypeTraits<FAvaPlayableRemoteControlValue> : public TStructOpsT
 	};
 };
 
+/** Container for the remote control values of a playable. */
 USTRUCT()
 struct AVALANCHEMEDIA_API FAvaPlayableRemoteControlValues
 {
@@ -85,10 +86,6 @@ struct AVALANCHEMEDIA_API FAvaPlayableRemoteControlValues
 	UPROPERTY()
 	TMap<FGuid, FAvaPlayableRemoteControlValue> ControllerValues;
 
-	/** Contains a set of entity guids that are bound to a controller action. */
-	UPROPERTY()
-	TSet<FGuid> EntitiesControlledByController;
-	
 	/** Copies the values (properties and controllers) from the given RemoteControlPreset. */
 	void CopyFrom(const URemoteControlPreset* InRemoteControlPreset, bool bInIsDefault);
 
@@ -146,11 +143,6 @@ struct AVALANCHEMEDIA_API FAvaPlayableRemoteControlValues
 	EAvaPlayableRemoteControlChanges ResetRemoteControlEntityValue(const FGuid& InId
 		, const FAvaPlayableRemoteControlValue& InReferenceValue, bool bInIsDefaults);
 
-	/**
-	 * Refreshes the EntitiesControlledByController set.
-	 */
-	void RefreshControlledEntities(const URemoteControlPreset* InRemoteControlPreset);
-
 	/**
 	 * Set the entity value from the given preset.
 	 * @return true if the value was set.
@@ -167,8 +159,10 @@ struct AVALANCHEMEDIA_API FAvaPlayableRemoteControlValues
 
 	/**
 	 *	Apply the entity values to the given remote control preset.
+	 *	@param InRemoteControlPreset RemoteControlPreset
+	 *	@param InSkipEntities set of entities to skip (will not be updated)
 	 */
-	void ApplyEntityValuesToRemoteControlPreset(URemoteControlPreset* InRemoteControlPreset) const;
+	void ApplyEntityValuesToRemoteControlPreset(URemoteControlPreset* InRemoteControlPreset, const TSet<FGuid>& InSkipEntities = TSet<FGuid>()) const;
 
 	/**
 	 *	Apply the controller values to the given remote control preset.
@@ -198,4 +192,11 @@ struct AVALANCHEMEDIA_API FAvaPlayableRemoteControlValues
 	 * Collect the referenced asset paths from the given values.
 	 */
 	static void CollectReferencedAssetPaths(const TMap<FGuid, FAvaPlayableRemoteControlValue>& InValues,TSet<FSoftObjectPath>& OutReferencedPaths);
-};
+
+	/**
+	 * Utility function to determine if a controller should be ignored by the playable management layer.
+	 * @param InController Controller
+	 * @return true if the controller should be ignored.
+	 */
+	static bool ShouldIgnoreController(const URCVirtualPropertyBase* InController);
+};
\ No newline at end of file
diff --git a/Engine/Plugins/Experimental/Avalanche/Source/AvalancheMedia/Public/Playable/AvaPlayableSettings.h b/Engine/Plugins/Experimental/Avalanche/Source/AvalancheMedia/Public/Playable/AvaPlayableSettings.h
--- a/Engine/Plugins/Experimental/Avalanche/Source/AvalancheMedia/Public/Playable/AvaPlayableSettings.h
+++ b/Engine/Plugins/Experimental/Avalanche/Source/AvalancheMedia/Public/Playable/AvaPlayableSettings.h
@@ -36,4 +36,11 @@ struct FAvaPlayableSettings
 	 */
 	UPROPERTY(Config, EditAnywhere, Category = Settings)
 	bool bHidePawnActors = true;
+
+	/**
+	 * List of postfix that can be added to controllers to be ignored by playable, playback and rundown implementation.
+	 * Those controllers can be used by other control systems at runtime.
+	 */
+	UPROPERTY(Config, EditAnywhere, Category = Settings)
+	TArray<FString> IgnoredControllerPostfix;
 };
diff --git a/Engine/Plugins/Experimental/Avalanche/Source/AvalancheMediaEditor/Private/Rundown/DetailsView/RemoteControl/Controllers/AvaRundownRCControllerItem.cpp b/Engine/Plugins/Experimental/Avalanche/Source/AvalancheMediaEditor/Private/Rundown/DetailsView/RemoteControl/Controllers/AvaRundownRCControllerItem.cpp
--- a/Engine/Plugins/Experimental/Avalanche/Source/AvalancheMediaEditor/Private/Rundown/DetailsView/RemoteControl/Controllers/AvaRundownRCControllerItem.cpp
+++ b/Engine/Plugins/Experimental/Avalanche/Source/AvalancheMediaEditor/Private/Rundown/DetailsView/RemoteControl/Controllers/AvaRundownRCControllerItem.cpp
@@ -4,16 +4,46 @@
 #include "Controller/RCController.h"
 #include "IRemoteControlUIModule.h"
 #include "RCVirtualProperty.h"
+#include "RemoteControlPreset.h"
 #include "SAvaRundownRCControllerItemRow.h"
+#include "Playable/AvaPlayableRemoteControlPresetInfo.h"
 #include "UObject/WeakObjectPtrTemplates.h"
 
 #define LOCTEXT_NAMESPACE "AvaRundownRCControllerItem"
 
-FAvaRundownRCControllerItem::FAvaRundownRCControllerItem(int32 InInstanceIndex, FName InAssetName, URCController* InController, const TSharedRef<IDetailTreeNode>& InTreeNode)
+namespace UE::AvaRundownRCControllerItem::Utils
+{
+	FText GetControllerDisplayName(URemoteControlPreset* InPreset, const FGuid& InControllerId)
+	{
+		if (InPreset)
+		{
+			if (URCVirtualPropertyBase* Controller = InPreset->GetController(InControllerId))
+			{
+				const FName DisplayName = Controller->DisplayName.IsNone() ? Controller->PropertyName : Controller->DisplayName;
+				return FText::FromName(DisplayName);
+			}
+		}
+		return FText::FromString(InControllerId.ToString());
+	}
+
+	FText GetEntityDisplayName(URemoteControlPreset* InPreset, const FGuid& InEntityId)
+	{
+		if (InPreset)
+		{
+			TWeakPtr<FRemoteControlEntity> EntityWeak = InPreset->GetExposedEntity(InEntityId);
+			if (TSharedPtr<FRemoteControlEntity> Entity = EntityWeak.Pin())
+			{
+				return FText::FromName(Entity->GetLabel());
+			}
+		}
+		return FText::FromString(InEntityId.ToString());
+	}
+}
+
+FAvaRundownRCControllerItem::FAvaRundownRCControllerItem(int32 InInstanceIndex, FName InAssetName, URCController* InController, const TSharedRef<IDetailTreeNode>& InTreeNode, const FAvaPlayableRemoteControlPresetInfo& InPresetInfo)
 {
 	AssetName = InAssetName;
 	InstanceIndex = InInstanceIndex;
-
 	Controller = InController;
 	
 	if (InController)
@@ -42,6 +72,34 @@ FAvaRundownRCControllerItem::FAvaRundownRCControllerItem(int32 InInstanceIndex,
 				ToolTipLines.Add(FText::Format(LOCTEXT("Controller_ToolTipLine_Desc", "Description: {0}"), InController->Description));
 			}
 			ToolTipLines.Add(FText::Format(LOCTEXT("Controller_ToolTipLine_DisplayIndex", "Display Index: {0}"), InController->DisplayIndex));
+			
+			if (InPresetInfo.IsControllerOverlapping(InController->Id))
+			{
+				// We want to report to the user which controller is overlapping and why.
+				using namespace UE::AvaRundownRCControllerItem::Utils;
+				if (URemoteControlPreset* Preset = InController->PresetWeakPtr.Get())
+				{
+					for (const TPair<FGuid, FAvaPlayableRemoteControlControlledEntityInfo>& ControlledEntity : InPresetInfo.EntitiesControlledByController)
+					{
+						if (ControlledEntity.Value.ControlledBy.Contains(InController->Id) && ControlledEntity.Value.ControlledBy.Num() > 1)
+						{
+							for (const FGuid& OtherController : ControlledEntity.Value.ControlledBy)
+							{
+								if (OtherController != InController->Id)
+								{
+									ToolTipLines.Add(
+										FText::Format(LOCTEXT("Controller_ToolTipLine_OverlappingWith", "Controller is overlapping with \"{0}\" because of entity \"{1}\""),
+											GetControllerDisplayName(Preset, OtherController), GetEntityDisplayName(Preset, ControlledEntity.Key)));
+								}
+							}
+						}
+					}
+				}
+			}
+			else
+			{
+				ToolTipLines.Add(LOCTEXT("Controller_ToolTipLine_NonOverlapping", "Controller is non-overlapping"));
+			}
 
 			for (const TPair<FName, FString>& MetaDataEntry : InController->Metadata)
 			{
diff --git a/Engine/Plugins/Experimental/Avalanche/Source/AvalancheMediaEditor/Private/Rundown/DetailsView/RemoteControl/Controllers/AvaRundownRCControllerItem.h b/Engine/Plugins/Experimental/Avalanche/Source/AvalancheMediaEditor/Private/Rundown/DetailsView/RemoteControl/Controllers/AvaRundownRCControllerItem.h
--- a/Engine/Plugins/Experimental/Avalanche/Source/AvalancheMediaEditor/Private/Rundown/DetailsView/RemoteControl/Controllers/AvaRundownRCControllerItem.h
+++ b/Engine/Plugins/Experimental/Avalanche/Source/AvalancheMediaEditor/Private/Rundown/DetailsView/RemoteControl/Controllers/AvaRundownRCControllerItem.h
@@ -11,11 +11,12 @@ class ITableRow;
 class SAvaRundownRCControllerPanel;
 class STableViewBase;
 class URCController;
+struct FAvaPlayableRemoteControlPresetInfo;
 
 class FAvaRundownRCControllerItem : public TSharedFromThis<FAvaRundownRCControllerItem>
 {
 public:
-	FAvaRundownRCControllerItem(int32 InInstanceIndex, FName InAssetName, URCController* InController, const TSharedRef<IDetailTreeNode>& InTreeNode);
+	FAvaRundownRCControllerItem(int32 InInstanceIndex, FName InAssetName, URCController* InController, const TSharedRef<IDetailTreeNode>& InTreeNode, const FAvaPlayableRemoteControlPresetInfo& InPresetInfo);
 	
 	TSharedRef<ITableRow> CreateWidget(TSharedRef<SAvaRundownRCControllerPanel> InControllerPanel, const TSharedRef<STableViewBase>& InOwnerTable) const;
 
diff --git a/Engine/Plugins/Experimental/Avalanche/Source/AvalancheMediaEditor/Private/Rundown/DetailsView/RemoteControl/Controllers/SAvaRundownRCControllerPanel.cpp b/Engine/Plugins/Experimental/Avalanche/Source/AvalancheMediaEditor/Private/Rundown/DetailsView/RemoteControl/Controllers/SAvaRundownRCControllerPanel.cpp
--- a/Engine/Plugins/Experimental/Avalanche/Source/AvalancheMediaEditor/Private/Rundown/DetailsView/RemoteControl/Controllers/SAvaRundownRCControllerPanel.cpp
+++ b/Engine/Plugins/Experimental/Avalanche/Source/AvalancheMediaEditor/Private/Rundown/DetailsView/RemoteControl/Controllers/SAvaRundownRCControllerPanel.cpp
@@ -4,18 +4,17 @@
 #include "AvaRundownRCControllerItem.h"
 #include "Behaviour/Builtin/Path/RCSetAssetByPathBehaviour.h"
 #include "Controller/RCController.h"
-#include "IAvaMediaEditorModule.h"
 #include "IDetailTreeNode.h"
 #include "IPropertyRowGenerator.h"
+#include "Modules/ModuleManager.h"
 #include "Playable/AvaPlayableRemoteControl.h"
+#include "Playable/AvaPlayableRemoteControlPresetInfo.h"
 #include "RCVirtualProperty.h"
-#include "Rundown/AvaRundownCommands.h"
 #include "Rundown/AvaRundownEditor.h"
 #include "Rundown/AvaRundownEditorUtils.h"
 #include "Rundown/AvaRundownManagedInstanceCache.h"
 #include "Rundown/AvaRundownPage.h"
 #include "Rundown/DetailsView/RemoteControl/Controllers/AvaRundownPageControllerContextMenu.h"
-#include "Rundown/Pages/AvaRundownPageControllerContext.h"
 #include "Widgets/Views/SHeaderRow.h"
 #include "Widgets/Views/SListView.h"
 
@@ -142,7 +141,7 @@ void SAvaRundownRCControllerPanel::RefreshForManagedInstance(int32 InInstanceInd
 
 				URCVirtualPropertyBase* const VirtualProperty = Preset->GetController(Property->GetFName());
 
-				if (!VirtualProperty)
+				if (!VirtualProperty || FAvaPlayableRemoteControlValues::ShouldIgnoreController(VirtualProperty))
 				{
 					continue;
 				}
@@ -153,19 +152,29 @@ void SAvaRundownRCControllerPanel::RefreshForManagedInstance(int32 InInstanceInd
 					if (!ControllerValueFromPage)
 					{
 						// If the value is not set in the page, fallback to default value from the template.
-						ControllerValueFromPage = InManagedInstance.GetDefaultRemoteControlValues().GetControllerValue(VirtualProperty->Id);
-
-						// Note: we are not adding the controller's default value in the page. (Tentative)
-						// Reason: Given that controller values are only applied for the managed RCP (ui),
-						// there is no need to save it in the page unless the user has modified from default.
-						// This way, we only keep user modified values.
-						
-						if (!ControllerValueFromPage)
+						if (const FAvaPlayableRemoteControlValue* const DefaultControllerValue = InManagedInstance.GetDefaultRemoteControlValues().GetControllerValue(VirtualProperty->Id))
+						{
+							// Ensure the default values have the default flag.
+							ensure(DefaultControllerValue->bIsDefault == true);
+
+							// WYSIWYG (Solution):
+							// Because Deterministic and "special" controllers are applied at runtime, 
+							// we need to add the default value (flagged as default) in the current page.
+							if (!SetSelectedPageControllerValue(VirtualProperty, *DefaultControllerValue))
+							{
+								UE_LOG(LogAvaPlayableRemoteControl, Error, TEXT("Controller \"%s\" (id:%s): failed to set value in currently selected page."),
+									*VirtualProperty->DisplayName.ToString(), *VirtualProperty->Id.ToString());
+							}
+
+							ControllerValueFromPage = DefaultControllerValue;
+						}
+						else
 						{
 							UE_LOG(LogAvaPlayableRemoteControl, Error, TEXT("Controller \"%s\" (id:%s) doesn't have a template default value."),
 								*VirtualProperty->DisplayName.ToString(), *VirtualProperty->Id.ToString());
 						}
 					}
+
 					if (ControllerValueFromPage)
 					{
 						using namespace UE::AvaPlayableRemoteControl;
@@ -176,10 +185,8 @@ void SAvaRundownRCControllerPanel::RefreshForManagedInstance(int32 InInstanceInd
 							if (!CurrentControllerValue.IsSameValueAs(*ControllerValueFromPage))
 							{
 								// This serves only to sync the Managed RCP's controller values to the page values.
-								// We want to be able to do this without affecting the entity values (preserve WYSIWYG).
-								// So, we temporarily disable the controller's behaviours.
-								FScopedPushControllerBehavioursEnable PushBehavioursEnable(VirtualProperty, false);
-								if (Failed(RemoteControlResult = SetValueOfController(VirtualProperty, ControllerValueFromPage->Value)))
+								// We want to be able to do this without affecting the entity values (preserve WYSIWYG), thus we disable behaviors.
+								if (Failed(RemoteControlResult = SetValueOfController(VirtualProperty, ControllerValueFromPage->Value, /*bInBehaviorsEnabled*/ false)))
 								{
 									UE_LOG(LogAvaPlayableRemoteControl, Error,
 										TEXT("Controller \"%s\" (id:%s): failed to set value in currently selected page: %s."),
@@ -206,7 +213,9 @@ void SAvaRundownRCControllerPanel::RefreshForManagedInstance(int32 InInstanceInd
 						}
 					}
 
-					ControllerItems.Add(MakeShared<FAvaRundownRCControllerItem>(InInstanceIndex, InManagedInstance.GetSourceAssetPath().GetAssetFName(), Controller, Child));
+					FName SourceAssetName = InManagedInstance.GetSourceAssetPath().GetAssetFName();
+					const FAvaPlayableRemoteControlPresetInfo& PresetInfo = InManagedInstance.GetRemoteControlPresetInfo();
+					ControllerItems.Add(MakeShared<FAvaRundownRCControllerItem>(InInstanceIndex, SourceAssetName, Controller, Child, PresetInfo));
 				}
 			}
 		}
@@ -238,7 +247,14 @@ void SAvaRundownRCControllerPanel::Refresh(const TArray<int32>& InSelectedPageId
 	{
 		if (const URemoteControlPreset* Preset = ManagedInstance ? ManagedInstance->GetRemoteControlPreset() : nullptr)
 		{
-			NumItems += Preset->GetNumControllers();
+			// Only count the controllers that are not ignored.
+			for (const URCVirtualPropertyBase* PropertyBase : Preset->GetControllers())
+			{
+				if (!FAvaPlayableRemoteControlValues::ShouldIgnoreController(PropertyBase))
+				{
+					NumItems++;
+				}
+			}
 		}
 	}
 
diff --git a/Engine/Plugins/Experimental/Avalanche/Source/AvalancheMediaEditor/Private/Rundown/DetailsView/RemoteControl/Properties/SAvaRundownPageRemoteControlProps.cpp b/Engine/Plugins/Experimental/Avalanche/Source/AvalancheMediaEditor/Private/Rundown/DetailsView/RemoteControl/Properties/SAvaRundownPageRemoteControlProps.cpp
--- a/Engine/Plugins/Experimental/Avalanche/Source/AvalancheMediaEditor/Private/Rundown/DetailsView/RemoteControl/Properties/SAvaRundownPageRemoteControlProps.cpp
+++ b/Engine/Plugins/Experimental/Avalanche/Source/AvalancheMediaEditor/Private/Rundown/DetailsView/RemoteControl/Properties/SAvaRundownPageRemoteControlProps.cpp
@@ -4,6 +4,7 @@
 #include "Item/AvaRundownRCFieldItem.h"
 #include "Modules/ModuleManager.h"
 #include "Playable/AvaPlayableRemoteControl.h"
+#include "Playable/AvaPlayableRemoteControlPresetInfo.h"
 #include "Playable/AvaPlayableRemoteControlValues.h"
 #include "Playback/AvaPlaybackUtils.h"
 #include "PropertyEditorModule.h"
@@ -313,8 +314,18 @@ void SAvaRundownPageRemoteControlProps::Refresh(const TArray<int32>& InSelectedP
 						
 						if (!DefaultEntityValue)
 						{
-							UE_LOG(LogAvaPlayableRemoteControl, Error, TEXT("Entity \"%s\" (id:%s) doesn't have a template default value."),
-								*FieldEntity->GetLabel().ToString(), *FieldEntity->GetId().ToString());
+							FString AccessError;
+							// Check if the property is correctly bound, that would explain why the value is missing.
+							if (!UE::AvaPlayableRemoteControl::HasReadAccess(FieldEntity, AccessError))
+							{
+								UE_LOG(LogAvaPlayableRemoteControl, Error, TEXT("Entity \"%s\" (id:%s) doesn't have a template default value. Reason: %s."),
+									*FieldEntity->GetLabel().ToString(), *FieldEntity->GetId().ToString(), *AccessError);
+							}
+							else
+							{
+								UE_LOG(LogAvaPlayableRemoteControl, Error, TEXT("Entity \"%s\" (id:%s) doesn't have a template default value."),
+									*FieldEntity->GetLabel().ToString(), *FieldEntity->GetId().ToString());
+							}
 							// TODO: UX improvement: instead of skipping, could add empty element, with error mark (and error message in tooltip).
 							continue;
 						}
@@ -343,7 +354,7 @@ void SAvaRundownPageRemoteControlProps::Refresh(const TArray<int32>& InSelectedP
 						continue;
 					}
 
-					const bool bEntityControlled = ManagedInstance->GetDefaultRemoteControlValues().EntitiesControlledByController.Contains(FieldEntity->GetId());
+					const bool bEntityControlled = ManagedInstance->GetRemoteControlPresetInfo().EntitiesControlledByController.Contains(FieldEntity->GetId());
 
 					if (!bEntityControlled || bShowControlledProperties)
 					{
