[MotionDesign] Uncomment drag and drop extension

diff --git a/Engine/Plugins/Experimental/Avalanche/Source/AvalancheSequencer/Private/NavigationTool/AvaNavigationToolProvider.cpp b/Engine/Plugins/Experimental/Avalanche/Source/AvalancheSequencer/Private/NavigationTool/AvaNavigationToolProvider.cpp
--- a/Engine/Plugins/Experimental/Avalanche/Source/AvalancheSequencer/Private/NavigationTool/AvaNavigationToolProvider.cpp
+++ b/Engine/Plugins/Experimental/Avalanche/Source/AvalancheSequencer/Private/NavigationTool/AvaNavigationToolProvider.cpp
@@ -1,4 +1,4 @@
-﻿// Copyright Epic Games, Inc. All Rights Reserved.
+// Copyright Epic Games, Inc. All Rights Reserved.
 
 #include "AvaNavigationToolProvider.h"
 #include "AvaSequence.h"
@@ -148,8 +148,8 @@ void FAvaNavigationToolProvider::OnActivate()
 	ExtendToolToolBar();
 	ExtendToolItemContextMenu();
 
-	//DragDropInitializedDelegate = UE::SequenceNavigator::FNavigationToolExtender::OnItemDragDropOpInitialized()
-	//	.AddSP(this, &FAvaNavigationToolProvider::ExtendToolItemDragDropOp);
+	DragDropInitializedDelegate = UE::SequenceNavigator::FNavigationToolExtender::OnItemDragDropOpInitialized()
+		.AddSP(this, &FAvaNavigationToolProvider::ExtendToolItemDragDropOp);
 }
 
 void FAvaNavigationToolProvider::OnDeactivate()
@@ -157,8 +157,8 @@ void FAvaNavigationToolProvider::OnDeactivate()
 	RemoveToolToolBarExtension();
 	RemoveToolItemContextMenuExtension();
 
-	//UE::SequenceNavigator::FNavigationToolExtender::OnItemDragDropOpInitialized()
-	//	.Remove(DragDropInitializedDelegate);
+	UE::SequenceNavigator::FNavigationToolExtender::OnItemDragDropOpInitialized()
+		.Remove(DragDropInitializedDelegate);
 }
 
 void FAvaNavigationToolProvider::OnExtendColumns(UE::SequenceNavigator::FNavigationToolColumnExtender& OutExtender)
