[Motion Design] Rundown Editor: fix a crash on refresh of page details when items are removed from the list because they get hidden.

diff --git a/Engine/Plugins/Experimental/Avalanche/Source/AvalancheMediaEditor/Private/Rundown/DetailsView/RemoteControl/Properties/SAvaRundownRCPropertyItemRow.cpp b/Engine/Plugins/Experimental/Avalanche/Source/AvalancheMediaEditor/Private/Rundown/DetailsView/RemoteControl/Properties/SAvaRundownRCPropertyItemRow.cpp
--- a/Engine/Plugins/Experimental/Avalanche/Source/AvalancheMediaEditor/Private/Rundown/DetailsView/RemoteControl/Properties/SAvaRundownRCPropertyItemRow.cpp
+++ b/Engine/Plugins/Experimental/Avalanche/Source/AvalancheMediaEditor/Private/Rundown/DetailsView/RemoteControl/Properties/SAvaRundownRCPropertyItemRow.cpp
@@ -159,7 +159,7 @@ FText SAvaRundownRCPropertyItemRow::GetPropertyTooltipText() const
 	FText OwnerText = LOCTEXT("InvalidOwnerText", "(Invalid)");
 	FText SubobjectPathText = LOCTEXT("InvalidSubobjectPathText", "(Invalid)");
 
-	const TSharedPtr<FRemoteControlEntity> Entity = ItemPtr->GetEntity();
+	const TSharedPtr<FRemoteControlEntity> Entity = ItemPtr ? ItemPtr->GetEntity() : nullptr;
 	if (Entity.IsValid())
 	{
 		const FString BindingPath = Entity->GetLastBindingPath().ToString();
