Motion Design: outliner to refresh when levels are added/removed

diff --git a/Engine/Plugins/Experimental/Avalanche/Source/AvalancheOutliner/Private/AvaOutliner.cpp b/Engine/Plugins/Experimental/Avalanche/Source/AvalancheOutliner/Private/AvaOutliner.cpp
--- a/Engine/Plugins/Experimental/Avalanche/Source/AvalancheOutliner/Private/AvaOutliner.cpp
+++ b/Engine/Plugins/Experimental/Avalanche/Source/AvalancheOutliner/Private/AvaOutliner.cpp
@@ -87,6 +87,8 @@ FAvaOutliner::~FAvaOutliner()
 
 	FCoreUObjectDelegates::OnObjectsReplaced.RemoveAll(this);
 	FEditorDelegates::OnEditorActorReplaced.RemoveAll(this);
+
+	UnbindWorldDelegates();
 }
 
 UAvaOutlinerSubsystem* FAvaOutliner::GetOutlinerSubsystem() const
@@ -388,6 +390,9 @@ void FAvaOutliner::Refresh()
 	SceneOutlinerParentMap.Reset();
 	if (UWorld* const World = GetWorld())
 	{
+		// Ensure the latest world has its delegates bound
+		BindWorldDelegates(World);
+
 		auto AddActorToParentMap = [this](UWorld* const InWorld, AActor* const InActor)
 			{
 				if (!IsValid(InWorld) || !IsValid(InActor))
@@ -1329,6 +1334,30 @@ void FAvaOutliner::SetOutlinerModified()
 	}
 }
 
+void FAvaOutliner::BindWorldDelegates(UWorld* InWorld)
+{
+	if (BoundWorldWeak == InWorld)
+	{
+		return;
+	}
+
+	UnbindWorldDelegates();
+	BoundWorldWeak = InWorld;
+
+	if (InWorld)
+	{
+		InWorld->OnLevelsChanged().AddSP(this, &FAvaOutliner::RequestRefresh);
+	}
+}
+
+void FAvaOutliner::UnbindWorldDelegates()
+{
+	if (UWorld* const BoundWorld = BoundWorldWeak.Get())
+	{
+		BoundWorld->OnLevelsChanged().RemoveAll(this);
+	}
+}
+
 void FAvaOutliner::AddItem(const FAvaOutlinerItemPtr& InItem)
 {
 	const FAvaOutlinerItemId ItemId = InItem->GetItemId();
diff --git a/Engine/Plugins/Experimental/Avalanche/Source/AvalancheOutliner/Private/AvaOutliner.h b/Engine/Plugins/Experimental/Avalanche/Source/AvalancheOutliner/Private/AvaOutliner.h
--- a/Engine/Plugins/Experimental/Avalanche/Source/AvalancheOutliner/Private/AvaOutliner.h
+++ b/Engine/Plugins/Experimental/Avalanche/Source/AvalancheOutliner/Private/AvaOutliner.h
@@ -262,6 +262,11 @@ public:
 	void SetOutlinerModified();
 
 private:
+	/** Binds a given worlds delegates to the outliner, unbinding the previous bound world's */
+	void BindWorldDelegates(UWorld* InWorld);
+	/** Unbinds the last bound world's delegates from outliner */
+	void UnbindWorldDelegates();
+
 	void AddItem(const FAvaOutlinerItemPtr& InItem);
 	void RemoveItem(const FAvaOutlinerItemId& InItemId);
 	void ForEachItem(TFunctionRef<void(const FAvaOutlinerItemPtr&)> InFunc);
@@ -272,6 +277,9 @@ private:
 	/** the base command list gotten from the outliner provider to append our command list to */
 	TWeakPtr<FUICommandList> BaseCommandListWeak;
 
+	/** Latest world with bound delegates. */
+	TWeakObjectPtr<UWorld> BoundWorldWeak;
+
 	/** the map of the registered items */
 	TMap<FAvaOutlinerItemId, FAvaOutlinerItemPtr> ItemMap;
 
