Motion Design : Modifiers
- Added inverted outliner tree sort mode to the translucency sort priority modifier
- Updated modifier and properties display name
- Updated modifier tests
- DisplayName will be set after modifier CDO is called, it will also be included in search keywords

diff --git a/Engine/Plugins/Experimental/ActorModifierCore/Source/ActorModifierCore/Private/Modifiers/ActorModifierCoreBase.cpp b/Engine/Plugins/Experimental/ActorModifierCore/Source/ActorModifierCore/Private/Modifiers/ActorModifierCoreBase.cpp
--- a/Engine/Plugins/Experimental/ActorModifierCore/Source/ActorModifierCore/Private/Modifiers/ActorModifierCoreBase.cpp
+++ b/Engine/Plugins/Experimental/ActorModifierCore/Source/ActorModifierCore/Private/Modifiers/ActorModifierCoreBase.cpp
@@ -664,6 +664,14 @@ void UActorModifierCoreBase::PostModifierCDOCreation()
 	{
 		Metadata = MakeShared<FActorModifierCoreMetadata>(this);
 		OnModifierCDOSetup(*Metadata);
+
+#if WITH_EDITOR
+		// Set a display name if none was provided
+		if (Metadata->GetDisplayName().IsEmpty())
+		{
+			Metadata->SetDisplayName(FText::FromString(FName::NameToDisplayString(Metadata->GetName().ToString(), false)));
+		}
+#endif
 	}
 }
 
diff --git a/Engine/Plugins/Experimental/ActorModifierCore/Source/ActorModifierCore/Private/Modifiers/ActorModifierCoreDefs.cpp b/Engine/Plugins/Experimental/ActorModifierCore/Source/ActorModifierCore/Private/Modifiers/ActorModifierCoreDefs.cpp
--- a/Engine/Plugins/Experimental/ActorModifierCore/Source/ActorModifierCore/Private/Modifiers/ActorModifierCoreDefs.cpp
+++ b/Engine/Plugins/Experimental/ActorModifierCore/Source/ActorModifierCore/Private/Modifiers/ActorModifierCoreDefs.cpp
@@ -191,11 +191,6 @@ void FActorModifierCoreMetadata::SetupProfilerInstanceInternal(TSharedPtr<FActor
 FActorModifierCoreMetadata& FActorModifierCoreMetadata::SetName(FName InName)
 {
 	Name = InName;
-
-#if WITH_EDITOR
-	DisplayName = FText::FromString(FName::NameToDisplayString(Name.ToString(), false));
-#endif
-
 	return *this;
 }
 
diff --git a/Engine/Plugins/Experimental/ActorModifierCore/Source/ActorModifierCoreEditor/Private/Modifiers/Customizations/ActorModifierCoreEditorStackCustomization.cpp b/Engine/Plugins/Experimental/ActorModifierCore/Source/ActorModifierCoreEditor/Private/Modifiers/Customizations/ActorModifierCoreEditorStackCustomization.cpp
--- a/Engine/Plugins/Experimental/ActorModifierCore/Source/ActorModifierCoreEditor/Private/Modifiers/Customizations/ActorModifierCoreEditorStackCustomization.cpp
+++ b/Engine/Plugins/Experimental/ActorModifierCore/Source/ActorModifierCoreEditor/Private/Modifiers/Customizations/ActorModifierCoreEditorStackCustomization.cpp
@@ -227,7 +227,7 @@ void UActorModifierCoreEditorStackCustomization::CustomizeItemHeader(const FOper
 		// Item keyword for search
 		UActorModifierCoreBase* Modifier = InItem->Get<UActorModifierCoreBase>(0);
 
-		const TSet<FString> SearchKeywords
+		TSet<FString> SearchKeywords
 		{
 			Modifier->GetModifierName().ToString(),
 			Modifier->GetModifierCategory().ToString()
@@ -237,11 +237,12 @@ void UActorModifierCoreEditorStackCustomization::CustomizeItemHeader(const FOper
 		FLinearColor ModifierColor = FLinearColor::Transparent;
 		FText ModifierTooltip = FText::GetEmpty();
 
-		ModifierSubsystem->ProcessModifierMetadata(Modifier->GetModifierName(), [&ModifierIcon, &ModifierColor, &ModifierTooltip](const FActorModifierCoreMetadata& InMetadata)
+		ModifierSubsystem->ProcessModifierMetadata(Modifier->GetModifierName(), [&SearchKeywords, &ModifierIcon, &ModifierColor, &ModifierTooltip](const FActorModifierCoreMetadata& InMetadata)
 		{
 			ModifierIcon = InMetadata.GetIcon();
 			ModifierColor = InMetadata.GetColor();
 			ModifierTooltip = InMetadata.GetDescription();
+			SearchKeywords.Add(InMetadata.GetDisplayName().ToString());
 			return true;
 		});
 
diff --git a/Engine/Plugins/Experimental/Avalanche/Source/AvalancheModifiers/Private/Modifiers/AvaTranslucentPriorityModifier.cpp b/Engine/Plugins/Experimental/Avalanche/Source/AvalancheModifiers/Private/Modifiers/AvaTranslucentPriorityModifier.cpp
--- a/Engine/Plugins/Experimental/Avalanche/Source/AvalancheModifiers/Private/Modifiers/AvaTranslucentPriorityModifier.cpp
+++ b/Engine/Plugins/Experimental/Avalanche/Source/AvalancheModifiers/Private/Modifiers/AvaTranslucentPriorityModifier.cpp
@@ -25,7 +25,7 @@ void UAvaTranslucentPriorityModifier::OnModifierCDOSetup(FActorModifierCoreMetad
 {
 	Super::OnModifierCDOSetup(InMetadata);
 
-	InMetadata.SetName(TEXT("TranslucentPriority"));
+	InMetadata.SetName(TEXT("TranslucentSortPriority"));
 	InMetadata.SetCategory(TEXT("Rendering"));
 #if WITH_EDITOR
 	InMetadata.SetDescription(LOCTEXT("ModifierDescription", "Set the translucency sort priority of a primitive component based on different modes"));
@@ -63,7 +63,7 @@ void UAvaTranslucentPriorityModifier::OnModifierAdded(EActorModifierCoreEnableRe
 	if (InReason == EActorModifierCoreEnableReason::User)
 	{
 		CameraActorWeak = GetDefaultCameraActor();
-		Mode = EAvaTranslucentPriorityModifierMode::AutoOutlinerTree;
+		Mode = EAvaTranslucentPriorityModifierMode::AutoOutlinerTreeInverted;
 	}
 
 	if (UAvaTranslucentPriorityModifierShared* SharedObject = GetShared<UAvaTranslucentPriorityModifierShared>(true))
@@ -424,7 +424,7 @@ void UAvaTranslucentPriorityModifier::OnSceneTreeTrackedActorRearranged(int32 In
 {
 	Super::OnSceneTreeTrackedActorRearranged(InIdx, InRearrangedActor);
 
-	if (Mode == EAvaTranslucentPriorityModifierMode::AutoOutlinerTree)
+	if (Mode == EAvaTranslucentPriorityModifierMode::AutoOutlinerTree || Mode == EAvaTranslucentPriorityModifierMode::AutoOutlinerTreeInverted)
 	{
 		MarkModifierDirty();
 	}
diff --git a/Engine/Plugins/Experimental/Avalanche/Source/AvalancheModifiers/Private/Shared/AvaTranslucentPriorityModifierShared.cpp b/Engine/Plugins/Experimental/Avalanche/Source/AvalancheModifiers/Private/Shared/AvaTranslucentPriorityModifierShared.cpp
--- a/Engine/Plugins/Experimental/Avalanche/Source/AvalancheModifiers/Private/Shared/AvaTranslucentPriorityModifierShared.cpp
+++ b/Engine/Plugins/Experimental/Avalanche/Source/AvalancheModifiers/Private/Shared/AvaTranslucentPriorityModifierShared.cpp
@@ -154,7 +154,9 @@ TArray<FAvaTranslucentPriorityModifierComponentState> UAvaTranslucentPriorityMod
 
 	SortedComponentStates.Reserve(ComponentStates.Num());
 
-	if (InModifierContext->GetMode() == EAvaTranslucentPriorityModifierMode::Manual)
+	const EAvaTranslucentPriorityModifierMode SortMode = InModifierContext->GetMode();
+
+	if (SortMode == EAvaTranslucentPriorityModifierMode::Manual)
 	{
 		// Only gather components the modifier is handling
 		for (const FAvaTranslucentPriorityModifierComponentState& ComponentState : ComponentStates)
@@ -173,7 +175,7 @@ TArray<FAvaTranslucentPriorityModifierComponentState> UAvaTranslucentPriorityMod
 			}
 		}
 	}
-	else if (InModifierContext->GetMode() == EAvaTranslucentPriorityModifierMode::AutoCameraDistance)
+	else if (SortMode == EAvaTranslucentPriorityModifierMode::AutoCameraDistance)
 	{
 		// Gather components with same camera context
 		const ACameraActor* CameraActor = InModifierContext->GetCameraActorWeak().Get();
@@ -212,8 +214,10 @@ TArray<FAvaTranslucentPriorityModifierComponentState> UAvaTranslucentPriorityMod
 			return ADist > BDist;
 		});
 	}
-	else if (InModifierContext->GetMode() == EAvaTranslucentPriorityModifierMode::AutoOutlinerTree)
+	else if (SortMode == EAvaTranslucentPriorityModifierMode::AutoOutlinerTree || SortMode == EAvaTranslucentPriorityModifierMode::AutoOutlinerTreeInverted)
 	{
+		const bool bInvertedTree = SortMode == EAvaTranslucentPriorityModifierMode::AutoOutlinerTreeInverted;
+
 		const AActor* ModifiedActor = InModifierContext->GetModifiedActor();
 		UWorld* const World = ModifiedActor->GetTypedOuter<UWorld>();
 
@@ -246,12 +250,12 @@ TArray<FAvaTranslucentPriorityModifierComponentState> UAvaTranslucentPriorityMod
 		TSharedPtr<IAvaOutliner> AvaOutliner = FAvaOutlinerUtils::EditorGetOutliner(World);
 		if (AvaOutliner.IsValid())
 		{
-			SortedComponentStates.StableSort([AvaOutliner](const FAvaTranslucentPriorityModifierComponentState& InComponentA, const FAvaTranslucentPriorityModifierComponentState& InComponentB)->bool
+			SortedComponentStates.StableSort([AvaOutliner, bInvertedTree](const FAvaTranslucentPriorityModifierComponentState& InComponentA, const FAvaTranslucentPriorityModifierComponentState& InComponentB)->bool
 			{
 				const FAvaOutlinerItemPtr OutlinerItemA = AvaOutliner->FindItem(InComponentA.GetOwningActor());
 				const FAvaOutlinerItemPtr OutlinerItemB = AvaOutliner->FindItem(InComponentB.GetOwningActor());
 
-				return UE::AvaOutliner::CompareOutlinerItemOrder(OutlinerItemB, OutlinerItemA);
+				return UE::AvaOutliner::CompareOutlinerItemOrder(OutlinerItemB, OutlinerItemA) ^ bInvertedTree;
 			});
 		}
 		else
@@ -262,12 +266,12 @@ TArray<FAvaTranslucentPriorityModifierComponentState> UAvaTranslucentPriorityMod
 			{
 				const FAvaSceneTree& SceneTree = SceneInterface->GetSceneTree();
 
-				SortedComponentStates.StableSort([&SceneTree, World](const FAvaTranslucentPriorityModifierComponentState& InComponentA, const FAvaTranslucentPriorityModifierComponentState& InComponentB)->bool
+				SortedComponentStates.StableSort([&SceneTree, World, bInvertedTree](const FAvaTranslucentPriorityModifierComponentState& InComponentA, const FAvaTranslucentPriorityModifierComponentState& InComponentB)->bool
 				{
 					const FAvaSceneTreeNode* SceneItemA = SceneTree.FindTreeNode(FAvaSceneItem(InComponentA.GetOwningActor(), World));
 					const FAvaSceneTreeNode* SceneItemB = SceneTree.FindTreeNode(FAvaSceneItem(InComponentB.GetOwningActor(), World));
 
-					return FAvaSceneTree::CompareTreeItemOrder(SceneItemB, SceneItemA);
+					return FAvaSceneTree::CompareTreeItemOrder(SceneItemB, SceneItemA) ^ bInvertedTree;
 				});
 			}
 #if WITH_EDITOR
diff --git a/Engine/Plugins/Experimental/Avalanche/Source/AvalancheModifiers/Private/Tests/Framework/AvaTestData.cpp b/Engine/Plugins/Experimental/Avalanche/Source/AvalancheModifiers/Private/Tests/Framework/AvaTestData.cpp
--- a/Engine/Plugins/Experimental/Avalanche/Source/AvalancheModifiers/Private/Tests/Framework/AvaTestData.cpp
+++ b/Engine/Plugins/Experimental/Avalanche/Source/AvalancheModifiers/Private/Tests/Framework/AvaTestData.cpp
@@ -34,7 +34,7 @@ FAvaTestData::FAvaTestData()
 		"LookAt",
 		"MaterialParameter",
 		"RadialArrange",
-		"TranslucentPriority",
+		"TranslucentSortPriority",
 		"Visibility"
 	};
 
@@ -48,13 +48,13 @@ FAvaTestData::FAvaTestData()
 		"LookAt",
 		"MaterialParameter",
 		"RadialArrange",
-		"TranslucentPriority",
+		"TranslucentSortPriority",
 		"Visibility"
 	};
 
 	ExpectedRenderingModifiers =
 	{
-		"TranslucentPriority",
+		"TranslucentSortPriority",
 		"Visibility"
 	};
 }
diff --git a/Engine/Plugins/Experimental/Avalanche/Source/AvalancheModifiers/Public/Modifiers/AvaTranslucentPriorityModifier.h b/Engine/Plugins/Experimental/Avalanche/Source/AvalancheModifiers/Public/Modifiers/AvaTranslucentPriorityModifier.h
--- a/Engine/Plugins/Experimental/Avalanche/Source/AvalancheModifiers/Public/Modifiers/AvaTranslucentPriorityModifier.h
+++ b/Engine/Plugins/Experimental/Avalanche/Source/AvalancheModifiers/Public/Modifiers/AvaTranslucentPriorityModifier.h
@@ -14,9 +14,11 @@ UENUM(BlueprintType)
 enum class EAvaTranslucentPriorityModifierMode : uint8
 {
 	/** The closer you are from the camera based on camera forward axis, the higher your sort priority will be */
-	AutoCameraDistance,
+	AutoCameraDistance          UMETA(DisplayName = "Camera Distance"),
 	/** The higher you are in the outline tree, the higher your sort priority will be */
-	AutoOutlinerTree,
+	AutoOutlinerTree            UMETA(DisplayName = "Outliner Top First"),
+	/** The lower you are in the outline tree, the higher your sort priority will be */
+	AutoOutlinerTreeInverted    UMETA(DisplayName = "Outliner Bottom First"),
 	/** Set it yourself */
 	Manual
 };
@@ -29,19 +31,19 @@ class UAvaTranslucentPriorityModifier : public UActorModifierArrangeBaseModifier
 	friend class UAvaTranslucentPriorityModifierShared;
 
 public:
-	UFUNCTION(BlueprintCallable, Category="Motion Design|Modifiers|TranslucentPriority")
+	UFUNCTION(BlueprintCallable, Category="Motion Design|Modifiers|TranslucentSortPriority")
     AVALANCHEMODIFIERS_API void SetMode(EAvaTranslucentPriorityModifierMode InMode);
 
-	UFUNCTION(BlueprintPure, Category="Motion Design|Modifiers|TranslucentPriority")
+	UFUNCTION(BlueprintPure, Category="Motion Design|Modifiers|TranslucentSortPriority")
     EAvaTranslucentPriorityModifierMode GetMode() const
     {
 	    return Mode;
     }
 
-	UFUNCTION(BlueprintCallable, Category="Motion Design|Modifiers|TranslucentPriority")
+	UFUNCTION(BlueprintCallable, Category="Motion Design|Modifiers|TranslucentSortPriority")
 	AVALANCHEMODIFIERS_API void SetCameraActor(ACameraActor* InCameraActor);
 
-	UFUNCTION(BlueprintPure, Category="Motion Design|Modifiers|TranslucentPriority")
+	UFUNCTION(BlueprintPure, Category="Motion Design|Modifiers|TranslucentSortPriority")
 	ACameraActor* GetCameraActor() const
 	{
 		return CameraActorWeak.Get();
@@ -54,37 +56,37 @@ public:
 		return CameraActorWeak;
 	}
 
-	UFUNCTION(BlueprintCallable, Category="Motion Design|Modifiers|TranslucentPriority")
+	UFUNCTION(BlueprintCallable, Category="Motion Design|Modifiers|TranslucentSortPriority")
     AVALANCHEMODIFIERS_API void SetSortPriority(int32 InSortPriority);
 
-	UFUNCTION(BlueprintPure, Category="Motion Design|Modifiers|TranslucentPriority")
+	UFUNCTION(BlueprintPure, Category="Motion Design|Modifiers|TranslucentSortPriority")
     int32 GetSortPriority() const
     {
 	    return SortPriority;
     }
 
-	UFUNCTION(BlueprintCallable, Category="Motion Design|Modifiers|TranslucentPriority")
+	UFUNCTION(BlueprintCallable, Category="Motion Design|Modifiers|TranslucentSortPriority")
 	AVALANCHEMODIFIERS_API void SetSortPriorityOffset(int32 InOffset);
 
-	UFUNCTION(BlueprintPure, Category="Motion Design|Modifiers|TranslucentPriority")
+	UFUNCTION(BlueprintPure, Category="Motion Design|Modifiers|TranslucentSortPriority")
 	int32 GetSortPriorityOffset() const
 	{
 		return SortPriorityOffset;
 	}
 
-	UFUNCTION(BlueprintCallable, Category="Motion Design|Modifiers|TranslucentPriority")
+	UFUNCTION(BlueprintCallable, Category="Motion Design|Modifiers|TranslucentSortPriority")
 	AVALANCHEMODIFIERS_API void SetSortPriorityStep(int32 InStep);
 
-	UFUNCTION(BlueprintPure, Category="Motion Design|Modifiers|TranslucentPriority")
+	UFUNCTION(BlueprintPure, Category="Motion Design|Modifiers|TranslucentSortPriority")
 	int32 GetSortPriorityStep() const
 	{
 		return SortPriorityStep;
 	}
 
-	UFUNCTION(BlueprintCallable, Category="Motion Design|Modifiers|TranslucentPriority")
+	UFUNCTION(BlueprintCallable, Category="Motion Design|Modifiers|TranslucentSortPriority")
 	AVALANCHEMODIFIERS_API void SetIncludeChildren(bool bInIncludeChildren);
 
-	UFUNCTION(BlueprintPure, Category="Motion Design|Modifiers|TranslucentPriority")
+	UFUNCTION(BlueprintPure, Category="Motion Design|Modifiers|TranslucentSortPriority")
 	bool GetIncludeChildren() const
 	{
 		return bIncludeChildren;
@@ -134,27 +136,27 @@ protected:
 	ACameraActor* GetDefaultCameraActor() const;
 
 	/** The sort mode we are currently in */
-	UPROPERTY(EditInstanceOnly, Setter, Getter, Category="TranslucentPriority", meta=(AllowPrivateAccess="true"))
+	UPROPERTY(EditInstanceOnly, Setter, Getter, Category="TranslucentSortPriority", meta=(DisplayName="Sort Mode", AllowPrivateAccess="true"))
 	EAvaTranslucentPriorityModifierMode Mode = EAvaTranslucentPriorityModifierMode::Manual;
 
 	/** The camera actor to compute the distance from */
-	UPROPERTY(EditInstanceOnly, Setter, Getter, Category="TranslucentPriority", meta=(DisplayName="CameraActor", EditCondition="Mode == EAvaTranslucentPriorityModifierMode::AutoCameraDistance", AllowPrivateAccess="true"))
+	UPROPERTY(EditInstanceOnly, Setter, Getter, Category="TranslucentSortPriority", meta=(DisplayName="Camera Actor", EditCondition="Mode == EAvaTranslucentPriorityModifierMode::AutoCameraDistance", AllowPrivateAccess="true"))
 	TWeakObjectPtr<ACameraActor> CameraActorWeak = nullptr;
 
 	/** The sort priority that will be set on the primitive component for manual mode */
-	UPROPERTY(EditInstanceOnly, Setter, Getter, Category="TranslucentPriority", meta=(EditCondition="Mode == EAvaTranslucentPriorityModifierMode::Manual", EditConditionHides, AllowPrivateAccess="true"))
+	UPROPERTY(EditInstanceOnly, Setter, Getter, Category="TranslucentSortPriority", meta=(EditCondition="Mode == EAvaTranslucentPriorityModifierMode::Manual", EditConditionHides, AllowPrivateAccess="true"))
 	int32 SortPriority = 0;
 
 	/** Sort priority offset shared across all modifiers in this same level */
-	UPROPERTY(Transient, EditInstanceOnly, Setter, Getter, Category="TranslucentPriority", meta=(AllowPrivateAccess="true"))
+	UPROPERTY(Transient, EditInstanceOnly, Setter, Getter, Category="TranslucentSortPriority", meta=(AllowPrivateAccess="true"))
 	int32 SortPriorityOffset = 0;
 
 	/** Sort priority incremental step shared across all modifiers in this same level */
-	UPROPERTY(Transient, EditInstanceOnly, Setter, Getter, Category="TranslucentPriority", meta=(AllowPrivateAccess="true"))
+	UPROPERTY(Transient, EditInstanceOnly, Setter, Getter, Category="TranslucentSortPriority", meta=(AllowPrivateAccess="true"))
 	int32 SortPriorityStep = 1;
 
 	/** If true, will include children too and update their sort priority */
-	UPROPERTY(EditInstanceOnly, Setter="SetIncludeChildren", Getter="GetIncludeChildren", Category="TranslucentPriority", meta=(AllowPrivateAccess="true"))
+	UPROPERTY(EditInstanceOnly, Setter="SetIncludeChildren", Getter="GetIncludeChildren", Category="TranslucentSortPriority", meta=(AllowPrivateAccess="true"))
 	bool bIncludeChildren = true;
 
 	/** The components this modifier is managing */
diff --git a/Engine/Plugins/Experimental/Avalanche/Source/AvalancheModifiers/Public/Modifiers/AvaVisibilityModifier.h b/Engine/Plugins/Experimental/Avalanche/Source/AvalancheModifiers/Public/Modifiers/AvaVisibilityModifier.h
--- a/Engine/Plugins/Experimental/Avalanche/Source/AvalancheModifiers/Public/Modifiers/AvaVisibilityModifier.h
+++ b/Engine/Plugins/Experimental/Avalanche/Source/AvalancheModifiers/Public/Modifiers/AvaVisibilityModifier.h
@@ -86,7 +86,7 @@ protected:
 	int32 Index = 0;
 
 	/** Treat index as a range from 0 to index */
-	UPROPERTY(EditInstanceOnly, Setter="SetTreatAsRange", Getter="GetTreatAsRange", Category="Visibility", meta=(AllowPrivateAccess="true"))
+	UPROPERTY(EditInstanceOnly, Setter="SetTreatAsRange", Getter="GetTreatAsRange", Category="Visibility", meta=(DisplayName="Keep Visible", AllowPrivateAccess="true"))
 	bool bTreatAsRange = false;
 
 	/** If true, will hide the child index range instead of showing. */
