Motion Design: fixed issue where asset registry search was done prior to scan finishing, causing intermittent issues where templates levels would not appear.

diff --git a/Engine/Plugins/Experimental/Avalanche/Source/AvalancheEditor/Private/AvaEditorModule.cpp b/Engine/Plugins/Experimental/Avalanche/Source/AvalancheEditor/Private/AvaEditorModule.cpp
--- a/Engine/Plugins/Experimental/Avalanche/Source/AvalancheEditor/Private/AvaEditorModule.cpp
+++ b/Engine/Plugins/Experimental/Avalanche/Source/AvalancheEditor/Private/AvaEditorModule.cpp
@@ -126,6 +126,12 @@ void FAvaEditorModule::ShutdownModule()
 	FCoreDelegates::OnPostEngineInit.RemoveAll(this);
 	FCoreDelegates::OnPreExit.RemoveAll(this);
 
+	if (IAssetRegistry* AssetRegistry = IAssetRegistry::Get())
+	{
+		AssetRegistry->OnKnownGathersComplete().Remove(OnKnownGathersCompleteHandle);
+		OnKnownGathersCompleteHandle.Reset();
+	}
+
 	if (UObjectInitialized() && !IsEngineExitRequested())
 	{
 		UnregisterCustomLayouts();
@@ -151,7 +157,20 @@ void FAvaEditorModule::PostEngineInit()
 	if (FSlateApplication::IsInitialized())
 	{
 		RegisterPropertyEditorCategories();
-		RegisterLevelTemplates();
+
+		if (IAssetRegistry* AssetRegistry = FModuleManager::LoadModuleChecked<FAssetRegistryModule>(TEXT("AssetRegistry")).TryGet())
+		{
+			// Asset Registry is still gathering assets, get notified when the Asset Registry is done scanning the current plugins to register template assets
+			if (AssetRegistry->IsGathering())
+			{
+				OnKnownGathersCompleteHandle = AssetRegistry->OnKnownGathersComplete().AddRaw(this, &FAvaEditorModule::RegisterLevelTemplates);
+			}
+			// Asset Registry is done gathering assets so can register level templates immediately
+			else
+			{
+				RegisterLevelTemplates();
+			}
+		}
 	}
 }
 
@@ -314,7 +333,6 @@ void FAvaEditorModule::RegisterLevelTemplates()
 	if (GUnrealEd)
 	{
 		// Register original template map,
-		// @todo: Deprecate post-beta 0.1.4
 		if (!GUnrealEd->IsTemplateMap(UE::AvaEditor::Private::DefaultLevelPath))
 		{
 			FTemplateMapInfo MapTemplate;
diff --git a/Engine/Plugins/Experimental/Avalanche/Source/AvalancheEditor/Private/AvaEditorModule.h b/Engine/Plugins/Experimental/Avalanche/Source/AvalancheEditor/Private/AvaEditorModule.h
--- a/Engine/Plugins/Experimental/Avalanche/Source/AvalancheEditor/Private/AvaEditorModule.h
+++ b/Engine/Plugins/Experimental/Avalanche/Source/AvalancheEditor/Private/AvaEditorModule.h
@@ -46,4 +46,7 @@ private:
 	void UnregisterCustomLayouts();
 
 	void RegisterLevelTemplates();
+
+	/** Delegate handle to when asset registry is done scanning assets */
+	FDelegateHandle OnKnownGathersCompleteHandle;
 };
