Motion Design: moved the following plugins from /Plugins/Experimental to /Plugins/VirtualProduction:
ActorModifier, ActorModifierCore, Motion Design, ClonerEffector, CustomDetailsView, Material Designer, GeometryMask, OperatorStack, PropertyAnimator, PropertyAnimatorCore, StormSync, StormSync Motion Design Bridge

#virtualized

diff --git a/Engine/Plugins/VirtualProduction/ActorModifier/Config/BaseActorModifier.ini b/Engine/Plugins/VirtualProduction/ActorModifier/Config/BaseActorModifier.ini
new file mode 100644
--- /dev/null
+++ b/Engine/Plugins/VirtualProduction/ActorModifier/Config/BaseActorModifier.ini
@@ -0,0 +1,33 @@
+﻿[CoreRedirects]
++EnumRedirects=(OldName="/Script/AvalancheModifiers.EAvaTransformSharedModifierState",NewName="/Script/ActorModifier.EActorModifierTransformSharedState")
++EnumRedirects=(OldName="/Script/Avalanche.EAvaReferenceContainer",NewName="/Script/ActorModifier.EActorModifierReferenceContainer")
++EnumRedirects=(OldName="/Script/AvalancheModifiers.EAvaVisibilityActor",NewName="/Script/ActorModifier.EActorModifierVisibilityActor")
++EnumRedirects=(OldName="/Script/Avalanche.EAvaCorner2D",NewName="/Script/ActorModifierLayout.EActorModifierGridArrangeCorner2D")
++EnumRedirects=(OldName="/Script/Avalanche.EAvaAxis",NewName="/Script/ActorModifier.EActorModifierAlignment")
++EnumRedirects=(OldName="/Script/AvalancheModifiers.EAvaGridArrangeDirection",NewName="/Script/ActorModifierLayout.EActorModifierGridArrangeDirection")
++EnumRedirects=(OldName="/Script/AvalancheModifiers.EAvaJustifyVertical",NewName="/Script/ActorModifierLayout.EActorModifierJustifyVertical")
++EnumRedirects=(OldName="/Script/AvalancheModifiers.EAvaJustifyHorizontal",NewName="/Script/ActorModifierLayout.EActorModifierJustifyHorizontal")
++EnumRedirects=(OldName="/Script/AvalancheModifiers.EAvaJustifyDepth",NewName="/Script/ActorModifierLayout.EActorModifierJustifyDepth")
++EnumRedirects=(OldName="/Script/AvalancheModifiers.EAvaRadialArrangement",NewName="/Script/ActorModifierLayout.EActorModifierRadialArrangeMode")
++EnumRedirects=(OldName="/Script/AvalancheModifiers.EAvaRadialArrangePlane",NewName="/Script/ActorModifierLayout.EActorModifierRadialArrangePlane")
++EnumRedirects=(OldName="/Script/AvalancheModifiers.EAvaModifiersAxis",NewName="/Script/ActorModifier.EActorModifierAxis")
+
++StructRedirects=(OldName="/Script/AvalancheModifiers.AvaTransformSharedModifierState",NewName="/Script/ActorModifier.ActorModifierTransformSharedModifierState")
++StructRedirects=(OldName="/Script/AvalancheModifiers.AvaTransformSharedActorState",NewName="/Script/ActorModifier.ActorModifierTransformSharedActorState")
++StructRedirects=(OldName="/Script/AvalancheModifiers.AvaSceneTreeActor",NewName="/Script/ActorModifier.ActorModifierSceneTreeActor")
++StructRedirects=(OldName="/Script/AvalancheModifiers.AvaVisibilitySharedModifierState",NewName="/Script/ActorModifier.ActorModifierVisibilitySharedModifierState")
++StructRedirects=(OldName="/Script/AvalancheModifiers.AvaVisibilitySharedActorState",NewName="/Script/ActorModifier.ActorModifierVisibilitySharedActorState")
++StructRedirects=(OldName="/Script/AvalancheModifiers.AvaAlignBetweenWeightedActor",NewName="/Script/ActorModifierLayout.ActorModifierAlignBetweenWeightedActor")
++StructRedirects=(OldName="/Script/Avalanche.AvaAnchorAlignment",NewName="/Script/ActorModifier.ActorModifierAnchorAlignment")
+
++ClassRedirects=(OldName="/Script/AvalancheModifiers.AvaTransformModifierShared",NewName="/Script/ActorModifier.ActorModifierTransformShared")
++ClassRedirects=(OldName="/Script/AvalancheModifiers.AvaVisibilityModifierShared",NewName="/Script/ActorModifier.ActorModifierVisibilityShared")
++ClassRedirects=(OldName="/Script/AvalancheModifiers.AvaAttachmentBaseModifier",NewName="/Script/ActorModifier.ActorModifierAttachmentBaseModifier")
++ClassRedirects=(OldName="/Script/AvalancheModifiers.AvaArrangeBaseModifier",NewName="/Script/ActorModifier.ActorModifierArrangeBaseModifier")
++ClassRedirects=(OldName="/Script/AvalancheModifiers.AvaAlignBetweenModifier",NewName="/Script/ActorModifierLayout.ActorModifierAlignBetweenModifier")
++ClassRedirects=(OldName="/Script/AvalancheModifiers.AvaAutoFollowModifier",NewName="/Script/ActorModifierLayout.ActorModifierAutoFollowModifier")
++ClassRedirects=(OldName="/Script/AvalancheModifiers.AvaGridArrangeModifier",NewName="/Script/ActorModifierLayout.ActorModifierGridArrangeModifier")
++ClassRedirects=(OldName="/Script/AvalancheModifiers.AvaJustifyModifier",NewName="/Script/ActorModifierLayout.ActorModifierJustifyModifier")
++ClassRedirects=(OldName="/Script/AvalancheModifiers.AvaLookAtModifier",NewName="/Script/ActorModifierLayout.ActorModifierLookAtModifier")
++ClassRedirects=(OldName="/Script/AvalancheModifiers.AvaRadialArrangeModifier",NewName="/Script/ActorModifierLayout.ActorModifierRadialArrangeModifier")
++ClassRedirects=(OldName="/Script/ActorModifierLayout.ActorModifierLayoutSplinePathModifier",NewName="/Script/ActorModifierLayout.ActorModifierSplinePathModifier")
diff --git a/Engine/Plugins/VirtualProduction/ActorModifier/Source/ActorModifier/Private/ActorModifierModule.cpp b/Engine/Plugins/VirtualProduction/ActorModifier/Source/ActorModifier/Private/ActorModifierModule.cpp
new file mode 100644
--- /dev/null
+++ b/Engine/Plugins/VirtualProduction/ActorModifier/Source/ActorModifier/Private/ActorModifierModule.cpp
@@ -0,0 +1,5 @@
+﻿// Copyright Epic Games, Inc. All Rights Reserved.
+
+#include "Modules/ModuleManager.h"
+	
+IMPLEMENT_MODULE(FDefaultModuleImpl, ActorModifier)
\ No newline at end of file
diff --git a/Engine/Plugins/VirtualProduction/ActorModifier/Source/ActorModifierEditor/Private/ActorModifierEditorModule.cpp b/Engine/Plugins/VirtualProduction/ActorModifier/Source/ActorModifierEditor/Private/ActorModifierEditorModule.cpp
new file mode 100644
--- /dev/null
+++ b/Engine/Plugins/VirtualProduction/ActorModifier/Source/ActorModifierEditor/Private/ActorModifierEditorModule.cpp
@@ -0,0 +1,34 @@
+﻿// Copyright Epic Games, Inc. All Rights Reserved.
+
+#include "ActorModifierEditorModule.h"
+
+#include "ActorModifierTypes.h"
+#include "Customizations/ActorModifierEditorAnchorAlignmentPropertyTypeCustomization.h"
+#include "Customizations/ActorModifierEditorSplinePathModifierDetailCustomization.h"
+#include "Modifiers/ActorModifierSplinePathModifier.h"
+#include "PropertyEditorModule.h"
+#include "Styles/ActorModifierEditorStyle.h"
+
+IMPLEMENT_MODULE(FActorModifierEditorModule, ActorModifierEditor)
+
+void FActorModifierEditorModule::StartupModule()
+{
+	FActorModifierEditorStyle::Get();
+	
+	FPropertyEditorModule& PropertyModule = FModuleManager::GetModuleChecked<FPropertyEditorModule>(TEXT("PropertyEditor"));
+
+	PropertyModule.RegisterCustomPropertyTypeLayout(FActorModifierAnchorAlignment::StaticStruct()->GetFName()
+	, FOnGetPropertyTypeCustomizationInstance::CreateStatic(&FActorModifierEditorAnchorAlignmentPropertyTypeCustomization::MakeInstance));
+
+	PropertyModule.RegisterCustomClassLayout(UActorModifierSplinePathModifier::StaticClass()->GetFName(), FOnGetDetailCustomizationInstance::CreateStatic(&FActorModifierEditorSplinePathModifierDetailCustomization::MakeInstance));
+}
+
+void FActorModifierEditorModule::ShutdownModule()
+{
+	if (FPropertyEditorModule* PropertyModule = FModuleManager::GetModulePtr<FPropertyEditorModule>(TEXT("PropertyEditor")))
+	{
+		PropertyModule->UnregisterCustomPropertyTypeLayout(FActorModifierAnchorAlignment::StaticStruct()->GetFName());
+		PropertyModule->UnregisterCustomClassLayout(UActorModifierSplinePathModifier::StaticClass()->GetFName());
+	}
+}
+    
diff --git a/Engine/Plugins/VirtualProduction/ActorModifier/Source/ActorModifierLayout/Private/ActorModifierLayoutModule.cpp b/Engine/Plugins/VirtualProduction/ActorModifier/Source/ActorModifierLayout/Private/ActorModifierLayoutModule.cpp
new file mode 100644
--- /dev/null
+++ b/Engine/Plugins/VirtualProduction/ActorModifier/Source/ActorModifierLayout/Private/ActorModifierLayoutModule.cpp
@@ -0,0 +1,5 @@
+﻿// Copyright Epic Games, Inc. All Rights Reserved.
+
+#include "Modules/ModuleManager.h"
+    
+IMPLEMENT_MODULE(FDefaultModuleImpl, ActorModifierLayout)
\ No newline at end of file
diff --git a/Engine/Plugins/VirtualProduction/ActorModifierCore/Config/BaseActorModifierCore.ini b/Engine/Plugins/VirtualProduction/ActorModifierCore/Config/BaseActorModifierCore.ini
new file mode 100644
--- /dev/null
+++ b/Engine/Plugins/VirtualProduction/ActorModifierCore/Config/BaseActorModifierCore.ini
@@ -0,0 +1,10 @@
+﻿[CoreRedirects]
++ClassRedirects=(OldName="/Script/TagCollection.TagCollectionModifierSharedActor",NewName="/Script/ActorModifierCore.ActorModifierCoreSharedActor")
++ClassRedirects=(OldName="/Script/TagCollection.TagCollectionModifierBase",NewName="/Script/ActorModifierCore.ActorModifierCoreBase")
++ClassRedirects=(OldName="/Script/TagCollection.TagCollectionModifierComponent",NewName="/Script/ActorModifierCore.ActorModifierCoreComponent")
++ClassRedirects=(OldName="/Script/TagCollection.TagCollectionModifierSharedObject",NewName="/Script/ActorModifierCore.ActorModifierCoreSharedObject")
++ClassRedirects=(OldName="/Script/TagCollection.TagCollectionModifierStack",NewName="/Script/ActorModifierCore.ActorModifierCoreStack")
++ClassRedirects=(OldName="/Script/TagCollectionEditor.TagCollectionModifierStackCustomization",NewName="/Script/ActorModifierCoreEditor.ActorModifierCoreEditorStackCustomization")
++ClassRedirects=(OldName="/Script/TagCollectionEditor.TagCollectionModifierExtensionSubsystem",NewName="/Script/ActorModifierCoreEditor.ActorModifierCoreEditorSubsystem")
++StructRedirects=(OldName="/Script/TagCollectionEditor.TagCollectionModifierPropertiesWrapper",NewName="/Script/ActorModifierCoreEditor.ActorModifierCoreEditorPropertiesWrapper")
++FunctionRedirects=(OldName="/Script/ActorModifierCore.ActorModifierCoreLibrary.FindModifierOfClass",NewName="/Script/ActorModifierCore.ActorModifierCoreLibrary.FindModifierByClass")
\ No newline at end of file
diff --git a/Engine/Plugins/VirtualProduction/ActorModifierCore/Source/ActorModifierCoreEditor/Private/ActorModifierCoreEditorModule.cpp b/Engine/Plugins/VirtualProduction/ActorModifierCore/Source/ActorModifierCoreEditor/Private/ActorModifierCoreEditorModule.cpp
new file mode 100644
--- /dev/null
+++ b/Engine/Plugins/VirtualProduction/ActorModifierCore/Source/ActorModifierCoreEditor/Private/ActorModifierCoreEditorModule.cpp
@@ -0,0 +1,51 @@
+// Copyright Epic Games, Inc. All Rights Reserved.
+
+#include "ActorModifierCoreEditorModule.h"
+
+#include "ActorModifierCoreBlueprint.h"
+#include "KismetCompilerModule.h"
+#include "Modifiers/ActorModifierCoreStack.h"
+#include "Modifiers/Blueprints/ActorModifierCoreBlueprintBase.h"
+#include "Modifiers/Customizations/ActorModifierCoreEditorDetailCustomization.h"
+#include "Modules/ModuleManager.h"
+#include "PropertyEditorModule.h"
+
+void FActorModifierCoreEditorModule::StartupModule()
+{
+	RegisterDetailCustomizations();
+	RegisterBlueprintCustomizations();
+}
+
+void FActorModifierCoreEditorModule::ShutdownModule()
+{
+	UnregisterDetailCustomizations();
+}
+
+void FActorModifierCoreEditorModule::RegisterDetailCustomizations() const
+{
+	// Register custom layouts
+	static FName PropertyEditor("PropertyEditor");
+	FPropertyEditorModule& PropertyModule = FModuleManager::GetModuleChecked<FPropertyEditorModule>(PropertyEditor);
+
+	PropertyModule.RegisterCustomClassLayout(UActorModifierCoreStack::StaticClass()->GetFName()
+		, FOnGetDetailCustomizationInstance::CreateStatic(&FActorModifierCoreEditorDetailCustomization::MakeInstance));
+}
+
+void FActorModifierCoreEditorModule::UnregisterDetailCustomizations() const
+{
+	// Unregister custom layouts
+	static FName PropertyEditor("PropertyEditor");
+	if (FModuleManager::Get().IsModuleLoaded(PropertyEditor))
+	{
+		FPropertyEditorModule& PropertyModule = FModuleManager::GetModuleChecked<FPropertyEditorModule>(PropertyEditor);
+		PropertyModule.UnregisterCustomClassLayout(UActorModifierCoreStack::StaticClass()->GetFName());
+	}
+}
+
+void FActorModifierCoreEditorModule::RegisterBlueprintCustomizations()
+{
+	IKismetCompilerInterface& KismetCompilerModule = FModuleManager::LoadModuleChecked<IKismetCompilerInterface>("KismetCompiler");
+	KismetCompilerModule.OverrideBPTypeForClass(UActorModifierCoreBlueprintBase::StaticClass(), UActorModifierCoreBlueprint::StaticClass());
+}
+
+IMPLEMENT_MODULE(FActorModifierCoreEditorModule, ActorModifierCoreEditor)
diff --git a/Engine/Plugins/VirtualProduction/Avalanche/Source/Avalanche/Public/AvaLog.h b/Engine/Plugins/VirtualProduction/Avalanche/Source/Avalanche/Public/AvaLog.h
new file mode 100644
--- /dev/null
+++ b/Engine/Plugins/VirtualProduction/Avalanche/Source/Avalanche/Public/AvaLog.h
@@ -0,0 +1,7 @@
+﻿// Copyright Epic Games, Inc. All Rights Reserved.
+
+#pragma once
+
+#include "Logging/LogMacros.h"
+
+AVALANCHE_API DECLARE_LOG_CATEGORY_EXTERN(LogAva, Log, All);
diff --git a/Engine/Plugins/VirtualProduction/Avalanche/Source/Avalanche/Public/AvaSceneState.h b/Engine/Plugins/VirtualProduction/Avalanche/Source/Avalanche/Public/AvaSceneState.h
new file mode 100644
--- /dev/null
+++ b/Engine/Plugins/VirtualProduction/Avalanche/Source/Avalanche/Public/AvaSceneState.h
@@ -0,0 +1,36 @@
+// Copyright Epic Games, Inc. All Rights Reserved.
+
+#pragma once
+
+#include "AvaTag.h"
+#include "Containers/Array.h"
+#include "Containers/Set.h"
+#include "UObject/Object.h"
+#include "AvaSceneState.generated.h"
+
+class UAvaAttribute;
+class UAvaSceneSettings;
+struct FAvaTagHandle;
+
+/** Object providing attribute information of the Scene */
+UCLASS(MinimalAPI)
+class UAvaSceneState : public UObject
+{
+	GENERATED_BODY()
+
+public:
+	void Initialize(UAvaSceneSettings* InSceneSettings);
+
+	AVALANCHE_API bool AddTagAttribute(const FAvaTagHandle& InTagHandle);
+	AVALANCHE_API bool RemoveTagAttribute(const FAvaTagHandle& InTagHandle);
+	AVALANCHE_API bool ContainsTagAttribute(const FAvaTagHandle& InTagHandle) const;
+
+	AVALANCHE_API bool AddNameAttribute(FName InName);
+	AVALANCHE_API bool RemoveNameAttribute(FName InName);
+	AVALANCHE_API bool ContainsNameAttribute(FName InName) const;
+
+private:
+	/** In-play Scene Attributes. Can be added to / removed from while in-play */
+	UPROPERTY(Transient)
+	TArray<TObjectPtr<UAvaAttribute>> SceneAttributes;
+};
diff --git a/Engine/Plugins/VirtualProduction/Avalanche/Source/Avalanche/Public/AvaSceneSubsystem.h b/Engine/Plugins/VirtualProduction/Avalanche/Source/Avalanche/Public/AvaSceneSubsystem.h
new file mode 100644
--- /dev/null
+++ b/Engine/Plugins/VirtualProduction/Avalanche/Source/Avalanche/Public/AvaSceneSubsystem.h
@@ -0,0 +1,39 @@
+// Copyright Epic Games, Inc. All Rights Reserved.
+
+#pragma once
+
+#include "Containers/Map.h"
+#include "Subsystems/WorldSubsystem.h"
+#include "UObject/ObjectPtr.h"
+#include "UObject/WeakInterfacePtr.h"
+#include "AvaSceneSubsystem.generated.h"
+
+class IAvaSceneInterface;
+class ULevel;
+
+UCLASS(MinimalAPI, DisplayName = "Motion Design Scene Subsystem")
+class UAvaSceneSubsystem : public UWorldSubsystem
+{
+	GENERATED_BODY()
+
+public:
+	void RegisterSceneInterface(ULevel* InLevel, IAvaSceneInterface* InSceneInterface);
+
+	/** Gets the Scene Interface for the World's Persistent Level */
+	AVALANCHE_API IAvaSceneInterface* GetSceneInterface() const;
+
+	/** Gets the Scene Interface for the provided Level */
+	AVALANCHE_API IAvaSceneInterface* GetSceneInterface(ULevel* InLevel) const;
+
+	/** Returns the first interface found in the given Level, if any */
+	AVALANCHE_API static IAvaSceneInterface* FindSceneInterface(ULevel* InLevel);
+
+protected:
+	//~ Begin UWorldSubsystem
+	AVALANCHE_API virtual void PostInitialize() override;
+	AVALANCHE_API virtual bool DoesSupportWorldType(const EWorldType::Type InWorldType) const override;
+	//~ End UWorldSubsystem
+
+private:
+	TMap<TWeakObjectPtr<ULevel>, TWeakInterfacePtr<IAvaSceneInterface>> SceneInterfaces;
+};
diff --git a/Engine/Plugins/VirtualProduction/Avalanche/Source/Avalanche/Public/Transition/Conditions/AvaTransitionRCControllerMatchCondition.h b/Engine/Plugins/VirtualProduction/Avalanche/Source/Avalanche/Public/Transition/Conditions/AvaTransitionRCControllerMatchCondition.h
new file mode 100644
--- /dev/null
+++ b/Engine/Plugins/VirtualProduction/Avalanche/Source/Avalanche/Public/Transition/Conditions/AvaTransitionRCControllerMatchCondition.h
@@ -0,0 +1,64 @@
+// Copyright Epic Games, Inc. All Rights Reserved.
+
+#pragma once
+
+#include "AvaRCControllerId.h"
+#include "AvaTransitionEnums.h"
+#include "Conditions/AvaTransitionCondition.h"
+#include "AvaTransitionRCControllerMatchCondition.generated.h"
+
+class UAvaSceneSubsystem;
+class URCVirtualPropertyBase;
+struct FAvaTransitionScene;
+
+USTRUCT()
+struct FAvaTransitionRCControllerMatchConditionInstanceData
+{
+	GENERATED_BODY()
+
+	UPROPERTY(EditAnywhere, Category="Parameter")
+	FAvaRCControllerId ControllerId;
+
+	UPROPERTY(EditAnywhere, Category="Parameter")
+	EAvaTransitionComparisonResult ValueComparisonType = EAvaTransitionComparisonResult::None;
+};
+
+USTRUCT(DisplayName="Compare RC Controller Values", Category="Remote Control")
+struct AVALANCHE_API FAvaTransitionRCControllerMatchCondition : public FAvaTransitionCondition
+{
+	GENERATED_BODY()
+
+	using FInstanceDataType = FAvaTransitionRCControllerMatchConditionInstanceData;
+
+	PRAGMA_DISABLE_DEPRECATION_WARNINGS
+	FAvaTransitionRCControllerMatchCondition() = default;
+	virtual ~FAvaTransitionRCControllerMatchCondition() override = default;
+	FAvaTransitionRCControllerMatchCondition(const FAvaTransitionRCControllerMatchCondition&) = default;
+	FAvaTransitionRCControllerMatchCondition(FAvaTransitionRCControllerMatchCondition&&) = default;
+	FAvaTransitionRCControllerMatchCondition& operator=(const FAvaTransitionRCControllerMatchCondition&) = default;
+	FAvaTransitionRCControllerMatchCondition& operator=(FAvaTransitionRCControllerMatchCondition&&) = default;
+	PRAGMA_ENABLE_DEPRECATION_WARNINGS
+
+	//~ Begin FStateTreeNodeBase
+#if WITH_EDITOR
+	virtual FText GetDescription(const FGuid& InId, FStateTreeDataView InInstanceDataView, const IStateTreeBindingLookup& InBindingLookup, EStateTreeNodeFormatting InFormatting) const override;
+#endif
+	virtual const UStruct* GetInstanceDataType() const override { return FInstanceDataType::StaticStruct(); }
+	virtual void PostLoad(FStateTreeDataView InInstanceDataView) override;
+	virtual bool Link(FStateTreeLinker& InLinker) override;
+	//~ End FStateTreeNodeBase
+
+	//~ Begin FStateTreeConditionBase
+	virtual bool TestCondition(FStateTreeExecutionContext& InContext) const override;
+	//~ End FStateTreeConditionBase
+
+	UE_DEPRECATED(5.5, "ControllerId has been moved to Instance Data")
+	UPROPERTY(meta = (DeprecatedProperty, DeprecationMessage = "Use the Instance Data ControllerId instead"))
+	FAvaRCControllerId ControllerId_DEPRECATED;
+
+	UE_DEPRECATED(5.5, "ValueComparisonType has been moved to Instance Data")
+	UPROPERTY(meta = (DeprecatedProperty, DeprecationMessage = "Use the Instance Data ValueComparisonType instead"))
+	EAvaTransitionComparisonResult ValueComparisonType_DEPRECATED = EAvaTransitionComparisonResult::None;
+
+	TStateTreeExternalDataHandle<UAvaSceneSubsystem> SceneSubsystemHandle;
+};
diff --git a/Engine/Plugins/VirtualProduction/Avalanche/Source/Avalanche/Public/Transition/Tasks/AvaSceneTask.h b/Engine/Plugins/VirtualProduction/Avalanche/Source/Avalanche/Public/Transition/Tasks/AvaSceneTask.h
new file mode 100644
--- /dev/null
+++ b/Engine/Plugins/VirtualProduction/Avalanche/Source/Avalanche/Public/Transition/Tasks/AvaSceneTask.h
@@ -0,0 +1,50 @@
+// Copyright Epic Games, Inc. All Rights Reserved.
+
+#pragma once
+
+#include "AvaTagHandle.h"
+#include "Tasks/AvaTransitionTask.h"
+#include "AvaSceneTask.generated.h"
+
+class IAvaSceneInterface;
+class UAvaSceneSubsystem;
+
+USTRUCT()
+struct FAvaSceneTaskInstanceData
+{
+	GENERATED_BODY()
+
+	UPROPERTY(EditAnywhere, Category="Parameter")
+	FAvaTagHandle TagAttribute;
+};
+
+USTRUCT(meta=(Hidden))
+struct AVALANCHE_API FAvaSceneTask : public FAvaTransitionTask
+{
+	GENERATED_BODY()
+
+	using FInstanceDataType = FAvaSceneTaskInstanceData;
+
+	PRAGMA_DISABLE_DEPRECATION_WARNINGS
+	FAvaSceneTask() = default;
+	virtual ~FAvaSceneTask() override = default;
+	FAvaSceneTask(const FAvaSceneTask&) = default;
+	FAvaSceneTask(FAvaSceneTask&&) = default;
+	FAvaSceneTask& operator=(const FAvaSceneTask&) = default;
+	FAvaSceneTask& operator=(FAvaSceneTask&&) = default;
+	PRAGMA_ENABLE_DEPRECATION_WARNINGS
+
+	//~ Begin FStateTreeNodeBase
+	virtual const UStruct* GetInstanceDataType() const override { return FInstanceDataType::StaticStruct(); }
+	virtual void PostLoad(FStateTreeDataView InInstanceDataView) override;
+	virtual bool Link(FStateTreeLinker& InLinker) override;
+	//~ End FStateTreeNodeBase
+
+	IAvaSceneInterface* GetScene(FStateTreeExecutionContext& InContext) const;
+
+	UE_DEPRECATED(5.5, "TagAttribute has been moved to Instance Data")
+	UPROPERTY(meta = (DeprecatedProperty, DeprecationMessage = "Use the Instance Data TagAttribute instead"))
+	FAvaTagHandle TagAttribute_DEPRECATED;
+
+	TStateTreeExternalDataHandle<UAvaSceneSubsystem> SceneSubsystemHandle;
+};
diff --git a/Engine/Plugins/VirtualProduction/Avalanche/Source/AvalancheEditor/Private/Outliner/AvaOutlinerTabSpawner.cpp b/Engine/Plugins/VirtualProduction/Avalanche/Source/AvalancheEditor/Private/Outliner/AvaOutlinerTabSpawner.cpp
new file mode 100644
--- /dev/null
+++ b/Engine/Plugins/VirtualProduction/Avalanche/Source/AvalancheEditor/Private/Outliner/AvaOutlinerTabSpawner.cpp
@@ -0,0 +1,91 @@
+﻿// Copyright Epic Games, Inc. All Rights Reserved.
+
+#include "AvaOutlinerTabSpawner.h"
+#include "AvaOutlinerExtension.h"
+#include "Framework/Docking/TabManager.h"
+#include "Framework/Docking/WorkspaceItem.h"
+#include "IAvaEditor.h"
+#include "IAvaOutliner.h"
+#include "IAvaOutlinerView.h"
+#include "Styling/AppStyle.h"
+
+#define LOCTEXT_NAMESPACE "AvaOutlinerTabSpawner"
+
+FName FAvaOutlinerTabSpawner::GetTabID(int32 InOutlinerId)
+{
+	return *FString::Printf(TEXT("AvaOutliner%d"), InOutlinerId + 1);
+}
+
+FAvaOutlinerTabSpawner::FAvaOutlinerTabSpawner(int32 InOutlinerId, const TSharedRef<IAvaEditor>& InEditor)
+	: FAvaTabSpawner(InEditor, FAvaOutlinerTabSpawner::GetTabID(InOutlinerId))
+	, OutlinerViewId(InOutlinerId)
+{
+	TabLabel = (OutlinerViewId == 0)
+		? LOCTEXT("FirstTabLabel", "Motion Design Outliner")
+		: FText::Format(LOCTEXT("OtherTabLabels", "Motion Design Outliner {0}"), FText::AsNumber(OutlinerViewId + 1));
+
+	TabTooltipText = LOCTEXT("TabTooltip", "Motion Design Outliner");
+
+	TabIcon = FSlateIcon(FAppStyle::GetAppStyleSetName(), "LevelEditor.Tabs.Outliner");
+}
+
+TSharedRef<SWidget> FAvaOutlinerTabSpawner::CreateTabBody()
+{
+	const TSharedPtr<IAvaEditor> Editor = EditorWeak.Pin();
+	if (!ensure(Editor.IsValid()))
+	{
+		return GetNullWidget();
+	}
+
+	const TSharedPtr<FAvaOutlinerExtension> OutlinerExtension = Editor->FindExtension<FAvaOutlinerExtension>();
+	if (!ensure(OutlinerExtension.IsValid()))
+	{
+		return GetNullWidget();
+	}
+
+	const TSharedPtr<IAvaOutliner> Outliner = OutlinerExtension->GetAvaOutliner();
+	if (!ensure(Outliner.IsValid()))
+	{
+		return GetNullWidget();
+	}
+
+	TSharedPtr<IAvaOutlinerView> OutlinerView = Outliner->GetOutlinerView(OutlinerViewId);
+	if (!OutlinerView.IsValid())
+	{
+		OutlinerView = Outliner->RegisterOutlinerView(OutlinerViewId);
+	}
+
+	if (!ensure(OutlinerView.IsValid()))
+	{
+		return GetNullWidget();
+	}
+	return OutlinerView->GetOutlinerWidget();
+}
+
+FTabSpawnerEntry& FAvaOutlinerTabSpawner::RegisterTabSpawner(const TSharedRef<FTabManager>& InTabManager, const TSharedPtr<FWorkspaceItem>& InWorkspaceMenu)
+{
+	TSharedPtr<FWorkspaceItem> GroupItem;
+
+	const FText OutlinerGroupName = LOCTEXT("SubMenuLabel", "Motion Design Outliner");
+
+	for (const TSharedRef<FWorkspaceItem>& Item : InWorkspaceMenu->GetChildItems())
+	{
+		if (Item->GetDisplayName().ToString() == OutlinerGroupName.ToString())
+		{
+			GroupItem = Item;
+			break;
+		}
+	}
+
+	if (!GroupItem.IsValid())
+	{
+		GroupItem = InWorkspaceMenu->AddGroup(OutlinerGroupName
+			, LOCTEXT("SubMenuTooltip", "Open a new Motion Design Outliner")
+			, FSlateIcon(FAppStyle::GetAppStyleSetName(), "LevelEditor.Tabs.Outliner"));
+	}
+
+	return FAvaTabSpawner::RegisterTabSpawner(InTabManager, nullptr)
+		.SetGroup(GroupItem.ToSharedRef());
+}
+
+#undef LOCTEXT_NAMESPACE
diff --git a/Engine/Plugins/VirtualProduction/Avalanche/Source/AvalancheEditorCore/Private/LevelEditor/AvaLevelEditorStyle.cpp b/Engine/Plugins/VirtualProduction/Avalanche/Source/AvalancheEditorCore/Private/LevelEditor/AvaLevelEditorStyle.cpp
new file mode 100644
--- /dev/null
+++ b/Engine/Plugins/VirtualProduction/Avalanche/Source/AvalancheEditorCore/Private/LevelEditor/AvaLevelEditorStyle.cpp
@@ -0,0 +1,31 @@
+// Copyright Epic Games, Inc. All Rights Reserved.
+
+#include "AvaLevelEditorStyle.h"
+#include "Brushes/SlateImageBrush.h"
+#include "Interfaces/IPluginManager.h"
+#include "Misc/Paths.h"
+#include "Styling/SlateStyleMacros.h"
+#include "Styling/SlateStyleRegistry.h"
+
+FAvaLevelEditorStyle::FAvaLevelEditorStyle()
+	: FSlateStyleSet(TEXT("AvaLevelEditor"))
+{
+	const FVector2f Icon16x16(16.f, 16.f);
+
+	const TSharedPtr<IPlugin> Plugin = IPluginManager::Get().FindPlugin(UE_PLUGIN_NAME);
+	check(Plugin.IsValid());
+
+	SetContentRoot(Plugin->GetBaseDir() / TEXT("Resources"));
+	SetCoreContentRoot(FPaths::EngineContentDir() / TEXT("Slate"));
+	
+	Set("AvaLevelEditor.CreateScene", new IMAGE_BRUSH_SVG("Icons/EditorIcons/LevelEditor_CreateScene_40", Icon16x16));
+	Set("AvaLevelEditor.ActivateScene", new IMAGE_BRUSH_SVG("Icons/EditorIcons/LevelEditor_ActivateScene_40", Icon16x16));
+	Set("AvaLevelEditor.DeactivateScene", new IMAGE_BRUSH_SVG("Icons/EditorIcons/LevelEditor_DeactivateScene_40", Icon16x16));
+
+	FSlateStyleRegistry::RegisterSlateStyle(*this);
+}
+
+FAvaLevelEditorStyle::~FAvaLevelEditorStyle()
+{
+	FSlateStyleRegistry::UnRegisterSlateStyle(*this);
+}
diff --git a/Engine/Plugins/VirtualProduction/Avalanche/Source/AvalancheInteractiveTools/Private/Tools/AvaInteractiveToolsActorToolNull.cpp b/Engine/Plugins/VirtualProduction/Avalanche/Source/AvalancheInteractiveTools/Private/Tools/AvaInteractiveToolsActorToolNull.cpp
new file mode 100644
--- /dev/null
+++ b/Engine/Plugins/VirtualProduction/Avalanche/Source/AvalancheInteractiveTools/Private/Tools/AvaInteractiveToolsActorToolNull.cpp
@@ -0,0 +1,32 @@
+// Copyright Epic Games, Inc. All Rights Reserved.
+
+#include "Tools/AvaInteractiveToolsActorToolNull.h"
+#include "AvaInteractiveToolsCommands.h"
+#include "AvalancheInteractiveToolsModule.h"
+#include "Builders/AvaInteractiveToolsToolBuilder.h"
+#include "Framework/AvaNullActor.h"
+
+UAvaInteractiveToolsActorToolNull::UAvaInteractiveToolsActorToolNull()
+{
+	ActorClass = AAvaNullActor::StaticClass();
+}
+
+void UAvaInteractiveToolsActorToolNull::OnRegisterTool(IAvalancheInteractiveToolsModule* InAITModule)
+{
+	Super::OnRegisterTool(InAITModule);
+
+	FAvaInteractiveToolsToolParameters ToolParameters =
+	{
+		FAvaInteractiveToolsCommands::Get().Tool_Actor_Null,
+		TEXT("Null Actor Tool"),
+		2000,
+		FAvalancheInteractiveToolsCreateBuilder::CreateLambda(
+			[](UEdMode* InEdMode)
+			{
+				return UAvaInteractiveToolsToolBuilder::CreateToolBuilder<UAvaInteractiveToolsActorToolNull>(InEdMode);
+			}),
+		ActorClass
+	};
+
+	InAITModule->RegisterTool(IAvalancheInteractiveToolsModule::CategoryNameActor, MoveTemp(ToolParameters));
+}
diff --git a/Engine/Plugins/VirtualProduction/Avalanche/Source/AvalancheInteractiveTools/Private/Tools/AvaInteractiveToolsActorToolSpline.cpp b/Engine/Plugins/VirtualProduction/Avalanche/Source/AvalancheInteractiveTools/Private/Tools/AvaInteractiveToolsActorToolSpline.cpp
new file mode 100644
--- /dev/null
+++ b/Engine/Plugins/VirtualProduction/Avalanche/Source/AvalancheInteractiveTools/Private/Tools/AvaInteractiveToolsActorToolSpline.cpp
@@ -0,0 +1,32 @@
+// Copyright Epic Games, Inc. All Rights Reserved.
+
+#include "Tools/AvaInteractiveToolsActorToolSpline.h"
+#include "AvaInteractiveToolsCommands.h"
+#include "AvalancheInteractiveToolsModule.h"
+#include "Builders/AvaInteractiveToolsToolBuilder.h"
+#include "Framework/AvaSplineActor.h"
+
+UAvaInteractiveToolsActorToolSpline::UAvaInteractiveToolsActorToolSpline()
+{
+	ActorClass = AAvaSplineActor::StaticClass();
+}
+
+void UAvaInteractiveToolsActorToolSpline::OnRegisterTool(IAvalancheInteractiveToolsModule* InAITModule)
+{
+	Super::OnRegisterTool(InAITModule);
+
+	FAvaInteractiveToolsToolParameters ToolParameters =
+	{
+		FAvaInteractiveToolsCommands::Get().Tool_Actor_Spline,
+		TEXT("Spline Actor Tool"),
+		5000,
+		FAvalancheInteractiveToolsCreateBuilder::CreateLambda(
+			[](UEdMode* InEdMode)
+			{
+				return UAvaInteractiveToolsToolBuilder::CreateToolBuilder<UAvaInteractiveToolsActorToolSpline>(InEdMode);
+			}),
+		ActorClass
+	};
+
+	InAITModule->RegisterTool(IAvalancheInteractiveToolsModule::CategoryNameActor, MoveTemp(ToolParameters));
+}
diff --git a/Engine/Plugins/VirtualProduction/Avalanche/Source/AvalancheMask/Private/AvaMaskSettings.cpp b/Engine/Plugins/VirtualProduction/Avalanche/Source/AvalancheMask/Private/AvaMaskSettings.cpp
new file mode 100644
--- /dev/null
+++ b/Engine/Plugins/VirtualProduction/Avalanche/Source/AvalancheMask/Private/AvaMaskSettings.cpp
@@ -0,0 +1,50 @@
+// Copyright Epic Games, Inc. All Rights Reserved.
+
+#include "AvaMaskSettings.h"
+
+#include "AvaMaskLog.h"
+
+#if WITH_EDITORONLY_DATA
+#include "AssetRegistry/AssetData.h"
+#include "AssetRegistry/AssetRegistryModule.h"
+#include "Interfaces/IPluginManager.h"
+#include "Materials/MaterialFunctionInterface.h"
+#endif
+
+UAvaMaskSettings::UAvaMaskSettings()
+{
+	CategoryName = TEXT("Motion Design");
+	SectionName = TEXT("Mask");
+}
+
+#if WITH_EDITORONLY_DATA
+UMaterialFunctionInterface* UAvaMaskSettings::GetMaterialFunction()
+{
+	if (UMaterialFunctionInterface* StrongMaterialFunction = MaterialFunction.Get())
+	{
+		return StrongMaterialFunction;		
+	}
+
+	UMaterialFunctionInterface* DefaultMaterialFunction = GetDefaultMaterialFunction();
+	MaterialFunction = DefaultMaterialFunction;
+	
+	return DefaultMaterialFunction;
+}
+
+UMaterialFunctionInterface* UAvaMaskSettings::GetDefaultMaterialFunction()
+{
+	// /Script/Engine.MaterialFunction'/GeometryMask/GeometryMask/MF_ApplyGeometryMask.MF_ApplyGeometryMask'
+	// const FString AssetPath = FString(TEXT("/GeometryMask")) / TEXT("GeometryMask") / TEXT("MF_ApplyGeometryMask");
+	// const FString AssetPath = TEXT("/Script/Engine.MaterialFunction'/GeometryMask/GeometryMask/MF_ApplyGeometryMask.MF_ApplyGeometryMask'");
+	const FString AssetPath = TEXT("/Script/Engine.MaterialFunction'/Avalanche/MaskResources/MF_ApplyMask2D_Single.MF_ApplyMask2D_Single'");
+	const FAssetRegistryModule& AssetRegistryModule = FModuleManager::LoadModuleChecked<FAssetRegistryModule>("AssetRegistry");
+	if (const FAssetData AssetData = AssetRegistryModule.Get().GetAssetByObjectPath(FSoftObjectPath(AssetPath));
+		!AssetData.IsValid())
+	{
+		UE_LOG(LogAvaMask, Error, TEXT("Default MaterialFunction asset is invalid or not found: '%s'"), *AssetPath);
+		return nullptr;
+	}
+	
+	return Cast<UMaterialFunctionInterface>(StaticLoadObject(UMaterialFunctionInterface::StaticClass(), this, *AssetPath));
+}
+#endif
diff --git a/Engine/Plugins/VirtualProduction/Avalanche/Source/AvalancheMask/Private/AvaMaskTypes.cpp b/Engine/Plugins/VirtualProduction/Avalanche/Source/AvalancheMask/Private/AvaMaskTypes.cpp
new file mode 100644
--- /dev/null
+++ b/Engine/Plugins/VirtualProduction/Avalanche/Source/AvalancheMask/Private/AvaMaskTypes.cpp
@@ -0,0 +1,86 @@
+﻿// Copyright Epic Games, Inc. All Rights Reserved.
+
+#include "AvaMaskTypes.h"
+
+#include "Engine/Texture.h"
+#include "Materials/MaterialInterface.h"
+
+bool FAvaMask2DMaterialParameters::SetFromMaterial(
+	const UMaterialInterface* InMaterial
+	, const FMaterialParameterInfo& InTextureParameter
+	, const FMaterialParameterInfo& InChannelParameter
+	, const FMaterialParameterInfo& InInvertParameter
+	, const FMaterialParameterInfo& InBaseOpacityParameter
+	, const FMaterialParameterInfo& InPaddingParameter
+	, const FMaterialParameterInfo& InFeatherParameter)
+{
+	if (!ensureAlways(InMaterial))
+	{
+		return false;
+	}
+
+	// Texture
+	{
+		UTexture* FoundTexture = nullptr;
+		if (!InMaterial->GetTextureParameterValue(InTextureParameter, FoundTexture))
+		{
+			return false;
+		}
+		Texture = FoundTexture;
+	}
+
+	// Base Opacity
+	{
+		float FoundValue = 0.0f;
+		if (!InMaterial->GetScalarParameterValue(InBaseOpacityParameter, FoundValue))
+		{
+			return false;
+		}
+		BaseOpacity = FoundValue;
+	}
+
+	// Channel
+	{
+		FLinearColor FoundValue;
+		if (!InMaterial->GetVectorParameterValue(InChannelParameter, FoundValue))
+		{
+			return false;
+		}
+		ChannelAsVector = FoundValue;
+		Channel = UE::GeometryMask::VectorToMaskChannel(ChannelAsVector);
+	}
+
+	// Invert
+	{
+		float FoundValue = 0.0f;
+		if (!InMaterial->GetScalarParameterValue(InInvertParameter, FoundValue))
+		{
+			return false;
+		}
+		bInvert = FoundValue > 0.1;
+	}
+
+	// Padding
+	{
+		FLinearColor FoundValue;
+		if (!InMaterial->GetVectorParameterValue(InPaddingParameter, FoundValue))
+		{
+			return false;			
+		}
+		Padding = FVector2f(FoundValue.R, FoundValue.G);
+	}
+
+	// Feathering
+	{
+		FLinearColor FoundValue;
+		if (!InMaterial->GetVectorParameterValue(InFeatherParameter, FoundValue))
+		{
+			return false;
+		}
+		bApplyFeathering = !FMath::IsNearlyZero(FoundValue.R);
+		OuterFeatherRadius = FoundValue.G;
+		InnerFeatherRadius = FoundValue.B;
+	}
+
+	return true;
+}
diff --git a/Engine/Plugins/VirtualProduction/Avalanche/Source/AvalancheMask/Private/AvaMaskUtilities.h b/Engine/Plugins/VirtualProduction/Avalanche/Source/AvalancheMask/Private/AvaMaskUtilities.h
new file mode 100644
--- /dev/null
+++ b/Engine/Plugins/VirtualProduction/Avalanche/Source/AvalancheMask/Private/AvaMaskUtilities.h
@@ -0,0 +1,106 @@
+// Copyright Epic Games, Inc. All Rights Reserved.
+
+#pragma once
+
+#include "Components/ActorComponent.h"
+#include "Engine/EngineTypes.h"
+#include "GeometryMaskTypes.h"
+#include "Materials/MaterialParameters.h"
+#include "Templates/SubclassOf.h"
+#include "Templates/UnrealTypeTraits.h"
+
+class UMaterialFunctionInterface;
+class UDynamicMaterialInstance;
+class UMaterialInstance;
+
+namespace UE::AvaMask::Internal
+{
+	/** Note that this assumes the parameters have the same name for all materials, which they should if using the provided material function. */
+	static FMaterialParameterInfo TextureParameterInfo = { TEXT("Mask_Texture") };
+	static FMaterialParameterInfo BaseOpacityParameterInfo = { TEXT("Mask_BaseOpacity") };
+	static FMaterialParameterInfo ChannelParameterInfo = { TEXT("Mask_Channel") };
+	static FMaterialParameterInfo InvertParameterInfo = { TEXT("Mask_Invert") };
+	static FMaterialParameterInfo PaddingParameterInfo = { TEXT("Mask_Padding") };
+	static FMaterialParameterInfo FeatherParameterInfo = { TEXT("Mask_Feather") };
+	
+	static TMap<FLinearColor, EGeometryMaskColorChannel> VectorToMaskChannelEnum =
+	{
+		{ {1, 0, 0, 0}, EGeometryMaskColorChannel::Red },
+		{ {0, 1, 0, 0}, EGeometryMaskColorChannel::Green },
+		{ {0, 0, 1, 0}, EGeometryMaskColorChannel::Blue },
+		{ {0, 0, 0, 1}, EGeometryMaskColorChannel::Alpha }
+	};
+	
+	static TMap<EGeometryMaskColorChannel, FLinearColor> MaskChannelEnumToVector =
+	{
+		{ EGeometryMaskColorChannel::Red, {1, 0, 0, 0} },
+		{ EGeometryMaskColorChannel::Green, {0, 1, 0, 0} },
+		{ EGeometryMaskColorChannel::Blue, {0, 0, 1, 0} },
+		{ EGeometryMaskColorChannel::Alpha, {0, 0, 0, 1} }
+	};
+
+	static FName HandleTag = TEXT("Modifier.Mask2D");
+
+	/** Utility function, will set "OtherActor" if the components owner differs from the provided one. */
+	FSoftComponentReference MakeComponentReference(const AActor* InOwner, const UActorComponent* InComponent);
+
+	/** Creates a unique (deterministic) key based on the provided parameters. */
+	uint32 MakeMaterialInstanceKey(const UMaterialInterface* InMaterial, const FName InMaskChannelName, EBlendMode InBlendMode, const int32 InSeed = 128);
+
+	UActorComponent* FindOrAddComponent(const TSubclassOf<UActorComponent>& InComponentClass, AActor* InActor);
+
+	template<
+		typename InComponentClass
+		UE_REQUIRES(std::is_base_of_v<UActorComponent, InComponentClass>)>
+	InComponentClass* FindOrAddComponent(AActor* InActor)
+	{
+		return Cast<InComponentClass>(FindOrAddComponent(InComponentClass::StaticClass(), InActor));
+	}
+
+	void RemoveComponentByInterface(const TSubclassOf<UInterface>& InInterfaceType, const AActor* InActor);
+
+	template<
+		typename InInterfaceType
+		UE_REQUIRES(std::is_base_of_v<UInterface, InInterfaceType>)>
+	void RemoveComponentByInterface(AActor* InActor)
+	{
+		RemoveComponentByInterface(InInterfaceType::StaticClass(), InActor);
+	}
+
+	FString GetGeneratedMaterialPath();
+
+	FString GetBlendModeString(const EBlendMode InBlendMode);
+
+	/** Returns whichever blend mode satisfies the requirement AND user specified blend mode */
+	EBlendMode GetTargetBlendMode(const EBlendMode InFromMaterial, const EBlendMode InRequired);
+
+	template <typename ElementType, typename KeyFuncs, typename AllocatorType, typename LookupType, typename LambdaType>
+	ElementType& FindOrAddByLambda(TSet<ElementType, KeyFuncs, AllocatorType>& Set, const LookupType& Key, LambdaType Lambda)
+	{
+		uint32 KeyHash = KeyFuncs::GetKeyHash(Key);
+		if (ElementType* Found = Set.FindByHash(Key))
+		{
+			return *Found;
+		}
+
+		ElementType NewElement = Invoke(Lambda, Key);
+		check(KeyFuncs::GetKeyHash(NewElement) == KeyHash);
+		FSetElementId NewElementId = Set.AddByHash(KeyHash, MoveTemp(NewElement));
+		return Set[NewElementId];
+	}
+
+	template <typename KeyType, typename ValueType, typename AllocatorType, typename KeyFuncs, typename InitKeyType, typename LambdaType>
+	ValueType& FindOrAddByLambda(TMap<KeyType, ValueType, AllocatorType, KeyFuncs>& Map, InitKeyType&& Key, LambdaType Lambda)
+	{
+		uint32 KeyHash = KeyFuncs::GetKeyHash(Key);
+		if (ValueType* Found = Map.FindByHash(KeyHash, Key))
+		{
+			return *Found;
+		}
+
+		ValueType NewElement = Invoke(Lambda);
+		return Map.AddByHash(KeyHash, Forward<InitKeyType>(Key), MoveTemp(NewElement));
+	}
+
+	UMaterialInterface* GetNonTransientParentMaterial(const UMaterialInstance* InMaterialInstance);
+}
diff --git a/Engine/Plugins/VirtualProduction/Avalanche/Source/AvalancheMask/Private/Handling/AvaMaskText3DActorMaterialCollectionHandle.cpp b/Engine/Plugins/VirtualProduction/Avalanche/Source/AvalancheMask/Private/Handling/AvaMaskText3DActorMaterialCollectionHandle.cpp
new file mode 100644
--- /dev/null
+++ b/Engine/Plugins/VirtualProduction/Avalanche/Source/AvalancheMask/Private/Handling/AvaMaskText3DActorMaterialCollectionHandle.cpp
@@ -0,0 +1,506 @@
+﻿// Copyright Epic Games, Inc. All Rights Reserved.
+
+#include "Handling/AvaMaskText3DActorMaterialCollectionHandle.h"
+
+#include "AvaMaskLog.h"
+#include "AvaMaskMaterialReference.h"
+#include "AvaMaskUtilities.h"
+#include "AvaObjectHandleSubsystem.h"
+#include "AvaText3DComponent.h"
+#include "Components/PrimitiveComponent.h"
+#include "Extensions/Text3DDefaultMaterialExtension.h"
+#include "GameFramework/Actor.h"
+#include "Handling/AvaHandleUtilities.h"
+#include "IAvaMaskMaterialCollectionHandle.h"
+#include "IAvaMaskMaterialHandle.h"
+#include "Materials/MaterialInterface.h"
+#include "Text3DComponent.h"
+
+namespace UE::AvaMask::Private
+{
+	EText3DMaterialBlendMode GetTargetTranslucencyType(
+		const EText3DMaterialBlendMode InFromMaterial
+		, const EBlendMode InRequired)
+	{
+		if (InRequired == EBlendMode::BLEND_Opaque)
+		{
+			return InFromMaterial;
+		}
+
+		if (InFromMaterial == EText3DMaterialBlendMode::Opaque)
+		{
+			return EText3DMaterialBlendMode::Translucent;
+		}
+
+		return InFromMaterial;
+	}
+}
+
+FAvaMaskText3DActorMaterialCollectionHandle::FAvaMaskText3DActorMaterialCollectionHandle(AActor* InActor)
+	: WeakActor(InActor)
+	, WeakComponent(InActor->GetComponentByClass<UText3DComponent>())
+{
+	if (UText3DComponent* Text3DComponent = WeakComponent.Get())
+	{
+		Text3DComponent->OnTextPostUpdate().AddRaw(this, &FAvaMaskText3DActorMaterialCollectionHandle::OnTextPostUpdate);
+	}
+}
+
+FAvaMaskText3DActorMaterialCollectionHandle::~FAvaMaskText3DActorMaterialCollectionHandle()
+{
+	if (UText3DComponent* Text3DComponent = WeakComponent.Get())
+	{
+		Text3DComponent->OnTextPostUpdate().RemoveAll(this);
+	}
+}
+
+TArray<TObjectPtr<UMaterialInterface>> FAvaMaskText3DActorMaterialCollectionHandle::GetMaterials()
+{
+	TArray<TObjectPtr<UMaterialInterface>> Materials;
+	Materials.Reserve(GetNumMaterials());
+	
+	if (const UText3DComponent* TextComponent = GetComponent())
+	{
+		Materials.Emplace(TextComponent->GetFrontMaterial());
+		Materials.Emplace(TextComponent->GetBevelMaterial());
+		Materials.Emplace(TextComponent->GetExtrudeMaterial());
+		Materials.Emplace(TextComponent->GetBackMaterial());
+	}
+
+	return Materials;
+}
+
+TArray<TSharedPtr<IAvaMaskMaterialHandle>> FAvaMaskText3DActorMaterialCollectionHandle::GetMaterialHandles()
+{
+	if (MaterialHandles.IsEmpty() || !UE::Ava::Internal::IsHandleValid(MaterialHandles.Last()))
+	{
+		if (UText3DComponent* TextComponent = GetComponent())
+		{
+			UAvaObjectHandleSubsystem* HandleSubsystem = UE::Ava::Internal::GetObjectHandleSubsystem();
+			if (!HandleSubsystem)
+			{
+				return MaterialHandles;
+			}
+
+			const UText3DMaterialExtensionBase* MaterialExtension = TextComponent->GetMaterialExtension();
+
+			UMaterialInterface* FrontMaterial = MaterialExtension->GetMaterial(EText3DGroupType::Front);
+			UMaterialInterface* BackMaterial = MaterialExtension->GetMaterial(EText3DGroupType::Back);
+			UMaterialInterface* ExtrudeMaterial = MaterialExtension->GetMaterial(EText3DGroupType::Extrude);
+			UMaterialInterface* BevelMaterial = MaterialExtension->GetMaterial(EText3DGroupType::Bevel);
+
+			MaterialHandles.Init(nullptr, static_cast<int32>(EText3DGroupType::TypeCount));
+			MaterialHandles[static_cast<int32>(EText3DGroupType::Front)] = HandleSubsystem->MakeHandle<IAvaMaskMaterialHandle>(FrontMaterial, UE::AvaMask::Internal::HandleTag);
+			MaterialHandles[static_cast<int32>(EText3DGroupType::Bevel)] = HandleSubsystem->MakeHandle<IAvaMaskMaterialHandle>(BevelMaterial, UE::AvaMask::Internal::HandleTag);
+			MaterialHandles[static_cast<int32>(EText3DGroupType::Extrude)] = HandleSubsystem->MakeHandle<IAvaMaskMaterialHandle>(ExtrudeMaterial, UE::AvaMask::Internal::HandleTag);
+			MaterialHandles[static_cast<int32>(EText3DGroupType::Back)] = HandleSubsystem->MakeHandle<IAvaMaskMaterialHandle>(BackMaterial, UE::AvaMask::Internal::HandleTag);
+		}
+	}
+
+	return MaterialHandles;
+}
+
+void FAvaMaskText3DActorMaterialCollectionHandle::SetMaterial(
+	const FSoftComponentReference& InComponent
+	, const int32 InSlotIdx
+	, UMaterialInterface* InMaterial)
+{
+	if (!ensure(InSlotIdx < 4))
+	{
+		UE_LOG(LogAvaMask, Error, TEXT("SlotIdx out of range, should be from 0-3, was %u"), InSlotIdx);
+		return;
+	}
+
+	if (UText3DComponent* TextComponent = GetComponent())
+	{
+		switch (InSlotIdx)
+		{
+		case 0:
+			TextComponent->SetFrontMaterial(InMaterial);
+
+		case 1:
+			TextComponent->SetBevelMaterial(InMaterial);
+
+		case 2:
+			TextComponent->SetExtrudeMaterial(InMaterial);
+
+		case 3:
+			TextComponent->SetBackMaterial(InMaterial);
+
+		default:
+			return;
+		}
+	}
+}
+
+void FAvaMaskText3DActorMaterialCollectionHandle::SetMaterials(
+	const TArray<TObjectPtr<UMaterialInterface>>& InMaterials
+	, const TBitArray<>& InSetToggle)
+{
+	if (!ensure(InMaterials.Num() == GetNumMaterials()))
+	{
+		UE_LOG(LogAvaMask, Warning, TEXT("Expected 4 materials, got %u"), InMaterials.Num());
+		return;
+	}
+
+	if (!ensure(InSetToggle.Num() == GetNumMaterials()))
+	{
+		UE_LOG(LogAvaMask, Warning, TEXT("Expected 4 materials, got %u"), InSetToggle.Num());
+		return;
+	}
+
+	if (UText3DComponent* TextComponent = GetComponent())
+	{
+		if (InSetToggle[0])
+		{
+			TextComponent->SetFrontMaterial(InMaterials[0]);	
+		}
+
+		if (InSetToggle[1])
+		{
+			TextComponent->SetBevelMaterial(InMaterials[1]);	
+		}
+
+		if (InSetToggle[2])
+		{
+			TextComponent->SetExtrudeMaterial(InMaterials[2]);	
+		}
+
+		if (InSetToggle[3])
+		{
+			TextComponent->SetBackMaterial(InMaterials[3]);	
+		}
+	}
+}
+
+int32 FAvaMaskText3DActorMaterialCollectionHandle::GetNumMaterials() const
+{
+	constexpr int32 TextMaterialCount = static_cast<int32>(EText3DGroupType::TypeCount);
+	static_assert(TextMaterialCount == 4);
+	
+	return TextMaterialCount;
+}
+
+void FAvaMaskText3DActorMaterialCollectionHandle::ForEachMaterial(
+	TFunctionRef<bool(
+		const FSoftComponentReference& InComponent
+		, const int32 InIdx
+		, UMaterialInterface* InMaterial)> InFunction)
+{
+	if (const UText3DComponent* TextComponent = GetComponent())
+	{
+		const FSoftComponentReference ComponentReference = UE::AvaMask::Internal::MakeComponentReference(TextComponent->GetOwner(), TextComponent);
+		if (!InFunction(ComponentReference, static_cast<int32>(EText3DGroupType::Front), TextComponent->GetFrontMaterial()))
+		{
+			return;
+		}
+
+		if (!InFunction(ComponentReference, static_cast<int32>(EText3DGroupType::Bevel), TextComponent->GetBevelMaterial()))
+		{
+			return;
+		}
+
+		if (!InFunction(ComponentReference, static_cast<int32>(EText3DGroupType::Extrude), TextComponent->GetExtrudeMaterial()))
+		{
+			return;
+		}
+
+		InFunction(ComponentReference, static_cast<int32>(EText3DGroupType::Back), TextComponent->GetBackMaterial());
+	}
+}
+
+void FAvaMaskText3DActorMaterialCollectionHandle::ForEachMaterialHandle(TFunctionRef<bool(const FSoftComponentReference& InComponent, const int32 InIdx, const bool bInIsSlotOccupied, const TSharedPtr<IAvaMaskMaterialHandle>& InMaterialHandle)> InFunction)
+{
+	if (const UText3DComponent* TextComponent = GetComponent())
+	{
+		// Effectively refreshes material handles if necessary
+		GetMaterialHandles();
+		
+		const FSoftComponentReference ComponentReference = UE::AvaMask::Internal::MakeComponentReference(TextComponent->GetOwner(), TextComponent);
+		TSharedPtr<IAvaMaskMaterialHandle> MaterialHandle = MaterialHandles[static_cast<int32>(EText3DGroupType::Front)];
+		if (!InFunction(ComponentReference, static_cast<int32>(EText3DGroupType::Front), MaterialHandle.IsValid(), MaterialHandle))
+		{
+			return;
+		}
+
+		MaterialHandle = MaterialHandles[static_cast<int32>(EText3DGroupType::Bevel)];
+		if (!InFunction(ComponentReference, static_cast<int32>(EText3DGroupType::Bevel), MaterialHandle.IsValid(), MaterialHandle))
+		{
+			return;
+		}
+
+		MaterialHandle = MaterialHandles[static_cast<int32>(EText3DGroupType::Extrude)];
+		if (!InFunction(ComponentReference, static_cast<int32>(EText3DGroupType::Extrude), MaterialHandle.IsValid(), MaterialHandle))
+		{
+			return;
+		}
+
+		MaterialHandle = MaterialHandles[static_cast<int32>(EText3DGroupType::Back)];
+		InFunction(ComponentReference, static_cast<int32>(EText3DGroupType::Back), MaterialHandle.IsValid(), MaterialHandle);
+	}
+}
+
+void FAvaMaskText3DActorMaterialCollectionHandle::MapEachMaterial(
+	TFunctionRef<UMaterialInterface*(
+		const FSoftComponentReference& InComponent
+		, const int32 InIdx
+		, UMaterialInterface* InMaterial)> InFunction)
+{	
+	if (const UText3DComponent* TextComponent = GetComponent())
+	{
+		TArray<TObjectPtr<UMaterialInterface>> MappedMaterials;
+		MappedMaterials.Init(nullptr, static_cast<TArray<TObjectPtr<UMaterialInterface>>::SizeType>(EText3DGroupType::TypeCount));
+
+		TBitArray<> SetFlags;
+        SetFlags.Init(false, MappedMaterials.Num());
+
+		const FSoftComponentReference ComponentReference = UE::AvaMask::Internal::MakeComponentReference(TextComponent->GetOwner(), TextComponent);
+		
+		MappedMaterials[0] = InFunction(ComponentReference, static_cast<int32>(EText3DGroupType::Front), TextComponent->GetFrontMaterial());
+		SetFlags[0] = MappedMaterials[0] != nullptr;
+
+		MappedMaterials[1] = InFunction(ComponentReference, static_cast<int32>(EText3DGroupType::Bevel), TextComponent->GetBevelMaterial());
+		SetFlags[1] = MappedMaterials[1] != nullptr;
+
+		MappedMaterials[2] = InFunction(ComponentReference, static_cast<int32>(EText3DGroupType::Extrude), TextComponent->GetExtrudeMaterial());
+		SetFlags[2] = MappedMaterials[2] != nullptr;
+
+		MappedMaterials[3] = InFunction(ComponentReference, static_cast<int32>(EText3DGroupType::Back), TextComponent->GetBackMaterial());
+		SetFlags[3] = MappedMaterials[3] != nullptr;
+
+		SetMaterials(MappedMaterials, SetFlags);
+	}
+}
+
+void FAvaMaskText3DActorMaterialCollectionHandle::MapEachMaterialHandle(TFunctionRef<TSharedPtr<IAvaMaskMaterialHandle>(const FSoftComponentReference& InComponent, const int32 InIdx, const bool bInIsSlotOccupied, const TSharedPtr<IAvaMaskMaterialHandle>& InMaterialHandle)> InFunction)
+{
+	if (const UText3DComponent* TextComponent = GetComponent())
+	{
+		// Effectively refreshes material handles if necessary
+		GetMaterialHandles();
+
+		TArray<TObjectPtr<UMaterialInterface>> MappedMaterials;
+		MappedMaterials.Reserve(GetNumMaterials());
+
+		TBitArray<> SetFlags;
+		SetFlags.Init(false, GetNumMaterials());
+
+		const FSoftComponentReference ComponentReference = UE::AvaMask::Internal::MakeComponentReference(TextComponent->GetOwner(), TextComponent);
+
+		auto TryMap = [this, InFunction, &ComponentReference, &MappedMaterials, &SetFlags](const EText3DGroupType InMaterialGroup)
+		{
+			int32 InMaterialGroupIdx = static_cast<int32>(InMaterialGroup);			
+			
+			TSharedPtr<IAvaMaskMaterialHandle>& MaterialHandle = MaterialHandles[InMaterialGroupIdx];
+			if (ensureAlways(UE::Ava::Internal::IsHandleValid(MaterialHandle)))
+			{
+				const TSharedPtr<IAvaMaskMaterialHandle> MappedMaterialHandle = InFunction(ComponentReference, InMaterialGroupIdx, true, MaterialHandle);
+				UMaterialInterface* RemappedMaterial = nullptr;
+				if (MappedMaterialHandle.IsValid())
+				{
+					RemappedMaterial = MappedMaterialHandle->GetMaterial();	
+				}
+				else
+				{
+					RemappedMaterial = MaterialHandle->GetMaterial();
+				}
+
+				MappedMaterials[InMaterialGroupIdx] = RemappedMaterial;
+				SetFlags[InMaterialGroupIdx] = RemappedMaterial != nullptr;
+			}	
+		};
+
+		TryMap(EText3DGroupType::Front);
+		TryMap(EText3DGroupType::Bevel);
+		TryMap(EText3DGroupType::Extrude);
+		TryMap(EText3DGroupType::Back);
+
+		SetMaterials(MappedMaterials, SetFlags);
+	}
+}
+
+bool FAvaMaskText3DActorMaterialCollectionHandle::IsValid() const
+{
+	return WeakActor.IsValid() && WeakComponent.IsValid();
+}
+
+bool FAvaMaskText3DActorMaterialCollectionHandle::SaveOriginalState(const FStructView& InHandleData)
+{
+	if (FHandleData* HandleData = InHandleData.GetPtr<FHandleData>())
+	{
+		ForEachMaterialHandle([this, HandleData](
+			const FSoftComponentReference& InComponent
+			, const int32 InSlotIdx
+			, const bool bInIsSlotOccupied
+			, const TSharedPtr<IAvaMaskMaterialHandle>& InMaterialHandle)
+		{
+			if (InMaterialHandle.IsValid())
+			{
+				FInstancedStruct& MaterialHandleData = HandleData->GroupMaterialData.FindOrAdd(InSlotIdx, InMaterialHandle->MakeDataStruct());
+				return InMaterialHandle->SaveOriginalState(MaterialHandleData);
+			}
+
+			return false;
+		});
+
+		if (const UText3DComponent* TextComponent = GetComponent())
+		{
+			if (const UText3DDefaultMaterialExtension* MaterialExtension = TextComponent->GetCastedMaterialExtension<UText3DDefaultMaterialExtension>())
+			{
+				HandleData->BlendMode = MaterialExtension->GetBlendMode();
+			}
+
+			return true;
+		}
+		
+		return true;
+	}
+
+	return false;
+}
+
+bool FAvaMaskText3DActorMaterialCollectionHandle::ApplyOriginalState(const FStructView& InHandleData)
+{
+	if (const FHandleData* HandleData = InHandleData.GetPtr<FHandleData>())
+	{
+		if (UText3DComponent* TextComponent = GetComponent())
+		{
+			// Only restore if the original was NOT translucent AND the current alpha is 1.0
+			if (UText3DDefaultMaterialExtension* MaterialExtension = TextComponent->GetCastedMaterialExtension<UText3DDefaultMaterialExtension>())
+			{
+				MaterialExtension->SetBlendMode(HandleData->BlendMode);
+			}
+			
+			Super::ApplyOriginalState(InHandleData);
+			
+			return true;
+		}
+
+		return false;
+	}
+
+	return false;
+}
+
+bool FAvaMaskText3DActorMaterialCollectionHandle::ApplyModifiedState(const FAvaMask2DSubjectParameters& InModifiedParameters,
+	const FStructView& InHandleData)
+{
+	if (InHandleData.GetPtr<FHandleData>())
+	{
+		if (UText3DComponent* TextComponent = GetComponent())
+		{
+			if (UText3DDefaultMaterialExtension* MaterialExtension = TextComponent->GetCastedMaterialExtension<UText3DDefaultMaterialExtension>())
+			{
+				const EText3DMaterialBlendMode TargetTranslucencyStyle = UE::AvaMask::Private::GetTargetTranslucencyType(MaterialExtension->GetBlendMode(), InModifiedParameters.MaterialParameters.BlendMode);
+				MaterialExtension->SetBlendMode(TargetTranslucencyStyle);
+
+				// Refresh materials if needed immediately
+				constexpr bool bImmediateUpdate = true;
+				TextComponent->RequestUpdate(EText3DRendererFlags::None, bImmediateUpdate);
+			}
+		}
+
+		if (!Super::ApplyModifiedState(InModifiedParameters, InHandleData))
+		{
+			return false;
+		}
+
+		return true;
+	}
+
+	return false;
+}
+
+bool FAvaMaskText3DActorMaterialCollectionHandle::IsSupported(const FAvaMaskMaterialReference& InInstance, FName InTag)
+{
+	if (InTag == UE::AvaMask::Internal::HandleTag)
+	{
+		if (const AActor* Actor = InInstance.GetTypedObject<AActor>())
+		{
+			return Actor->GetComponentByClass<UText3DComponent>() != nullptr;
+		}
+	}
+
+	return false;
+}
+
+FStructView FAvaMaskText3DActorMaterialCollectionHandle::GetMaterialHandleData(
+	FHandleData* InParentHandleData
+	, const FSoftComponentReference& InComponent
+	, const int32 InSlotIdx)
+{
+	if (InParentHandleData)
+	{
+		if (ensure(InParentHandleData->GroupMaterialData.Contains(InSlotIdx)))
+		{
+			return InParentHandleData->GroupMaterialData[InSlotIdx];
+		}
+	}
+
+	return nullptr;
+}
+
+FStructView FAvaMaskText3DActorMaterialCollectionHandle::GetOrAddMaterialHandleData(
+	FHandleData* InParentHandleData
+	, const TSharedPtr<IAvaMaskMaterialHandle>& InMaterialHandle
+	, const FSoftComponentReference& InComponent
+	, const int32 InSlotIdx)
+{
+	if (InParentHandleData)
+	{
+		return InParentHandleData->GroupMaterialData.FindOrAdd(InSlotIdx, InMaterialHandle->MakeDataStruct());
+	}
+
+	return nullptr;
+}
+
+UText3DComponent* FAvaMaskText3DActorMaterialCollectionHandle::GetComponent()
+{
+	if (UText3DComponent* Component = WeakComponent.Get())
+	{
+		return Component;
+	}
+
+	if (const AActor* Actor = WeakActor.Get())
+	{
+		WeakComponent = Actor->GetComponentByClass<UText3DComponent>();
+		if (UText3DComponent* Component = WeakComponent.Get())
+		{
+			return Component;
+		}
+	}
+
+	return nullptr;
+}
+
+void FAvaMaskText3DActorMaterialCollectionHandle::OnTextPostUpdate(UText3DComponent* InComponent, EText3DRendererFlags InFlags)
+{
+	if (EnumHasAnyFlags(InFlags, EText3DRendererFlags::Material))
+	{
+		constexpr int32 FrontIndex = static_cast<int32>(EText3DGroupType::Front);
+		constexpr int32 BackIndex = static_cast<int32>(EText3DGroupType::Back);
+		constexpr int32 ExtrudeIndex = static_cast<int32>(EText3DGroupType::Extrude);
+		constexpr int32 BevelIndex = static_cast<int32>(EText3DGroupType::Bevel);
+		
+		bool bChanged = false;
+		if (MaterialHandles.IsEmpty() || !UE::Ava::Internal::IsHandleValid(MaterialHandles.Last()))
+		{
+			GetMaterialHandles();
+			bChanged = true;
+		}
+		else if (MaterialHandles[FrontIndex]->GetMaterial() != InComponent->GetFrontMaterial()
+			|| MaterialHandles[BevelIndex]->GetMaterial() != InComponent->GetBevelMaterial()
+			|| MaterialHandles[ExtrudeIndex]->GetMaterial() != InComponent->GetExtrudeMaterial()
+			|| MaterialHandles[BackIndex]->GetMaterial() != InComponent->GetBackMaterial())
+		{
+			MaterialHandles.Empty();
+			GetMaterialHandles();
+			bChanged = true;
+		}
+
+		if (bChanged)
+		{
+			OnSourceMaterialsChanged().ExecuteIfBound(InComponent, MaterialHandles);
+		}
+	}
+}
diff --git a/Engine/Plugins/VirtualProduction/Avalanche/Source/AvalancheMask/Private/Handling/IAvaMaskMaterialHandle.h b/Engine/Plugins/VirtualProduction/Avalanche/Source/AvalancheMask/Private/Handling/IAvaMaskMaterialHandle.h
new file mode 100644
--- /dev/null
+++ b/Engine/Plugins/VirtualProduction/Avalanche/Source/AvalancheMask/Private/Handling/IAvaMaskMaterialHandle.h
@@ -0,0 +1,165 @@
+﻿// Copyright Epic Games, Inc. All Rights Reserved.
+
+#pragma once
+
+#include "AvaMaskLog.h"
+#include "AvaMaskTypes.h"
+#include "AvaMaskUtilities.h"
+#include "Engine/Texture.h"
+#include "IAvaMaterialHandle.h"
+#include "StructUtils/InstancedStruct.h"
+#include "Materials/AvaMaterialUtils.h"
+#include "Materials/MaterialInstance.h"
+#include "Materials/MaterialInstanceDynamic.h"
+#include "StructUtils/StructView.h"
+
+class UAvaMaskMaterialInstanceSubsystem;
+
+namespace UE::AvaMask::Internal
+{
+	UAvaMaskMaterialInstanceSubsystem* GetMaterialInstanceSubsystem();
+}
+
+/** Interface between the mask system and an individual material. */
+class IAvaMaskMaterialHandle
+	: public IAvaMaterialHandle
+{
+public:
+	UE_AVA_INHERITS(IAvaMaskMaterialHandle, IAvaMaterialHandle);
+
+	virtual ~IAvaMaskMaterialHandle() override = default;
+
+	virtual FInstancedStruct MakeDataStruct() = 0;
+	virtual const UScriptStruct* GetDataStructType() = 0;
+
+	virtual bool GetMaskParameters(FAvaMask2DMaterialParameters& OutParameters) = 0;
+	virtual bool SetMaskParameters(UTexture* InTexture, float InBaseOpacity, EGeometryMaskColorChannel InChannel, bool bInInverted = false, const FVector2f& InPadding = FVector2f::Zero(), bool bInApplyFeathering = false, float InOuterFeathering = 0.0f, float InInnerFeathering = 0.0f) = 0;
+
+	virtual const EBlendMode GetBlendMode() = 0;
+	virtual void SetBlendMode(const EBlendMode InBlendMode) = 0;
+
+	/** Provides the opportunity to propagate changes to a derived material. */
+	virtual void SetChildMaterial(UMaterialInstanceDynamic* InChildMaterial) = 0;
+
+	virtual bool SaveOriginalState(const FStructView& InHandleData) = 0;
+	
+	virtual bool ApplyOriginalState(
+		const FStructView& InHandleData
+		, TUniqueFunction<void(UMaterialInterface*)>&& InMaterialSetter) = 0;
+	
+	virtual bool ApplyModifiedState(
+		const FAvaMask2DSubjectParameters& InModifiedParameters
+		, const FStructView& InHandleData
+		, TUniqueFunction<void(UMaterialInterface*)>&& InMaterialSetter) = 0;
+
+	virtual bool HasRequiredParameters(TArray<FString>& OutMissingParameterNames) = 0;	
+};
+
+template <typename HandleDataType>
+class TAvaMaskMaterialHandle
+	: public IAvaMaskMaterialHandle
+{
+	static_assert(TModels_V<CStaticStructProvider, HandleDataType>, "HandleDataType should be a UStruct");
+
+public:
+	UE_AVA_INHERITS_WITH_SUPER(TAvaMaskMaterialHandle, IAvaMaskMaterialHandle);
+	
+	using FHandleData = HandleDataType;
+
+public:
+	virtual FInstancedStruct MakeDataStruct() override;
+	virtual const UScriptStruct* GetDataStructType() override;
+
+	virtual bool GetMaskParameters(FAvaMask2DMaterialParameters& OutParameters) override;
+	virtual bool SetMaskParameters(UTexture* InTexture, float InBaseOpacity, EGeometryMaskColorChannel InChannel, bool bInInverted = false, const FVector2f& InPadding = FVector2f::Zero(), bool bInApplyFeathering = false, float InOuterFeathering = 0.0f, float InInnerFeathering = 0.0f) override;
+
+	virtual bool HasRequiredParameters(TArray<FString>& OutMissingParameterNames) override;
+
+protected:
+	virtual UMaterialInstanceDynamic* GetMaterialInstance() = 0;
+};
+
+template <typename HandleDataType>
+FInstancedStruct TAvaMaskMaterialHandle<HandleDataType>::MakeDataStruct()
+{
+	FInstancedStruct MaterialData = FInstancedStruct::Make<FHandleData>();
+	return MaterialData;
+}
+
+template <typename HandleDataType>
+const UScriptStruct* TAvaMaskMaterialHandle<HandleDataType>::GetDataStructType()
+{
+	return FHandleData::StaticStruct();
+}
+
+template <typename HandleDataType>
+bool TAvaMaskMaterialHandle<HandleDataType>::GetMaskParameters(FAvaMask2DMaterialParameters& OutParameters)
+{
+	if (const UMaterialInstanceDynamic* MaterialInstance = GetMaterialInstance())
+	{
+		OutParameters.SetFromMaterial(
+			MaterialInstance
+			, UE::AvaMask::Internal::TextureParameterInfo
+			, UE::AvaMask::Internal::ChannelParameterInfo
+			, UE::AvaMask::Internal::InvertParameterInfo
+			, UE::AvaMask::Internal::BaseOpacityParameterInfo
+			, UE::AvaMask::Internal::PaddingParameterInfo
+			, UE::AvaMask::Internal::FeatherParameterInfo);
+		return true;
+	}
+
+	return false;
+}
+
+template <typename HandleDataType>
+bool TAvaMaskMaterialHandle<HandleDataType>::SetMaskParameters(UTexture* InTexture
+	, float InBaseOpacity
+	, EGeometryMaskColorChannel InChannel
+	, bool bInInverted
+	, const FVector2f& InPadding
+	, bool bInApplyFeathering
+	, float InOuterFeathering
+	, float InInnerFeathering)
+{
+	if (UMaterialInstanceDynamic* MaterialInstance = GetMaterialInstance())
+	{
+		UE_LOG(LogAvaMask, VeryVerbose, TEXT("SetParameters: Texture:%s, Channel:%s"), InTexture ? *InTexture->GetName() : TEXT("(None)"), *UE::AvaMask::Internal::MaskChannelEnumToVector[InChannel].ToString());
+
+		UTexture* DefaultTexture = nullptr;
+		MaterialInstance->GetTextureParameterDefaultValue(UE::AvaMask::Internal::TextureParameterInfo, DefaultTexture);
+
+		InTexture = InTexture ? InTexture : DefaultTexture;
+
+		MaterialInstance->SetTextureParameterValueByInfo(UE::AvaMask::Internal::TextureParameterInfo, InTexture);
+		MaterialInstance->SetVectorParameterValueByInfo(UE::AvaMask::Internal::ChannelParameterInfo, UE::AvaMask::Internal::MaskChannelEnumToVector[InChannel]);
+		MaterialInstance->SetScalarParameterValueByInfo(UE::AvaMask::Internal::InvertParameterInfo, bInInverted ? 1.0f : 0.0f);
+		MaterialInstance->SetScalarParameterValueByInfo(UE::AvaMask::Internal::BaseOpacityParameterInfo, FMath::Clamp(InBaseOpacity, 0.0f, 1.0f));
+		MaterialInstance->SetVectorParameterValueByInfo(UE::AvaMask::Internal::PaddingParameterInfo, FLinearColor(FMath::Max(0, InPadding.X), FMath::Max(0, InPadding.Y), 0, 0));
+		MaterialInstance->SetVectorParameterValueByInfo(UE::AvaMask::Internal::FeatherParameterInfo, FLinearColor(bInApplyFeathering ? 1.0f : 0.0f, FMath::Max(0, InOuterFeathering), FMath::Max(0, InInnerFeathering), FMath::Max(0, FMath::Max(InOuterFeathering, InInnerFeathering))));
+
+		return true;
+	}
+
+	return false;
+}
+
+template <typename HandleDataType>
+bool TAvaMaskMaterialHandle<HandleDataType>::HasRequiredParameters(TArray<FString>& OutMissingParameterNames)
+{
+	if (const UMaterialInterface* Material = GetMaterial())
+	{
+		return UE::Ava::MaterialHasParameters(
+			*Material
+			, {
+				{ UE::AvaMask::Internal::TextureParameterInfo.Name, EMaterialParameterType::Texture }
+				, { UE::AvaMask::Internal::BaseOpacityParameterInfo.Name, EMaterialParameterType::Scalar }
+				, { UE::AvaMask::Internal::ChannelParameterInfo.Name, EMaterialParameterType::Vector }
+				, { UE::AvaMask::Internal::InvertParameterInfo.Name, EMaterialParameterType::Scalar }
+			}
+			, OutMissingParameterNames);
+	}
+
+	UE_LOG(LogAvaMask, Warning, TEXT("HasRequiredParameters: Material was invalid."));
+
+	return false;
+}
diff --git a/Engine/Plugins/VirtualProduction/Avalanche/Source/AvalancheMask/Private/Mask2D/AvaMask2DReadModifier.cpp b/Engine/Plugins/VirtualProduction/Avalanche/Source/AvalancheMask/Private/Mask2D/AvaMask2DReadModifier.cpp
new file mode 100644
--- /dev/null
+++ b/Engine/Plugins/VirtualProduction/Avalanche/Source/AvalancheMask/Private/Mask2D/AvaMask2DReadModifier.cpp
@@ -0,0 +1,194 @@
+// Copyright Epic Games, Inc. All Rights Reserved.
+
+#include "Mask2D/AvaMask2DReadModifier.h"
+
+#include "AvaMaskUtilities.h"
+#include "Components/PrimitiveComponent.h"
+#include "Engine/CanvasRenderTarget2D.h"
+#include "Engine/Engine.h"
+#include "GeometryMaskCanvas.h"
+#include "GeometryMaskCanvasResource.h"
+#include "GeometryMaskReadComponent.h"
+#include "Handling/AvaHandleUtilities.h"
+#include "Handling/AvaObjectHandleSubsystem.h"
+#include "Handling/IAvaMaskMaterialCollectionHandle.h"
+#include "Handling/IAvaMaskMaterialHandle.h"
+#include "Materials/Material.h"
+#include "Materials/MaterialInstanceDynamic.h"
+#include "Modifiers/ActorModifierCoreStack.h"
+
+#define LOCTEXT_NAMESPACE "AvaMask2DReadModifier"
+
+#if WITH_EDITOR
+const TAvaPropertyChangeDispatcher<UAvaMask2DReadModifier> UAvaMask2DReadModifier::PropertyChangeDispatcher =
+{
+	{ GET_MEMBER_NAME_CHECKED(UAvaMask2DReadModifier, BaseOpacity), &UAvaMask2DReadModifier::OnBaseOpacityChanged }
+};
+#endif
+
+void UAvaMask2DReadModifier::SetBaseOpacity(float InBaseOpacity)
+{
+	InBaseOpacity = FMath::Clamp(InBaseOpacity, 0.0f, 1.0f);
+	if (!FMath::IsNearlyEqual(BaseOpacity, InBaseOpacity))
+	{
+		BaseOpacity = InBaseOpacity;
+		OnBaseOpacityChanged();
+	}
+}
+
+#if WITH_EDITOR
+void UAvaMask2DReadModifier::PostEditChangeProperty(FPropertyChangedEvent& InPropertyChangedEvent)
+{
+	Super::PostEditChangeProperty(InPropertyChangedEvent);
+
+	UAvaMask2DReadModifier::PropertyChangeDispatcher.OnPropertyChanged(this, InPropertyChangedEvent);
+}
+#endif
+
+void UAvaMask2DReadModifier::Apply()
+{
+	Super::Apply();
+
+	if (const UGeometryMaskCanvas* Canvas = GetCurrentCanvas())
+	{
+		for (TPair<TWeakObjectPtr<AActor>, FAvaMask2DActorData>& ActorDataPair : ActorData)
+		{
+			if (AActor* Actor = ActorDataPair.Key.Get())
+			{
+				if (!ApplyRead(Actor, ActorDataPair.Value))
+				{
+					return;
+				}			
+			}
+		}
+	}
+
+	Next();
+}
+
+EBlendMode UAvaMask2DReadModifier::GetBlendMode() const
+{
+	return bUseBlur || bUseFeathering ? EBlendMode::BLEND_Translucent : EBlendMode::BLEND_Masked;
+}
+
+bool UAvaMask2DReadModifier::ApplyRead(AActor* InActor, FAvaMask2DActorData& InActorData)
+{
+	// Only add read/write component to actors with primitives
+	if (!ActorSupportsMaskReadWrite(InActor))
+	{
+		// Ok to continue, just nothing done for this actor
+		return true;
+	}
+
+	if (const UGeometryMaskCanvas* Canvas = GetCurrentCanvas())
+	{
+		// Get MaterialCollectionHandle
+		const TSharedPtr<IAvaMaskMaterialCollectionHandle>& MaterialCollectionHandle =
+			UE::Ava::Internal::FindOrAddHandleByLambda(
+				MaterialCollectionHandles
+				, InActor
+				, [this, InActor]()
+				{
+					return GetObjectHandleSubsystem()->MakeHandle<IAvaMaskMaterialCollectionHandle>(InActor, UE::AvaMask::Internal::HandleTag);
+				});
+
+		const UTexture* CanvasTexture = TryResolveCanvasTexture(InActor, InActorData);
+		if (!CanvasTexture)
+		{
+			return true;
+		}
+
+		{
+			FAvaMask2DSubjectParameters ApplyParameters;
+			ApplyParameters.MaterialParameters.BlendMode = GetBlendMode();		
+			ApplyParameters.MaterialParameters.CanvasName = Canvas->GetCanvasName();
+			ApplyParameters.MaterialParameters.Texture = Canvas->GetTexture();
+			ApplyParameters.MaterialParameters.BaseOpacity = BaseOpacity;
+			ApplyParameters.MaterialParameters.Channel = Canvas->GetColorChannel();
+			ApplyParameters.MaterialParameters.ChannelAsVector = UE::GeometryMask::MaskChannelEnumToVector[ApplyParameters.MaterialParameters.Channel];
+			ApplyParameters.MaterialParameters.bInvert = bInverted;
+
+			FGeometryMaskDrawingContext DrawingContext(Canvas->GetCanvasId().Level);
+
+			const FIntPoint ViewportSize = Canvas->GetResource()->GetMaxViewportSize();
+			const FVector2f ViewportPadding = FVector2f::One() - FVector2f(ViewportSize + Canvas->GetResource()->GetViewportPadding(DrawingContext)) / ViewportSize;
+			ApplyParameters.MaterialParameters.Padding = ViewportPadding;
+
+			ApplyParameters.MaterialParameters.bApplyFeathering = bUseFeathering;
+			ApplyParameters.MaterialParameters.OuterFeatherRadius = OuterFeatherRadius;
+			ApplyParameters.MaterialParameters.InnerFeatherRadius = InnerFeatherRadius;
+
+			FInstancedStruct& MaterialCollectionData = MaterialCollectionHandleData.FindOrAdd(InActor, MaterialCollectionHandle->MakeDataStruct());
+			MaterialCollectionHandle->ApplyModifiedState(ApplyParameters, MaterialCollectionData);
+		}
+
+		if (UGeometryMaskReadComponent* ReadComponent = FindOrAddMaskComponent<UGeometryMaskReadComponent>(InActor))
+		{
+			FGeometryMaskReadParameters Parameters = ReadComponent->GetParameters();
+			Parameters.CanvasName = Channel;
+			Parameters.ColorChannel = Canvas->GetColorChannel();
+			Parameters.bInvert = bInverted;
+
+			ReadComponent->SetParameters(Parameters);
+		}
+
+		// Only check for valid materials here, after ApplyModifiedState has occured once
+		FText FailReason;
+		if (!MaterialCollectionHandle->ValidateMaterials(FailReason))
+		{
+			Fail(FailReason);
+			return false;
+		}
+	}
+
+	return true;
+}
+
+void UAvaMask2DReadModifier::OnModifierCDOSetup(FActorModifierCoreMetadata& InMetadata)
+{
+	Super::OnModifierCDOSetup(InMetadata);
+
+	InMetadata.SetName(TEXT("MaskRead"));
+	InMetadata.SetCategory(TEXT("Rendering"));
+	InMetadata.DisallowAfter(TEXT("MaskWrite"));
+	InMetadata.DisallowBefore(TEXT("MaskWrite"));
+#if WITH_EDITOR
+	InMetadata.SetDisplayName(FText::FromString(TEXT("Masked Layer (Output)")));
+	InMetadata.SetDescription(LOCTEXT("ModifierDescription", "Reads from a canvas texture and uses it on materials"));
+#endif
+}
+
+void UAvaMask2DReadModifier::SetupMaskComponent(UActorComponent* InComponent)
+{
+	if (!InComponent)
+	{
+		return;
+	}
+
+	if (IGeometryMaskReadInterface* MaskReader = Cast<IGeometryMaskReadInterface>(InComponent))
+	{
+		SetupMaskReadComponent(MaskReader);
+	}
+	
+	UAvaMask2DBaseModifier::SetupMaskComponent(InComponent);
+}
+
+void UAvaMask2DReadModifier::SetupMaskReadComponent(IGeometryMaskReadInterface* InMaskReader)
+{
+	if (const UGeometryMaskCanvas* Canvas = GetCurrentCanvas())
+	{
+		// Set canvas name to read from
+		FGeometryMaskReadParameters ReadParameters = InMaskReader->GetParameters();
+		ReadParameters.CanvasName = Channel;
+		ReadParameters.ColorChannel = Canvas->GetColorChannel();
+
+		InMaskReader->SetParameters(ReadParameters);
+	}
+}
+
+void UAvaMask2DReadModifier::OnBaseOpacityChanged()
+{
+	MarkModifierDirty();
+}
+
+#undef LOCTEXT_NAMESPACE
diff --git a/Engine/Plugins/VirtualProduction/Avalanche/Source/AvalancheMask/Private/Mask2D/AvaMask2DWriteModifier.cpp b/Engine/Plugins/VirtualProduction/Avalanche/Source/AvalancheMask/Private/Mask2D/AvaMask2DWriteModifier.cpp
new file mode 100644
--- /dev/null
+++ b/Engine/Plugins/VirtualProduction/Avalanche/Source/AvalancheMask/Private/Mask2D/AvaMask2DWriteModifier.cpp
@@ -0,0 +1,153 @@
+// Copyright Epic Games, Inc. All Rights Reserved.
+
+#include "Mask2D/AvaMask2DWriteModifier.h"
+
+#include "AvaMaskUtilities.h"
+#include "Components/PrimitiveComponent.h"
+#include "Engine/CanvasRenderTarget2D.h"
+#include "Engine/Engine.h"
+#include "Framework/AvaGizmoComponent.h"
+#include "GeometryMaskCanvas.h"
+#include "GeometryMaskWriteComponent.h"
+#include "Materials/Material.h"
+#include "Materials/MaterialInstanceDynamic.h"
+#include "Modifiers/ActorModifierCoreStack.h"
+
+#define LOCTEXT_NAMESPACE "AvaMask2DWriteModifier"
+
+#if WITH_EDITOR
+const TAvaPropertyChangeDispatcher<UAvaMask2DWriteModifier> UAvaMask2DWriteModifier::PropertyChangeDispatcher =
+{
+	{ GET_MEMBER_NAME_CHECKED(UAvaMask2DWriteModifier, WriteOperation), &UAvaMask2DWriteModifier::OnWriteOperationChanged }
+};
+#endif
+
+#if WITH_EDITOR
+void UAvaMask2DWriteModifier::PostEditChangeProperty(FPropertyChangedEvent& InPropertyChangedEvent)
+{
+	Super::PostEditChangeProperty(InPropertyChangedEvent);
+
+	UAvaMask2DWriteModifier::PropertyChangeDispatcher.OnPropertyChanged(this, InPropertyChangedEvent);
+}
+#endif
+
+void UAvaMask2DWriteModifier::Apply()
+{
+	Super::Apply();
+	
+	if (GetCurrentCanvas())
+	{
+		for (TPair<TWeakObjectPtr<AActor>, FAvaMask2DActorData>& ActorDataPair : ActorData)
+		{
+			if (AActor* Actor = ActorDataPair.Key.Get())
+			{
+				if (!ApplyWrite(Actor, ActorDataPair.Value))
+				{
+					return;
+				}
+			}
+		}
+	}
+
+	Next();
+}
+
+bool UAvaMask2DWriteModifier::ApplyWrite(AActor* InActor, FAvaMask2DActorData& InActorData)
+{
+	// Only add read/write component to actors with primitives
+	if (ActorSupportsMaskReadWrite(InActor))
+	{
+		if (const UGeometryMaskCanvas* GeometryCanvas = GetCurrentCanvas())
+		{
+			const UTexture* CanvasTexture = TryResolveCanvasTexture(InActor, InActorData);
+			if (!CanvasTexture)
+			{
+				return true;
+			}
+		
+			if (UGeometryMaskWriteMeshComponent* WriteComponent = FindOrAddMaskComponent<UGeometryMaskWriteMeshComponent>(InActor))
+			{
+				FGeometryMaskWriteParameters Parameters = WriteComponent->GetParameters();
+				Parameters.CanvasName = Channel;
+				Parameters.ColorChannel = GeometryCanvas->GetColorChannel();
+				Parameters.OperationType = WriteOperation;
+
+				WriteComponent->SetParameters(Parameters);
+			}
+		}
+	}
+
+	if (UAvaGizmoComponent* VisualizationComponent = UE::AvaMask::Internal::FindOrAddComponent<UAvaGizmoComponent>(InActor))
+	{
+		VisualizationComponent->SetVisibleInEditor(true);
+		VisualizationComponent->SetHiddenInGame(false);
+		VisualizationComponent->SetShowWireframe(false);
+		VisualizationComponent->SetWireframeColor(FLinearColor::Red);
+		VisualizationComponent->SetsStencil(true);
+		VisualizationComponent->SetStencilId(150);
+		VisualizationComponent->SetRenderInMainPass(false);
+		VisualizationComponent->SetRenderDepth(false);
+		VisualizationComponent->SetMaterialToDefault();
+		VisualizationComponent->ApplyGizmoValues();
+	}
+
+	return true;
+}
+
+void UAvaMask2DWriteModifier::SetWriteOperation(const EGeometryMaskCompositeOperation InWriteOperation)
+{
+	if (WriteOperation != InWriteOperation)
+	{
+		WriteOperation = InWriteOperation;
+		OnWriteOperationChanged();
+	}
+}
+
+void UAvaMask2DWriteModifier::OnModifierCDOSetup(FActorModifierCoreMetadata& InMetadata)
+{
+	Super::OnModifierCDOSetup(InMetadata);
+
+	InMetadata.SetName(TEXT("MaskWrite"));
+	InMetadata.SetCategory(TEXT("Rendering"));
+	InMetadata.DisallowAfter(TEXT("MaskRead"));
+	InMetadata.DisallowBefore(TEXT("MaskRead"));
+#if WITH_EDITOR
+	InMetadata.SetDisplayName(FText::FromString(TEXT("Mask Layer (Input)")));
+	InMetadata.SetDescription(LOCTEXT("ModifierDescription", "Writes to a canvas texture used by the masked layer in materials"));
+#endif
+}
+
+void UAvaMask2DWriteModifier::SetupMaskComponent(UActorComponent* InComponent)
+{
+	if (!InComponent)
+	{
+		return;
+	}
+	
+	if (IGeometryMaskWriteInterface* MaskWriter = Cast<IGeometryMaskWriteInterface>(InComponent))
+	{
+		SetupMaskWriteComponent(MaskWriter);
+	}
+
+	Super::SetupMaskComponent(InComponent);
+}
+
+void UAvaMask2DWriteModifier::SetupMaskWriteComponent(IGeometryMaskWriteInterface* InMaskWriter)
+{
+	if (const UGeometryMaskCanvas* Canvas = GetCurrentCanvas())
+	{
+		// Set canvas name to write to
+		FGeometryMaskWriteParameters WriteParameters = InMaskWriter->GetParameters();
+		WriteParameters.CanvasName = Channel;
+		WriteParameters.ColorChannel = Canvas->GetColorChannel();
+
+		InMaskWriter->SetParameters(WriteParameters);
+	}
+}
+
+void UAvaMask2DWriteModifier::OnWriteOperationChanged()
+{
+	MarkModifierDirty();
+}
+
+#undef LOCTEXT_NAMESPACE
diff --git a/Engine/Plugins/VirtualProduction/Avalanche/Source/AvalancheMask/Public/Mask2D/AvaMask2DBaseModifier.h b/Engine/Plugins/VirtualProduction/Avalanche/Source/AvalancheMask/Public/Mask2D/AvaMask2DBaseModifier.h
new file mode 100644
--- /dev/null
+++ b/Engine/Plugins/VirtualProduction/Avalanche/Source/AvalancheMask/Public/Mask2D/AvaMask2DBaseModifier.h
@@ -0,0 +1,247 @@
+// Copyright Epic Games, Inc. All Rights Reserved.
+
+#pragma once
+
+#include "AvaMaskTypes.h"
+#include "AvaPropertyChangeDispatcher.h"
+#include "GeometryMaskTypes.h"
+#include "Modifiers/ActorModifierArrangeBaseModifier.h"
+#include "StructUtils/StructView.h"
+#include "UObject/ObjectPtr.h"
+
+#include "AvaMask2DBaseModifier.generated.h"
+
+class AActor;
+class IAvaActorHandle;
+class IAvaComponentHandle;
+class IAvaMaskMaterialCollectionHandle;
+class IAvaMaskMaterialHandle;
+class IGeometryMaskReadInterface;
+class IGeometryMaskWriteInterface;
+class UActorComponent;
+class UAvaMaskMaterialInstanceSubsystem;
+class UAvaObjectHandleSubsystem;
+class UCanvasRenderTarget2D;
+class UGeometryMaskCanvas;
+class UMaterialInstanceDynamic;
+class UMaterialInterface;
+class UPrimitiveComponent;
+class UTexture;
+struct FAvaShapeParametricMaterial;
+struct FMaterialParameterInfo;
+
+UENUM(BlueprintType)
+enum class EAvaMask2DMode : uint8
+{
+	Read		UMETA(DisplayName = "Target", ToolTip = "Use the specified Mask Channel to apply to this geometry"),
+	Write		UMETA(DisplayName = "Source", ToolTip = "Use the specified Mask Channel to render this geometry to"),
+};
+
+/** Used for target actors to store essential information */
+USTRUCT()
+struct FAvaMask2DActorData
+{
+	GENERATED_BODY()
+
+	/** The canvas texture to apply to this actor materials */
+	UPROPERTY()
+	TWeakObjectPtr<UTexture> CanvasTextureWeak;
+
+	/** Original assigned materials, used to create instances from, and restore to when modifier removed. */
+	UPROPERTY()
+	TMap<FAvaMask2DComponentMaterialPath, TObjectPtr<UMaterialInterface>> OriginalMaterials;
+
+	/** Instantiated materials corresponding to the above. */
+	UPROPERTY()
+	TMap<FAvaMask2DComponentMaterialPath, TObjectPtr<UMaterialInstanceDynamic>> GeneratedMaterialInstances;
+};
+
+/** Uses scene actors to create a mask texture and applies it to attached actors */
+UCLASS(Abstract)
+class AVALANCHEMASK_API UAvaMask2DBaseModifier
+	: public UActorModifierArrangeBaseModifier
+{
+	GENERATED_BODY()
+
+	friend class UAvaMask2DModifierShared;
+	
+public:
+	//~ Begin UObject
+#if WITH_EDITOR
+	virtual void PostEditChangeProperty(FPropertyChangedEvent& InPropertyChangedEvent) override;
+#endif
+	//~ End UObject
+
+	bool UseParentChannel() const { return bUseParentChannel; }
+	void SetUseParentChannel(const bool bInUseParentChannel);
+
+	const FName GetChannel() const;
+	void SetChannel(FName InChannel);
+
+	bool IsInverted() const { return bInverted; }
+	void SetIsInverted(const bool bInInvert);
+
+	bool UseBlur() const { return bUseBlur; }
+	void UseBlur(bool bInUseBlur);
+
+	float GetBlurStrength() const { return BlurStrength; }
+	void SetBlurStrength(float InBlurStrength);
+
+	bool UseFeathering() const { return bUseFeathering; }
+	void UseFeathering(bool bInUseFeathering);
+
+	int32 GetOuterFeatherRadius() const { return OuterFeatherRadius; }
+	void SetOuterFeatherRadius(int32 InFeatherRadius);
+
+	int32 GetInnerFeatherRadius() const { return InnerFeatherRadius; }
+	void SetInnerFeatherRadius(int32 InFeatherRadius);
+
+	/** Utility to generate a unique mask name based on the associated Actor */
+	FName GenerateUniqueMaskName() const;
+
+	/** Utility to find an existing modifier on the provided actor or it's parent */
+	static UAvaMask2DBaseModifier* FindMaskModifierOnActor(const AActor* InActor);
+
+protected:
+	//~ Begin UActorModifierCoreBase
+	virtual void OnModifierCDOSetup(FActorModifierCoreMetadata& InMetadata) override;
+	virtual void OnModifierAdded(EActorModifierCoreEnableReason InReason) override;
+	virtual void OnModifierRemoved(EActorModifierCoreDisableReason InReason) override;
+	virtual void OnModifiedActorTransformed() override;
+	virtual void SavePreState() override;
+	virtual void RestorePreState() override;
+	virtual void Apply() override;
+	//~ End UActorModifierCoreBase
+
+	//~ Begin IAvaSceneTreeUpdateModifierExtension
+	virtual void OnSceneTreeTrackedActorParentChanged(int32 InIdx, const TArray<TWeakObjectPtr<AActor>>& InPreviousParentActor, const TArray<TWeakObjectPtr<AActor>>& InNewParentActor) override;
+	virtual void OnSceneTreeTrackedActorChildrenChanged(int32 InIdx, const TSet<TWeakObjectPtr<AActor>>& InPreviousChildrenActors, const TSet<TWeakObjectPtr<AActor>>& InNewChildrenActors) override;
+	//~ End IAvaSceneTreeUpdateModifierExtension
+
+	//~ Begin IAvaRenderStateUpdateExtension
+	virtual void OnRenderStateUpdated(AActor* InActor, UActorComponent* InComponent) override;
+	//~ End IAvaRenderStateUpdateExtension
+
+	virtual void SaveActorPreState(AActor* InActor, FAvaMask2DActorData& InActorData);
+	virtual void RestoreActorPreState(AActor* InActor, const FAvaMask2DActorData& InActorData);
+
+	void OnUseParentChannelChanged();
+	void OnChannelChanged();
+	void OnInvertedChanged();
+	void OnBlurChanged();
+	void OnFeatherChanged();
+	void OnCanvasChanged();
+
+	/** Stores specific properties on the canvas locally. */	
+	void CanvasParamsToLocal();
+
+	/** Apply locally stored parameters to the canvas. */
+	void LocalParamsToCanvas();
+
+	void SetupChannelName();
+	virtual void SetupMaskComponent(UActorComponent* InComponent);
+	virtual void RemoveFromActor(AActor* InActor);
+
+	void OnMaskSetCanvas(const UGeometryMaskCanvas* InCanvas, AActor* InActor);
+
+	UActorComponent* FindOrAddMaskComponent(TSubclassOf<UActorComponent> InComponentClass, AActor* InActor);
+
+	template<
+		typename InComponentClass
+		UE_REQUIRES(std::is_base_of_v<UActorComponent, InComponentClass>)>
+	InComponentClass* FindOrAddMaskComponent(AActor* InActor)
+	{
+		return Cast<InComponentClass>(FindOrAddMaskComponent(InComponentClass::StaticClass(), InActor));
+	}
+
+	static bool ActorSupportsMaskReadWrite(const AActor* InActor);
+
+	/** Returns true if parent channel was found */
+	bool TryResolveParentChannel();
+
+	void TryResolveCanvas();
+	UTexture* TryResolveCanvasTexture(AActor* InActor, FAvaMask2DActorData& InActorData);
+
+	UAvaObjectHandleSubsystem* GetObjectHandleSubsystem();
+	UAvaMaskMaterialInstanceSubsystem* GetMaterialInstanceSubsystem();
+
+	/** Returns or resolves the currently referenced canvas. */
+	UGeometryMaskCanvas* GetCurrentCanvas();
+	
+	virtual void OnMaterialsChanged(UObject* InMaterialOwner, const TArray<TSharedPtr<IAvaMaskMaterialHandle>>& InMaterialHandles);
+
+#if WITH_EDITOR
+	UFUNCTION(CallInEditor, Category = "Mask2D")
+	void VisualizeMask();
+#endif
+
+	/** Whether to get the channel from the parent (first one that specifies a mask channel) */
+	UPROPERTY(EditInstanceOnly, BlueprintReadWrite, Getter = "UseParentChannel", Setter = "SetUseParentChannel", Category = "Mask2D", meta = (AllowPrivateAccess = "true", ValidEnumValues = "Self,Parent"))
+	bool bUseParentChannel = false;
+
+	/** Channel found when GetChannelFromParent is true */
+	UPROPERTY(VisibleAnywhere, Category = "Mask2D", meta = (EditCondition = "bUseParentChannel", EditConditionHides))
+	FName ParentChannel = NAME_None;
+
+	/** Channel to read to or write from */
+	UPROPERTY(EditInstanceOnly, BlueprintReadWrite, Getter, Setter, Category = "Mask2D", meta = (EditCondition = "!bUseParentChannel", EditConditionHides, AllowPrivateAccess = "true"))
+	FName Channel = NAME_None;
+
+	/** Whether to apply the mask as inverted (visible becomes invisible and vice versa) */
+	UPROPERTY(EditInstanceOnly, BlueprintReadWrite, Getter = "IsInverted", Setter = "SetIsInverted", Category = "Mask2D", meta = (AllowPrivateAccess = "true"))
+	bool bInverted = false;
+
+	UPROPERTY(BlueprintReadWrite, Getter = "UseBlur", Setter = "UseBlur", Category = "Mask2D|Shared", meta = (AllowPrivateAccess = "true"))
+	bool bUseBlur = false;
+
+	UPROPERTY(BlueprintReadWrite, Getter, Setter, Category = "Mask2D|Shared", meta = (ClampMin = 0.0, AllowPrivateAccess = "true", EditCondition = "bUseBlur", EditConditionHides))
+	float BlurStrength = 16.0f;
+
+	UPROPERTY(BlueprintReadWrite, Getter = "UseFeathering", Setter = "UseFeathering", Category = "Mask2D|Shared", meta = (AllowPrivateAccess = "true"))
+	bool bUseFeathering = false;
+
+	UPROPERTY(BlueprintReadWrite, Getter, Setter, Category = "Mask2D|Shared", meta = (ClampMin = 0, AllowPrivateAccess = "true", EditCondition = "bUseFeathering", EditConditionHides))
+	int32 OuterFeatherRadius = 16;
+
+	UPROPERTY(BlueprintReadWrite, Getter, Setter, Category = "Mask2D|Shared", meta = (ClampMin = 0, AllowPrivateAccess = "true", EditCondition = "bUseFeathering", EditConditionHides))
+	int32 InnerFeatherRadius = 16;
+
+	UPROPERTY(Transient)
+	TWeakObjectPtr<UAvaObjectHandleSubsystem> ObjectHandleSubsystem;
+	
+	UPROPERTY(Transient)
+	TWeakObjectPtr<UAvaMaskMaterialInstanceSubsystem> MaterialInstanceSubsystem;
+
+	/**
+	 * Cached auto channel name, used when the user adds the modifier and immediately sets the mode to Source,
+	 * in which case a different channel name is chosen.
+	 */
+	UPROPERTY(Transient)
+	FName AutoChannelName;
+
+	/** Cached actor data to apply/restore */
+	UPROPERTY(DuplicateTransient)
+	TMap<TWeakObjectPtr<AActor>, FAvaMask2DActorData> ActorData;
+
+	/** Reference to the last resolved canvas name, as stored in CanvasWeak */
+	UPROPERTY(Transient, DuplicateTransient)
+	FName LastResolvedCanvasName;
+
+	/** Reference to the Canvas used */
+	UPROPERTY(EditAnywhere, Transient, DuplicateTransient, Category = "Mask2D|Shared", NoClear,	meta = (DisplayName = "Canvas", AllowPrivateAccess = "true", ShowInnerProperties, NoResetToDefault))
+	TWeakObjectPtr<UGeometryMaskCanvas> CanvasWeak;
+
+	UPROPERTY(DuplicateTransient)
+	TMap<TWeakObjectPtr<UObject>, FInstancedStruct> MaterialCollectionHandleData;
+
+	UPROPERTY(DuplicateTransient)
+	TMap<TWeakObjectPtr<UMaterialInterface>, FInstancedStruct> MaterialHandleData;
+
+	TMap<TObjectKey<UObject>, TSharedPtr<IAvaMaskMaterialCollectionHandle>> MaterialCollectionHandles;
+
+private:
+#if WITH_EDITOR
+	/** Used for PECP */
+	static const TAvaPropertyChangeDispatcher<UAvaMask2DBaseModifier> PropertyChangeDispatcher;
+#endif
+};
diff --git a/Engine/Plugins/VirtualProduction/Avalanche/Source/AvalancheMask/Public/Mask2D/AvaMask2DReadModifier.h b/Engine/Plugins/VirtualProduction/Avalanche/Source/AvalancheMask/Public/Mask2D/AvaMask2DReadModifier.h
new file mode 100644
--- /dev/null
+++ b/Engine/Plugins/VirtualProduction/Avalanche/Source/AvalancheMask/Public/Mask2D/AvaMask2DReadModifier.h
@@ -0,0 +1,77 @@
+// Copyright Epic Games, Inc. All Rights Reserved.
+
+#pragma once
+
+#include "AvaMask2DBaseModifier.h"
+#include "AvaPropertyChangeDispatcher.h"
+#include "Modifiers/ActorModifierArrangeBaseModifier.h"
+#include "UObject/ObjectPtr.h"
+#include "UObject/Package.h"
+#include "UObject/WeakObjectPtr.h"
+
+#include "AvaMask2DReadModifier.generated.h"
+
+class AActor;
+class IAvaActorHandle;
+class IAvaComponentHandle;
+class IAvaMaskMaterialCollectionHandle;
+class IAvaMaskMaterialHandle;
+class IGeometryMaskReadInterface;
+class IGeometryMaskWriteInterface;
+class UActorComponent;
+class UAvaMaskMaterialInstanceSubsystem;
+class UAvaObjectHandleSubsystem;
+class UGeometryMaskCanvas;
+class UCanvasRenderTarget2D;
+class UMaterialInstanceDynamic;
+class UMaterialInterface;
+class UPrimitiveComponent;
+class UTexture;
+struct FAvaShapeParametricMaterial;
+struct FMaterialParameterInfo;
+
+/** Uses scene actors to create a mask texture and applies it to attached actors */
+UCLASS(BlueprintType)
+class AVALANCHEMASK_API UAvaMask2DReadModifier
+	: public UAvaMask2DBaseModifier
+{
+	GENERATED_BODY()
+
+	friend class UAvaMask2DModifierShared;
+	
+public:
+	//~ Begin UObject
+#if WITH_EDITOR
+	virtual void PostEditChangeProperty(FPropertyChangedEvent& InPropertyChangedEvent) override;
+#endif
+	//~ End UObject
+
+	const float GetBaseOpacity() const { return BaseOpacity; }
+	void SetBaseOpacity(float InBaseOpacity);
+
+protected:
+	//~ Begin UActorModifierCoreBase
+	virtual void OnModifierCDOSetup(FActorModifierCoreMetadata& InMetadata) override;
+	virtual void Apply() override;
+	//~ End UActorModifierCoreBase
+
+	/** Get target blend mode dependent on modifier settings. */
+	EBlendMode GetBlendMode() const;
+
+	void OnBaseOpacityChanged();
+
+	virtual void SetupMaskComponent(UActorComponent* InComponent) override;
+	void SetupMaskReadComponent(IGeometryMaskReadInterface* InMaskReader);
+
+	bool ApplyRead(AActor* InActor, FAvaMask2DActorData& InActorData);
+
+	/** Base opacity/alpha to use in Read mode */
+	UPROPERTY(EditInstanceOnly, BlueprintReadWrite, Getter, Setter, Category = "Mask2D", meta = (ClampMin = 0.0, ClampMax = 1.0, AllowPrivateAccess = "true"))
+	float BaseOpacity = 0.0f;
+
+private:
+#if WITH_EDITOR
+	/** Used for PECP */
+	static const TAvaPropertyChangeDispatcher<UAvaMask2DReadModifier> PropertyChangeDispatcher;
+#endif
+};
diff --git a/Engine/Plugins/VirtualProduction/Avalanche/Source/AvalancheMask/Public/Mask2D/AvaMask2DWriteModifier.h b/Engine/Plugins/VirtualProduction/Avalanche/Source/AvalancheMask/Public/Mask2D/AvaMask2DWriteModifier.h
new file mode 100644
--- /dev/null
+++ b/Engine/Plugins/VirtualProduction/Avalanche/Source/AvalancheMask/Public/Mask2D/AvaMask2DWriteModifier.h
@@ -0,0 +1,74 @@
+// Copyright Epic Games, Inc. All Rights Reserved.
+
+#pragma once
+
+#include "AvaMask2DBaseModifier.h"
+#include "AvaPropertyChangeDispatcher.h"
+#include "GeometryMaskTypes.h"
+#include "UObject/ObjectPtr.h"
+#include "UObject/Package.h"
+#include "UObject/WeakObjectPtr.h"
+
+#include "AvaMask2DWriteModifier.generated.h"
+
+class AActor;
+class IAvaActorHandle;
+class IAvaComponentHandle;
+class IAvaMaskMaterialCollectionHandle;
+class IAvaMaskMaterialHandle;
+class IGeometryMaskReadInterface;
+class IGeometryMaskWriteInterface;
+class UActorComponent;
+class UAvaMaskMaterialInstanceSubsystem;
+class UAvaObjectHandleSubsystem;
+class UCanvasRenderTarget2D;
+class UGeometryMaskCanvas;
+class UMaterialInstanceDynamic;
+class UMaterialInterface;
+class UPrimitiveComponent;
+class UTexture;
+struct FAvaShapeParametricMaterial;
+struct FMaterialParameterInfo;
+
+/** Uses scene actors to create a mask texture and applies it to attached actors */
+UCLASS(BlueprintType)
+class AVALANCHEMASK_API UAvaMask2DWriteModifier
+	: public UAvaMask2DBaseModifier
+{
+	GENERATED_BODY()
+
+	friend class UAvaMask2DModifierShared;
+	
+public:
+	//~ Begin UObject
+#if WITH_EDITOR
+	virtual void PostEditChangeProperty(FPropertyChangedEvent& InPropertyChangedEvent) override;
+#endif
+	//~ End UObject
+
+	EGeometryMaskCompositeOperation GetWriteOperation() const { return WriteOperation; }
+	void SetWriteOperation(const EGeometryMaskCompositeOperation InWriteOperation);
+
+protected:
+	//~ Begin UActorModifierCoreBase
+	virtual void OnModifierCDOSetup(FActorModifierCoreMetadata& InMetadata) override;
+	virtual void Apply() override;
+	//~ End UActorModifierCoreBase
+
+	void OnWriteOperationChanged();
+
+	void SetupMaskComponent(UActorComponent* InComponent);
+	void SetupMaskWriteComponent(IGeometryMaskWriteInterface* InMaskWriter);
+
+	bool ApplyWrite(AActor* InActor, FAvaMask2DActorData& InActorData);
+
+protected:
+	/** How to write to the chosen mask channel */
+	UPROPERTY(EditInstanceOnly, BlueprintReadWrite, Getter, Setter, Category = "Mask2D", meta = (ValidEnumValues = "Add, Subtract", AllowPrivateAccess = "true"))
+	EGeometryMaskCompositeOperation WriteOperation = EGeometryMaskCompositeOperation::Add;
+
+#if WITH_EDITOR
+	/** Used for PECP */
+	static const TAvaPropertyChangeDispatcher<UAvaMask2DWriteModifier> PropertyChangeDispatcher;
+#endif
+};
diff --git a/Engine/Plugins/VirtualProduction/Avalanche/Source/AvalancheMedia/Private/Playback/AvaPlaybackServerProcess.cpp b/Engine/Plugins/VirtualProduction/Avalanche/Source/AvalancheMedia/Private/Playback/AvaPlaybackServerProcess.cpp
new file mode 100644
--- /dev/null
+++ b/Engine/Plugins/VirtualProduction/Avalanche/Source/AvalancheMedia/Private/Playback/AvaPlaybackServerProcess.cpp
@@ -0,0 +1,183 @@
+// Copyright Epic Games, Inc. All Rights Reserved.
+
+#include "AvaPlaybackServerProcess.h"
+
+#include "AvaMediaSettings.h"
+#include "IAvaMediaModule.h"
+#include "Misc/App.h"
+#include "Misc/Paths.h"
+#include "Playback/AvaPlaybackClient.h"
+
+namespace UE::AvaPlaybackServerProcess::Private
+{
+	FString GetLocalServerName()
+	{
+		return UAvaMediaSettings::Get().LocalPlaybackServerSettings.ServerName;
+	}
+
+	const TCHAR* GetPlaybackServerLogReplicationVerbosityString()
+	{
+		const UAvaMediaSettings& Settings = UAvaMediaSettings::Get(); 
+		return ToString(UAvaMediaSettings::ToLogVerbosity(Settings.PlaybackServerLogReplicationVerbosity));
+	}
+}
+
+FAvaPlaybackServerProcess::FAvaPlaybackServerProcess(FProcHandle&& InProcessHandle)
+	: ProcessHandle(MoveTemp(InProcessHandle))
+{}
+
+bool FAvaPlaybackServerProcess::IsLaunched()
+{
+	return ProcessHandle.IsValid() && FPlatformProcess::IsProcRunning(ProcessHandle);
+}
+
+bool FAvaPlaybackServerProcess::Launch()
+{
+	if (IsLaunched())
+	{
+		UE_LOG(LogAvaMedia, Warning, TEXT("Local playback server is already running."));
+		return true;
+	}
+
+	using namespace UE::AvaPlaybackServerProcess::Private;
+
+	FString GameNameOrProjectFile;
+	if (FPaths::IsProjectFilePathSet())
+	{
+		GameNameOrProjectFile = FString::Printf(TEXT("\"%s\""), *FPaths::GetProjectFilePath());
+	}
+	else
+	{
+		GameNameOrProjectFile = FApp::GetProjectName();
+	}
+
+	FString CommandLine;
+	CommandLine += GameNameOrProjectFile;
+
+	// Todo: It will likely need a more complete set of parameters. Reference: UDisplayClusterLaunchEditorProjectSettings.
+	const FAvaMediaLocalPlaybackServerSettings& Settings = UAvaMediaSettings::Get().LocalPlaybackServerSettings;
+
+	// Standard arguments
+	const FIntPoint Resolution = Settings.Resolution;
+	CommandLine += FString::Printf(TEXT(" -game -windowed -ResX=%d -ResY=%d -messaging"), Resolution.X, Resolution.Y);
+	if (Settings.bEnableLogConsole)
+	{
+		CommandLine += TEXT(" -log");
+	}
+	// Motion Design arguments
+	CommandLine += FString::Printf(TEXT(" -MotionDesignPlaybackServerStart=%s"), *GetLocalServerName());
+	CommandLine += TEXT(" -MotionDesignPlaybackClientSuppress");
+	CommandLine += FString::Printf(TEXT(" -MotionDesignPlaybackServerLogReplication=%s"), GetPlaybackServerLogReplicationVerbosityString());
+	// Storm Sync arguments
+	CommandLine += TEXT(" -NoStormSyncServerAutoStart");
+
+	if (!Settings.ExtraCommandLineArguments.IsEmpty())
+	{
+		CommandLine += TEXT(" ");
+		CommandLine += Settings.ExtraCommandLineArguments;
+	}
+
+	FString LogCommands;
+	
+	for (const FAvaPlaybackServerLoggingEntry& LoggingEntry : Settings.Logging)
+	{
+		if (!LoggingEntry.Category.IsNone())
+		{
+			LogCommands += FString::Printf(TEXT("%s%s %s"),
+					LogCommands.IsEmpty() ? TEXT("") : TEXT(", "), *LoggingEntry.Category.ToString(),
+					ToString(UAvaMediaSettings::ToLogVerbosity(LoggingEntry.VerbosityLevel)));
+		}
+	}
+
+	if (!LogCommands.IsEmpty())
+	{
+		CommandLine += FString::Printf(TEXT(" -LogCmds=\"%s\""), *LogCommands);		
+	}
+
+	const FString ExecutablePath = FPlatformProcess::ExecutablePath();
+
+	uint32 ProcessID = 0;
+	constexpr bool bLaunchDetached = true;
+	constexpr bool bLaunchMinimized = false;
+	constexpr bool bLaunchWindowHidden = false;
+	constexpr uint32 PriorityModifier = 0;
+
+	UE_LOG(LogAvaMedia, Log, TEXT("Launching a playback server in game mode in a new process with the following command line:"));
+	UE_LOG(LogAvaMedia, Log, TEXT("%s %s"), *ExecutablePath, *CommandLine);
+
+	ProcessHandle = FPlatformProcess::CreateProc(
+		*ExecutablePath, *CommandLine, bLaunchDetached,
+		bLaunchMinimized, bLaunchWindowHidden, &ProcessID,
+		PriorityModifier, nullptr, nullptr, nullptr);
+	
+	return ProcessHandle.IsValid();
+}
+
+void FAvaPlaybackServerProcess::Stop()
+{
+	FPlatformProcess::TerminateProc(ProcessHandle);
+}
+
+TSharedPtr<FAvaPlaybackServerProcess> FAvaPlaybackServerProcess::Find(const FAvaPlaybackClient& InPlaybackClient)
+{
+	using namespace UE::AvaPlaybackServerProcess::Private;
+
+	// We use the playback client to find connected server processes that are considered local.
+	// Local means that it is in the same project content path and the same machine. We figure out
+	// if it is the same machine if the process Id can be opened.
+
+	const FString ProjectContentPath = InPlaybackClient.GetProjectContentPath();
+	const uint32 ClientProcessId = FPlatformProcess::GetCurrentProcessId();
+	const FString LocalServerName = GetLocalServerName();
+	
+	// First try with the configured local server name.
+	if (InPlaybackClient.GetServerProjectContentPath(LocalServerName) == ProjectContentPath)
+	{
+		const uint32 ServerProcessId = InPlaybackClient.GetServerProcessId(LocalServerName);
+		if (ServerProcessId != 0 && ClientProcessId != ServerProcessId)
+		{
+			FProcHandle ProcessHandle = FPlatformProcess::OpenProcess(ServerProcessId);
+			if (ProcessHandle.IsValid())
+			{
+				return MakeShared<FAvaPlaybackServerProcess>(MoveTemp(ProcessHandle));
+			}
+
+			// Warn about the server name collision issue.
+			UE_LOG(LogAvaMedia, Warning, TEXT("Found Connected Server \"%s\" but it is not a local process."), *LocalServerName);
+			UE_LOG(LogAvaMedia, Warning, TEXT("This indicates a name collision between the servers. The local server should be renamed."));
+		}
+	}
+
+	// Fallback
+	// Try to find any connected server that is "local", i.e. on the same machine and
+	// in the same project content folder, but on a different process id.
+	
+	TArray<FString> ServerNames = InPlaybackClient.GetServerNames();
+	for (const FString& ServerName : ServerNames)
+	{
+		if (InPlaybackClient.GetServerProjectContentPath(ServerName) == ProjectContentPath)
+		{
+			const uint32 ServerProcessId = InPlaybackClient.GetServerProcessId(ServerName);
+			if (ClientProcessId != ServerProcessId)
+			{
+				FProcHandle ProcessHandle = FPlatformProcess::OpenProcess(ServerProcessId);
+				if (ProcessHandle.IsValid())
+				{
+					return MakeShared<FAvaPlaybackServerProcess>(MoveTemp(ProcessHandle));
+				}
+			}
+		}
+	}
+	
+	return nullptr;	
+}
+
+TSharedPtr<FAvaPlaybackServerProcess> FAvaPlaybackServerProcess::FindOrCreate(const FAvaPlaybackClient& InPlaybackClient)
+{
+	TSharedPtr<FAvaPlaybackServerProcess> ServerProcess = Find(InPlaybackClient);
+	if (!ServerProcess.IsValid())
+	{
+		ServerProcess = MakeShared<FAvaPlaybackServerProcess>();
+	}
+	return ServerProcess;
+}
\ No newline at end of file
