Motion Design: moved the following plugins from /Plugins/Experimental to /Plugins/VirtualProduction:

diff --git a/Engine/Plugins/VirtualProduction/Avalanche/Source/AvalancheMedia/Private/Rundown/AvaRundownPageAssetUtils.cpp b/Engine/Plugins/VirtualProduction/Avalanche/Source/AvalancheMedia/Private/Rundown/AvaRundownPageAssetUtils.cpp
new file mode 100644
--- /dev/null
+++ b/Engine/Plugins/VirtualProduction/Avalanche/Source/AvalancheMedia/Private/Rundown/AvaRundownPageAssetUtils.cpp
@@ -0,0 +1,95 @@
+// Copyright Epic Games, Inc. All Rights Reserved.
+
+#include "Rundown/AvaRundownPageAssetUtils.h"
+
+#include "AvaScene.h"
+#include "AvaTagHandle.h"
+#include "AvaTransitionSubsystem.h"
+#include "AvaTransitionTree.h"
+#include "Behavior/IAvaTransitionBehavior.h"
+#include "Engine/Level.h"
+#include "Engine/World.h"
+#include "IAvaSceneInterface.h"
+
+namespace UE::AvaRundownPageAssetUtils::Private
+{
+	template<typename InInterfaceType>
+	const InInterfaceType* GetInterfaceFromAsset(const UObject* InLoadedAsset)
+	{
+		if (const InInterfaceType* const Interface = Cast<InInterfaceType>(InLoadedAsset))
+		{
+			return Interface;
+		}
+		if (const UWorld* const World = Cast<UWorld>(InLoadedAsset))
+		{
+			if (const AAvaScene* AvaScene = AAvaScene::GetScene(World->PersistentLevel, false))
+			{
+				return static_cast<const InInterfaceType*>(AvaScene);
+			}
+		}
+		return nullptr;
+	}
+}
+
+const IAvaSceneInterface* FAvaRundownPageAssetUtils::GetSceneInterface(const UObject* InLoadedAsset)
+{
+	using namespace UE::AvaRundownPageAssetUtils::Private;
+	return GetInterfaceFromAsset<IAvaSceneInterface>(InLoadedAsset);
+}
+
+const UAvaTransitionTree* FAvaRundownPageAssetUtils::FindTransitionTree(const IAvaSceneInterface* InSceneInterface)
+{
+	if (!InSceneInterface)
+	{
+		return nullptr;
+	}
+	
+	ULevel* SceneLevel = InSceneInterface->GetSceneLevel();
+	if (!SceneLevel)
+	{
+		return nullptr;
+	}
+
+	// Method 1 - Using the UAvaTransitionSubsystem
+	// This method may not work if the world is not initialized (inactive, as in the managed instance).
+	if (const UWorld* SceneWorld = SceneLevel->GetWorld())
+	{
+		if (const UAvaTransitionSubsystem* TransitionSubsystem = SceneWorld->GetSubsystem<UAvaTransitionSubsystem>())
+		{
+			if (const IAvaTransitionBehavior* TransitionBehavior = TransitionSubsystem->GetTransitionBehavior(SceneLevel))
+			{
+				if (const UAvaTransitionTree* TransitionTree = TransitionBehavior->GetTransitionTree())
+				{
+					return TransitionTree;
+				}
+			}
+			// If we had a subsystem and it didn't find the behavior then there was no state tree either and we are done.
+			return nullptr;
+		}
+	}
+	
+	// Method 2 - Fallback using direct lookup of the private actor class.
+	// This is needed to lookup in managed instances and source asset.
+
+	static FString AvaTransitionBehaviorActor(TEXT("AvaTransitionBehaviorActor"));
+	for (const AActor* Actor : SceneLevel->Actors)
+	{
+		const UClass* ActorClass = IsValid(Actor) ? Actor->GetClass() : nullptr;
+		if (ActorClass && ActorClass->GetName() == AvaTransitionBehaviorActor)
+		{
+			const FObjectPropertyBase* const Property = FindFProperty<FObjectPropertyBase>(ActorClass, TEXT("TransitionTree"));
+			if (Property
+				&& Property->PropertyClass
+				&& Property->PropertyClass->IsChildOf(UAvaTransitionTree::StaticClass()))
+			{
+				return Cast<UAvaTransitionTree>(Property->GetObjectPropertyValue_InContainer(Actor));
+			}
+		}
+	}
+	return nullptr;
+}
+
+FAvaTagHandle FAvaRundownPageAssetUtils::GetTransitionLayerTag(const UAvaTransitionTree* InTransitionTree)
+{
+	return InTransitionTree ? InTransitionTree->GetTransitionLayer() : FAvaTagHandle();
+}
\ No newline at end of file
diff --git a/Engine/Plugins/VirtualProduction/Avalanche/Source/AvalancheModifiers/Private/Shared/AvaTranslucentPriorityModifierShared.cpp b/Engine/Plugins/VirtualProduction/Avalanche/Source/AvalancheModifiers/Private/Shared/AvaTranslucentPriorityModifierShared.cpp
new file mode 100644
--- /dev/null
+++ b/Engine/Plugins/VirtualProduction/Avalanche/Source/AvalancheModifiers/Private/Shared/AvaTranslucentPriorityModifierShared.cpp
@@ -0,0 +1,355 @@
+// Copyright Epic Games, Inc. All Rights Reserved.
+
+#include "Shared/AvaTranslucentPriorityModifierShared.h"
+#include "AvaActorUtils.h"
+#include "AvaSceneItem.h"
+#include "AvaSceneTree.h"
+#include "AvaSceneTreeNode.h"
+#include "Camera/CameraActor.h"
+#include "Camera/CameraComponent.h"
+#include "Framework/AvaGameInstance.h"
+#include "IAvaSceneInterface.h"
+#include "Modifiers/AvaTranslucentPriorityModifier.h"
+
+#if WITH_EDITOR
+#include "AvaOutlinerUtils.h"
+#include "IAvaOutliner.h"
+#include "Item/AvaOutlinerItemUtils.h"
+#endif
+
+void UAvaTranslucentPriorityModifierShared::SetComponentsState(UAvaTranslucentPriorityModifier* InModifierContext, const TSet<TWeakObjectPtr<UPrimitiveComponent>>& InComponents)
+{
+	if (!IsValid(InModifierContext))
+	{
+		return;
+	}
+
+	// Remove old component states
+	for (TSet<FAvaTranslucentPriorityModifierComponentState>::TIterator It(ComponentStates); It; ++It)
+	{
+		UPrimitiveComponent* PrimitiveComponent = It->PrimitiveComponentWeak.Get();
+
+		if (PrimitiveComponent && InComponents.Contains(PrimitiveComponent))
+		{
+			It->ModifierWeak = InModifierContext;
+			continue;
+		}
+
+		const UAvaTranslucentPriorityModifier* Modifier = It->ModifierWeak.Get();
+		if (!PrimitiveComponent || !Modifier || Modifier == InModifierContext)
+		{
+			It.RemoveCurrent();
+		}
+	}
+
+	// Add missing components
+	for (const TWeakObjectPtr<UPrimitiveComponent>& PrimitiveComponentWeak : InComponents)
+	{
+		if (UPrimitiveComponent* PrimitiveComponent = PrimitiveComponentWeak.Get())
+		{
+			SaveComponentState(InModifierContext, PrimitiveComponent, true);
+		}
+	}
+}
+
+void UAvaTranslucentPriorityModifierShared::SaveComponentState(UAvaTranslucentPriorityModifier* InModifierContext, UPrimitiveComponent* InComponent, bool bInOverrideContext)
+{
+	if (!IsValid(InComponent))
+	{
+		return;
+	}
+
+	FAvaTranslucentPriorityModifierComponentState& ComponentState = ComponentStates.FindOrAdd(FAvaTranslucentPriorityModifierComponentState(InComponent));
+
+	if (!ComponentState.ModifierWeak.IsValid())
+	{
+		ComponentState.Save();
+		bInOverrideContext = true;
+	}
+
+	if (bInOverrideContext)
+	{
+		ComponentState.ModifierWeak = InModifierContext;
+	}
+}
+
+void UAvaTranslucentPriorityModifierShared::RestoreComponentState(const UAvaTranslucentPriorityModifier* InModifierContext, UPrimitiveComponent* InComponent, bool bInClearState)
+{
+	if (!IsValid(InComponent))
+	{
+		return;
+	}
+
+	if (const FAvaTranslucentPriorityModifierComponentState* ComponentState = ComponentStates.Find(FAvaTranslucentPriorityModifierComponentState(InComponent)))
+	{
+		if (ComponentState->ModifierWeak == InModifierContext)
+		{
+			ComponentState->Restore();
+
+			if (bInClearState)
+			{
+				ComponentStates.Remove(*ComponentState);
+
+				// This is needed in case other parent translucent modifier take ownership of the component
+				UActorComponent::MarkRenderStateDirtyEvent.Broadcast(*InComponent);
+			}
+		}
+	}
+}
+
+void UAvaTranslucentPriorityModifierShared::RestoreComponentsState(const UAvaTranslucentPriorityModifier* InModifierContext, bool bInClearState)
+{
+	if (!InModifierContext)
+	{
+		return;
+	}
+
+	for (TSet<FAvaTranslucentPriorityModifierComponentState>::TIterator It(ComponentStates); It; ++It)
+	{
+		if (InModifierContext == It->ModifierWeak)
+		{
+			RestoreComponentState(It->ModifierWeak.Get(), It->PrimitiveComponentWeak.Get(), bInClearState);
+		}
+	}
+}
+
+void UAvaTranslucentPriorityModifierShared::RestoreComponentsState(const UAvaTranslucentPriorityModifier* InModifierContext, const TSet<TWeakObjectPtr<UPrimitiveComponent>>& InComponents, bool bInClearState)
+{
+	if (!InModifierContext)
+	{
+		return;
+	}
+
+	for (const TWeakObjectPtr<UPrimitiveComponent>& Component : InComponents)
+	{
+		RestoreComponentState(InModifierContext, Component.Get(), bInClearState);
+	}
+}
+
+UAvaTranslucentPriorityModifier* UAvaTranslucentPriorityModifierShared::FindModifierContext(UPrimitiveComponent* InComponent) const
+{
+	UAvaTranslucentPriorityModifier* ModifierContext = nullptr;
+
+	if (!InComponent)
+	{
+		return ModifierContext;
+	}
+
+	if (const FAvaTranslucentPriorityModifierComponentState* ComponentState = ComponentStates.Find(FAvaTranslucentPriorityModifierComponentState(InComponent)))
+	{
+		ModifierContext = ComponentState->ModifierWeak.Get();
+	}
+
+	return ModifierContext;
+}
+
+TArray<FAvaTranslucentPriorityModifierComponentState> UAvaTranslucentPriorityModifierShared::GetSortedComponentStates(UAvaTranslucentPriorityModifier* InModifierContext) const
+{
+	TArray<FAvaTranslucentPriorityModifierComponentState> SortedComponentStates;
+
+	if (!InModifierContext)
+	{
+		return SortedComponentStates;
+	}
+
+	SortedComponentStates.Reserve(ComponentStates.Num());
+
+	const EAvaTranslucentPriorityModifierMode SortMode = InModifierContext->GetMode();
+
+	if (SortMode == EAvaTranslucentPriorityModifierMode::Manual)
+	{
+		// Only gather components the modifier is handling
+		for (const FAvaTranslucentPriorityModifierComponentState& ComponentState : ComponentStates)
+		{
+			const UAvaTranslucentPriorityModifier* ModifierContext = ComponentState.ModifierWeak.Get();
+			const UPrimitiveComponent* PrimitiveComponent = ComponentState.PrimitiveComponentWeak.Get();
+
+			if (!ModifierContext || !PrimitiveComponent)
+			{
+				continue;
+			}
+
+			if (ModifierContext == InModifierContext)
+			{
+				SortedComponentStates.Add(ComponentState);
+			}
+		}
+	}
+	else if (SortMode == EAvaTranslucentPriorityModifierMode::AutoCameraDistance)
+	{
+		// Gather components with same camera context
+		const ACameraActor* CameraActor = InModifierContext->GetCameraActorWeak().Get();
+
+		if (!CameraActor || CameraActor->GetWorld() != InModifierContext->GetWorld())
+		{
+			return SortedComponentStates;
+		}
+
+		for (const FAvaTranslucentPriorityModifierComponentState& ComponentState : ComponentStates)
+		{
+			const UAvaTranslucentPriorityModifier* ModifierContext = ComponentState.ModifierWeak.Get();
+			const UPrimitiveComponent* PrimitiveComponent = ComponentState.PrimitiveComponentWeak.Get();
+
+			if (!ModifierContext || !PrimitiveComponent)
+			{
+				continue;
+			}
+
+			if (ModifierContext->GetMode() == InModifierContext->GetMode()
+				&& ModifierContext->CameraActorWeak == CameraActor)
+			{
+				SortedComponentStates.Add(ComponentState);
+			}
+		}
+
+		// Sort current modifiers by distance X to a specific camera
+		SortedComponentStates.StableSort([CameraActor](const FAvaTranslucentPriorityModifierComponentState& InComponentA, const FAvaTranslucentPriorityModifierComponentState& InComponentB)->bool
+		{
+			const FVector AForwardComponentLocation = CameraActor->GetActorForwardVector() * InComponentA.GetComponentLocation();
+			const float ADist = FVector::Distance(AForwardComponentLocation, CameraActor->GetActorLocation());
+
+			const FVector BForwardComponentLocation = CameraActor->GetActorForwardVector() * InComponentB.GetComponentLocation();
+			const float BDist = FVector::Distance(BForwardComponentLocation, CameraActor->GetActorLocation());
+
+			return ADist > BDist;
+		});
+	}
+	else if (SortMode == EAvaTranslucentPriorityModifierMode::AutoOutlinerTree || SortMode == EAvaTranslucentPriorityModifierMode::AutoOutlinerTreeInverted)
+	{
+		const bool bInvertedTree = SortMode == EAvaTranslucentPriorityModifierMode::AutoOutlinerTreeInverted;
+
+		const AActor* ModifiedActor = InModifierContext->GetModifiedActor();
+		UWorld* const World = ModifiedActor->GetTypedOuter<UWorld>();
+
+		for (const FAvaTranslucentPriorityModifierComponentState& ComponentState : ComponentStates)
+		{
+			const UAvaTranslucentPriorityModifier* ModifierContext = ComponentState.ModifierWeak.Get();
+			const UPrimitiveComponent* PrimitiveComponent = ComponentState.PrimitiveComponentWeak.Get();
+
+			if (!ModifierContext || !PrimitiveComponent)
+			{
+				continue;
+			}
+
+			const AActor* OwningActor = ModifierContext->GetModifiedActor();
+			if (!OwningActor)
+			{
+				continue;
+			}
+
+			if (ModifierContext->GetMode() == InModifierContext->GetMode()
+				&& OwningActor->GetTypedOuter<UWorld>() == World)
+			{
+				SortedComponentStates.Add(ComponentState);
+			}
+		}
+
+		// Sort current modifiers by outliner hierarchy
+#if WITH_EDITOR
+		// Use the outliner first if available for sorting
+		TSharedPtr<IAvaOutliner> AvaOutliner = FAvaOutlinerUtils::EditorGetOutliner(World);
+		if (AvaOutliner.IsValid())
+		{
+			SortedComponentStates.StableSort([AvaOutliner, bInvertedTree](const FAvaTranslucentPriorityModifierComponentState& InComponentA, const FAvaTranslucentPriorityModifierComponentState& InComponentB)->bool
+			{
+				const FAvaOutlinerItemPtr OutlinerItemA = AvaOutliner->FindItem(InComponentA.GetOwningActor());
+				const FAvaOutlinerItemPtr OutlinerItemB = AvaOutliner->FindItem(InComponentB.GetOwningActor());
+
+				return UE::AvaOutliner::CompareOutlinerItemOrder(OutlinerItemB, OutlinerItemA) ^ bInvertedTree;
+			});
+		}
+		else
+		{
+#endif
+			// Fallback on scene tree sorting
+			if (const IAvaSceneInterface* SceneInterface = FAvaActorUtils::GetSceneInterfaceFromActor(ModifiedActor))
+			{
+				const FAvaSceneTree& SceneTree = SceneInterface->GetSceneTree();
+
+				SortedComponentStates.StableSort([&SceneTree, World, bInvertedTree](const FAvaTranslucentPriorityModifierComponentState& InComponentA, const FAvaTranslucentPriorityModifierComponentState& InComponentB)->bool
+				{
+					const FAvaSceneTreeNode* SceneItemA = SceneTree.FindTreeNode(FAvaSceneItem(InComponentA.GetOwningActor(), World));
+					const FAvaSceneTreeNode* SceneItemB = SceneTree.FindTreeNode(FAvaSceneItem(InComponentB.GetOwningActor(), World));
+
+					return FAvaSceneTree::CompareTreeItemOrder(SceneItemB, SceneItemA) ^ bInvertedTree;
+				});
+			}
+#if WITH_EDITOR
+		}
+#endif
+	}
+
+	return SortedComponentStates;
+}
+
+void UAvaTranslucentPriorityModifierShared::SetSortPriorityOffset(int32 InOffset)
+{
+	if (SortPriorityOffset == InOffset)
+	{
+		return;
+	}
+
+	SortPriorityOffset = InOffset;
+	OnLevelGlobalsChangedDelegate.Broadcast();
+}
+
+void UAvaTranslucentPriorityModifierShared::SetSortPriorityStep(int32 InStep)
+{
+	if (SortPriorityStep == InStep)
+	{
+		return;
+	}
+
+	SortPriorityStep = InStep;
+	OnLevelGlobalsChangedDelegate.Broadcast();
+}
+
+void UAvaTranslucentPriorityModifierShared::PostLoad()
+{
+	Super::PostLoad();
+
+	// Clean up invalid components
+	for (TSet<FAvaTranslucentPriorityModifierComponentState>::TIterator It(ComponentStates); It; ++It)
+	{
+		if (!It->PrimitiveComponentWeak.IsValid())
+		{
+			It.RemoveCurrent();
+		}
+	}
+}
+
+void FAvaTranslucentPriorityModifierComponentState::Save()
+{
+	if (const UPrimitiveComponent* Component = PrimitiveComponentWeak.Get())
+	{
+		SortPriority = Component->TranslucencySortPriority;
+	}
+}
+
+void FAvaTranslucentPriorityModifierComponentState::Restore() const
+{
+	if (UPrimitiveComponent* Component = PrimitiveComponentWeak.Get())
+	{
+		Component->SetTranslucentSortPriority(SortPriority);
+	}
+}
+
+AActor* FAvaTranslucentPriorityModifierComponentState::GetOwningActor() const
+{
+	if (const UPrimitiveComponent* Component = PrimitiveComponentWeak.Get())
+	{
+		return Component->GetOwner();
+	}
+
+	return nullptr;
+}
+
+FVector FAvaTranslucentPriorityModifierComponentState::GetComponentLocation() const
+{
+	if (const UPrimitiveComponent* Component = PrimitiveComponentWeak.Get())
+	{
+		return Component->GetComponentLocation();
+	}
+
+	return FVector::ZeroVector;
+}
diff --git a/Engine/Plugins/VirtualProduction/Avalanche/Source/AvalancheOutliner/Private/Tests/AvaOutliner.spec.cpp b/Engine/Plugins/VirtualProduction/Avalanche/Source/AvalancheOutliner/Private/Tests/AvaOutliner.spec.cpp
new file mode 100644
--- /dev/null
+++ b/Engine/Plugins/VirtualProduction/Avalanche/Source/AvalancheOutliner/Private/Tests/AvaOutliner.spec.cpp
@@ -0,0 +1,302 @@
+// Copyright Epic Games, Inc. All Rights Reserved.
+
+#include "AvaOutliner.h"
+#include "AvaOutlinerView.h"
+#include "AvaOutlinerTestUtils.h"
+#include "Engine/DirectionalLight.h"
+#include "Engine/StaticMeshActor.h"
+#include "EngineUtils.h"
+#include "Item/AvaOutlinerActor.h"
+#include "Misc/AutomationTest.h"
+
+#if WITH_AUTOMATION_TESTS
+
+BEGIN_DEFINE_SPEC(FAvaOutlinerSpec, "Avalanche.Outliner"
+	, EAutomationTestFlags::EngineFilter
+	| EAutomationTestFlags::EditorContext
+	| EAutomationTestFlags_ApplicationContextMask)
+
+	// Types of common situations that are used across multiple tests
+	enum class EWhenType
+	{
+		NoItemsSelected,
+		ItemsSelected
+	};
+	
+	void When(const FString& InDescription, TFunction<void()>&& DoWork)
+	{
+		Describe("When " + InDescription, MoveTemp(DoWork));
+	}
+
+	void When(EWhenType InWhenType, TFunction<void()>&& DoWork)
+	{
+		switch (InWhenType)
+		{
+			case EWhenType::NoItemsSelected:
+				When("no items are selected", MoveTemp(DoWork));
+				break;
+
+			case EWhenType::ItemsSelected:
+				When("multiple items are selected", [this, DoWork = MoveTemp(DoWork)]()
+					{
+						BeforeEach([this]()
+							{
+								OutlinerView->SelectItems(TestItems);
+							});
+
+						DoWork();
+					});
+				break;
+		}
+	}
+
+	TSharedPtr<FAvaOutliner> Outliner;
+
+	// The outliner view instance created before each test
+	TSharedPtr<FAvaOutlinerView> OutlinerView;
+
+	// Custom test provider to fill some of the info/actions requested by the outliner (e.g. USelection Instances, or Actor Duplication)
+	TSharedPtr<UE::AvaOutliner::Private::FAvaOutlinerProviderTest> OutlinerProvider;
+
+	// A list of outliner items to test things with
+	TArray<FAvaOutlinerItemPtr> TestItems;
+
+END_DEFINE_SPEC(FAvaOutlinerSpec)
+
+void FAvaOutlinerSpec::Define()
+{
+	BeforeEach([this]()
+	{
+		OutlinerProvider = MakeShared<UE::AvaOutliner::Private::FAvaOutlinerProviderTest>();
+
+		constexpr int32 OutlinerViewId = 0;
+
+		Outliner = OutlinerProvider->GetOutliner();
+		Outliner->RegisterOutlinerView(OutlinerViewId);
+		Outliner->Refresh();
+
+		OutlinerView = Outliner->GetMostRecentOutlinerView();
+
+		const TArray<UObject*> TestObjects
+		{
+			OutlinerProvider->GetDirectionalLight(),
+			OutlinerProvider->GetFloor(),
+			OutlinerProvider->GetSkySphere(),
+		};
+
+		TestItems = OutlinerProvider->GetOutlinerItems(TestObjects);
+	 });
+
+	Describe("Spawning an Actor", [this]()
+	{
+		BeforeEach([this]()
+		{
+			OutlinerProvider->TestSpawnActor();
+			Outliner->Refresh();
+		});
+
+		It("should register an item for the actor", [this]()
+		{
+			const FAvaOutlinerItemPtr Item = Outliner->FindItem(OutlinerProvider->GetTestSpawnActor());
+			TestTrue(TEXT("Item found for the actor"), Item.IsValid());
+
+			if (!TestTrue(TEXT("Item is an actor item"), Item.IsValid() && Item->IsA<FAvaOutlinerActor>()))
+			{
+				return;
+			}
+
+			const AActor* const Actor = Item->CastTo<FAvaOutlinerActor>()->GetActor();
+			TestTrue(TEXT("Actor item matches the spawned actor"), Actor == OutlinerProvider->GetTestSpawnActor());	
+		});
+	});
+	
+	Describe("Duplicating Selected Items", [this]()
+	{
+		When(EWhenType::NoItemsSelected, [this]()
+		{
+			It("should keep the outliner with no selections", [this]()
+			{
+				OutlinerView->DuplicateSelected();
+				Outliner->Refresh();
+				TestEqual(TEXT("Selected Item count is 0")
+					, Outliner->GetSelectedItemCount()
+					, 0);
+			});
+		});
+
+		When(EWhenType::ItemsSelected, [this]()
+		{
+			Describe("When selected objects have duplicated", [this]()
+			{
+				BeforeEach([this]()
+				{
+					OutlinerView->DuplicateSelected();
+				});
+				
+				It("should have a matching pending item action count", [this]()
+				{
+					TestEqual(TEXT("Number of items enqueued match the requested number")
+						, Outliner->GetPendingItemActionCount()
+						, TestItems.Num());
+				});
+
+				Describe("When outliner has refreshed", [this]()
+				{
+					BeforeEach([this]()
+					{
+						Outliner->Refresh();
+					});
+
+					It("should have a matching pending item action count", [this]()
+					{
+						TestEqual(TEXT("The template item count and the number of copies match")
+							, Outliner->GetSelectedItemCount()
+							, TestItems.Num());	
+					});
+
+					It("should have the duplicate copies selected instead of the template", [this]()
+					{
+						// These should be the new items that were created from duplication
+						const TArray<FAvaOutlinerItemPtr> DuplicateItems = Outliner->GetSelectedItems();
+
+						TestFalse(TEXT("The newly selected item count is not zero"), DuplicateItems.IsEmpty());
+
+						constexpr int32 TestItemIndex = 0;
+							
+						FAvaOutliner::SortItems(TestItems);
+							
+						TestNotEqual(TEXT("The newly selected items are different from the old selected items")
+							, DuplicateItems[TestItemIndex]
+							, TestItems[TestItemIndex]);
+					});
+
+					It("should have each duplicate item right below its respective template item", [this]()
+					{
+						constexpr int32 TestItemIndex = 0;
+
+						FAvaOutliner::SortItems(TestItems);
+						
+						const FAvaOutlinerItemPtr TemplateItem  = TestItems[TestItemIndex];
+						const FAvaOutlinerItemPtr DuplicateItem = Outliner->GetSelectedItems()[TestItemIndex];
+
+						TestEqual(TEXT("The template item and duplicate item share the same parent item")
+							, TemplateItem->GetParent()
+							, DuplicateItem->GetParent());
+
+						const FAvaOutlinerItemPtr ParentItem = TemplateItem->GetParent();
+
+						const int32 TemplateItemIndex  = ParentItem->GetChildIndex(TemplateItem);
+						const int32 DuplicateItemIndex = ParentItem->GetChildIndex(DuplicateItem);
+
+						TestEqual(TEXT("The duplicate item is right below template item")
+							, TemplateItemIndex + 1
+							, DuplicateItemIndex);
+					});
+				});
+			});
+		});	
+	});
+
+	Describe("Grouping Selected Items", [this]()
+	{
+		When("grouping with null", [this]()
+		{
+			When(EWhenType::NoItemsSelected, [this]()
+			{
+				It("should keep the outliner with no selections", [this]()
+				{
+					Outliner->GroupSelection(nullptr);
+					Outliner->Refresh();
+					TestEqual(TEXT("Selected Item count is 0")
+						, Outliner->GetSelectedItemCount()
+						, 0);
+				});
+			});
+
+			When(EWhenType::ItemsSelected, [this]()
+			{
+				It("should keep the outliner with the same selected items", [this]()
+				{
+					const TArray<FAvaOutlinerItemPtr> PrevSelectedItems = Outliner->GetSelectedItems();
+
+					Outliner->GroupSelection(nullptr);
+					Outliner->Refresh();
+
+					const TArray<FAvaOutlinerItemPtr> CurrSelectedItems = Outliner->GetSelectedItems();
+
+					TestEqual(TEXT("Selected Item count remains the same")
+						, PrevSelectedItems.Num()
+						, CurrSelectedItems.Num());
+
+					TestEqual(TEXT("Selected Items remain the same")
+						, PrevSelectedItems[0]
+						, CurrSelectedItems[0]);
+				});
+			});
+		});
+
+		When("grouping with an already existing actor", [this]()
+		{
+			When(EWhenType::NoItemsSelected, [this]()
+			{
+				BeforeEach([this]()
+				{
+					AActor* const GroupingActor = OutlinerProvider->GetVolumeActor();
+					TestNotNull(TEXT("Grouping actor exists"), GroupingActor);
+					Outliner->GroupSelection(GroupingActor);
+					Outliner->Refresh();
+				});
+
+				It("should keep the outliner with no selections", [this]()
+				{
+					TestEqual(TEXT("Selected item count is zero"), Outliner->GetSelectedItemCount(), 0);
+				});
+			});
+
+			When(EWhenType::ItemsSelected, [this]()
+			{
+				BeforeEach([this]()
+				{
+					AActor* const GroupingActor = OutlinerProvider->GetVolumeActor();
+					TestNotNull(TEXT("Grouping actor exists"), GroupingActor);
+					Outliner->GroupSelection(GroupingActor);
+					Outliner->Refresh();
+				});
+
+				It("should keep the outliner with the selected items", [this]()
+				{
+					const TArray<FAvaOutlinerItemPtr> SelectedItems = Outliner->GetSelectedItems();
+
+					TestEqual(TEXT("Selected item count is consistent"), SelectedItems.Num(), TestItems.Num());
+
+					if (TestNotEqual(TEXT("Selected item count is not zero"), SelectedItems.Num(), 0))
+					{
+						TestEqual(TEXT("The selected items are consistent"), TestItems[0], SelectedItems[0]);
+					}
+				});
+
+				It("should have placed the selected items under the grouping actor", [this]()
+				{
+					const AActor* const GroupingActor = OutlinerProvider->GetVolumeActor();
+					const FAvaOutlinerItemPtr GroupingItem = Outliner->FindItem(GroupingActor);
+
+					TestNotNull(TEXT("Grouping item exists in outliner"), GroupingItem.Get());
+
+					const int32 ChildIndex = GroupingItem->GetChildIndex(TestItems[0]);
+
+					TestTrue(TEXT("the selected items are under the grouping item"), ChildIndex != INDEX_NONE);
+				});
+			});
+		});
+	});
+
+	AfterEach([this]()
+	{
+		TestItems.Reset();
+		OutlinerView.Reset();
+		Outliner.Reset();
+		OutlinerProvider.Reset();
+	});
+}
+
+#endif
diff --git a/Engine/Plugins/VirtualProduction/Avalanche/Source/AvalanchePropertyAnimator/Private/AvaPropertyAnimatorModule.cpp b/Engine/Plugins/VirtualProduction/Avalanche/Source/AvalanchePropertyAnimator/Private/AvaPropertyAnimatorModule.cpp
new file mode 100644
--- /dev/null
+++ b/Engine/Plugins/VirtualProduction/Avalanche/Source/AvalanchePropertyAnimator/Private/AvaPropertyAnimatorModule.cpp
@@ -0,0 +1,5 @@
+﻿// Copyright Epic Games, Inc. All Rights Reserved.
+
+#include "Modules/ModuleManager.h"
+
+IMPLEMENT_MODULE(FDefaultModuleImpl, AvalanchePropertyAnimator)
\ No newline at end of file
diff --git a/Engine/Plugins/VirtualProduction/Avalanche/Source/AvalancheSVGEditor/Private/Tool/AvaSVGActorTool.cpp b/Engine/Plugins/VirtualProduction/Avalanche/Source/AvalancheSVGEditor/Private/Tool/AvaSVGActorTool.cpp
new file mode 100644
--- /dev/null
+++ b/Engine/Plugins/VirtualProduction/Avalanche/Source/AvalancheSVGEditor/Private/Tool/AvaSVGActorTool.cpp
@@ -0,0 +1,38 @@
+// Copyright Epic Games, Inc. All Rights Reserved.
+
+#include "Tool/AvaSVGActorTool.h"
+#include "AvaInteractiveToolsSettings.h"
+#include "Builders/AvaInteractiveToolsToolBuilder.h"
+#include "Factories/SVGActorFactory.h"
+#include "SVGActor.h"
+#include "SVGImporter.h"
+#include "SVGImporterEditorCommands.h"
+
+UAvaSVGActorTool::UAvaSVGActorTool()
+{
+	ActorClass = ASVGActor::StaticClass();
+}
+
+void UAvaSVGActorTool::OnRegisterTool(IAvalancheInteractiveToolsModule* InAITModule)
+{
+	Super::OnRegisterTool(InAITModule);
+
+	// Load svg module
+	FSVGImporterModule::Get();
+
+	FAvaInteractiveToolsToolParameters ToolParameters =
+	{
+		FSVGImporterEditorCommands::GetExternal().SpawnSVGActor,
+		TEXT("SVG Actor Tool"),
+		6000,
+		FAvalancheInteractiveToolsCreateBuilder::CreateLambda(
+			[](UEdMode* InEdMode)
+			{
+				return UAvaInteractiveToolsToolBuilder::CreateToolBuilder<UAvaSVGActorTool>(InEdMode);
+			}),
+		ActorClass,
+		CreateActorFactory<USVGActorFactory>()
+	};
+
+	InAITModule->RegisterTool(IAvalancheInteractiveToolsModule::CategoryNameActor, MoveTemp(ToolParameters));
+}
diff --git a/Engine/Plugins/VirtualProduction/Avalanche/Source/AvalancheSequence/Private/Transition/AvaTransitionSequenceTask.cpp b/Engine/Plugins/VirtualProduction/Avalanche/Source/AvalancheSequence/Private/Transition/AvaTransitionSequenceTask.cpp
new file mode 100644
--- /dev/null
+++ b/Engine/Plugins/VirtualProduction/Avalanche/Source/AvalancheSequence/Private/Transition/AvaTransitionSequenceTask.cpp
@@ -0,0 +1,26 @@
+// Copyright Epic Games, Inc. All Rights Reserved.
+
+#include "Transition/AvaTransitionSequenceTask.h"
+#include "AvaTransitionUtils.h"
+#include "StateTreeExecutionContext.h"
+#include "Transition/AvaTransitionSequenceUtils.h"
+
+void FAvaTransitionSequenceTask::PostLoad(FStateTreeDataView InInstanceDataView)
+{
+	Super::PostLoad(InInstanceDataView);
+
+	PRAGMA_DISABLE_DEPRECATION_WARNINGS
+	if (QueryType_DEPRECATED != EAvaTransitionSequenceQueryType::None)
+	{
+		if (FInstanceDataType* InstanceData = UE::AvaTransition::TryGetInstanceData(*this, InInstanceDataView))
+		{
+			InstanceData->WaitType = WaitType_DEPRECATED;
+		}
+	}
+	PRAGMA_ENABLE_DEPRECATION_WARNINGS
+}
+
+EAvaTransitionSequenceWaitType FAvaTransitionSequenceTask::GetWaitType(FStateTreeExecutionContext& InContext) const
+{
+	return InContext.GetInstanceData(*this).WaitType;
+}
diff --git a/Engine/Plugins/VirtualProduction/Avalanche/Source/AvalancheSequence/Public/Transition/AvaTransitionInitializeSequence.h b/Engine/Plugins/VirtualProduction/Avalanche/Source/AvalancheSequence/Public/Transition/AvaTransitionInitializeSequence.h
new file mode 100644
--- /dev/null
+++ b/Engine/Plugins/VirtualProduction/Avalanche/Source/AvalancheSequence/Public/Transition/AvaTransitionInitializeSequence.h
@@ -0,0 +1,54 @@
+// Copyright Epic Games, Inc. All Rights Reserved.
+
+#pragma once
+
+#include "AvaSequenceShared.h"
+#include "AvaTransitionSequenceTask.h"
+#include "AvaTransitionInitializeSequence.generated.h"
+
+USTRUCT()
+struct FAvaTransitionInitSequenceTaskInstanceData : public FAvaTransitionSequenceTaskBaseInstanceData
+{
+	GENERATED_BODY()
+
+	UPROPERTY(EditAnywhere, Category = "Parameter")
+	FAvaSequenceTime InitializeTime = FAvaSequenceTime(0.0);
+
+	UPROPERTY(EditAnywhere, Category = "Parameter")
+	EAvaSequencePlayMode PlayMode = EAvaSequencePlayMode::Forward;
+};
+
+USTRUCT(DisplayName = "Initialize Sequence", Category = "Sequence Playback")
+struct AVALANCHESEQUENCE_API FAvaTransitionInitializeSequence : public FAvaTransitionSequenceTaskBase
+{
+	GENERATED_BODY()
+
+	using FInstanceDataType = FAvaTransitionInitSequenceTaskInstanceData;
+
+	PRAGMA_DISABLE_DEPRECATION_WARNINGS
+	FAvaTransitionInitializeSequence() = default;
+	virtual ~FAvaTransitionInitializeSequence() override = default;
+	FAvaTransitionInitializeSequence(const FAvaTransitionInitializeSequence&) = default;
+	FAvaTransitionInitializeSequence(FAvaTransitionInitializeSequence&&) = default;
+	FAvaTransitionInitializeSequence& operator=(const FAvaTransitionInitializeSequence&) = default;
+	FAvaTransitionInitializeSequence& operator=(FAvaTransitionInitializeSequence&&) = default;
+	PRAGMA_ENABLE_DEPRECATION_WARNINGS
+
+	//~ Begin FStateTreeNodeBase
+	virtual const UStruct* GetInstanceDataType() const override { return FInstanceDataType::StaticStruct(); }
+	virtual void PostLoad(FStateTreeDataView InInstanceDataView) override;
+	//~ End FStateTreeNodeBase
+
+	//~ Begin FAvaTransitionSequenceTaskBase
+	virtual EAvaTransitionSequenceWaitType GetWaitType(FStateTreeExecutionContext& InContext) const override;
+	virtual TArray<UAvaSequencePlayer*> ExecuteSequenceTask(FStateTreeExecutionContext& InContext) const override;
+	//~ End FAvaTransitionSequenceTaskBase
+
+	UE_DEPRECATED(5.5, "InitializeTime has been moved to Instance Data")
+	UPROPERTY(meta = (DeprecatedProperty, DeprecationMessage = "Use the Instance Data InitializeTime instead"))
+	FAvaSequenceTime InitializeTime_DEPRECATED = FAvaSequenceTime(0.0);
+
+	UE_DEPRECATED(5.5, "PlayMode has been moved to Instance Data")
+	UPROPERTY(meta = (DeprecatedProperty, DeprecationMessage = "Use the Instance Data PlayMode instead"))
+	EAvaSequencePlayMode PlayMode_DEPRECATED = EAvaSequencePlayMode::Forward;
+};
diff --git a/Engine/Plugins/VirtualProduction/Avalanche/Source/AvalancheSequence/Public/Transition/AvaTransitionPlaySequenceTask.h b/Engine/Plugins/VirtualProduction/Avalanche/Source/AvalancheSequence/Public/Transition/AvaTransitionPlaySequenceTask.h
new file mode 100644
--- /dev/null
+++ b/Engine/Plugins/VirtualProduction/Avalanche/Source/AvalancheSequence/Public/Transition/AvaTransitionPlaySequenceTask.h
@@ -0,0 +1,50 @@
+// Copyright Epic Games, Inc. All Rights Reserved.
+
+#pragma once
+
+#include "AvaSequenceShared.h"
+#include "AvaTransitionSequenceTask.h"
+#include "AvaTransitionPlaySequenceTask.generated.h"
+
+USTRUCT()
+struct FAvaTransitionPlaySequenceTaskInstanceData : public FAvaTransitionSequenceTaskInstanceData
+{
+	GENERATED_BODY()
+
+	/** Sequence Play Settings */
+	UPROPERTY(EditAnywhere, Category="Parameter")
+	FAvaSequencePlayParams PlaySettings;
+};
+
+USTRUCT(DisplayName = "Play Sequence", Category="Sequence Playback")
+struct AVALANCHESEQUENCE_API FAvaTransitionPlaySequenceTask : public FAvaTransitionSequenceTask
+{
+	GENERATED_BODY()
+
+	using FInstanceDataType = FAvaTransitionPlaySequenceTaskInstanceData;
+
+	PRAGMA_DISABLE_DEPRECATION_WARNINGS
+	FAvaTransitionPlaySequenceTask() = default;
+	virtual ~FAvaTransitionPlaySequenceTask() override = default;
+	FAvaTransitionPlaySequenceTask(const FAvaTransitionPlaySequenceTask&) = default;
+	FAvaTransitionPlaySequenceTask(FAvaTransitionPlaySequenceTask&&) = default;
+	FAvaTransitionPlaySequenceTask& operator=(const FAvaTransitionPlaySequenceTask&) = default;
+	FAvaTransitionPlaySequenceTask& operator=(FAvaTransitionPlaySequenceTask&&) = default;
+	PRAGMA_ENABLE_DEPRECATION_WARNINGS
+
+	//~ Begin FStateTreeNodeBase
+	virtual const UStruct* GetInstanceDataType() const override { return FInstanceDataType::StaticStruct(); }
+	virtual void PostLoad(FStateTreeDataView InInstanceDataView) override;
+#if WITH_EDITOR
+	virtual FText GetDescription(const FGuid& InId, FStateTreeDataView InInstanceDataView, const IStateTreeBindingLookup& InBindingLookup, EStateTreeNodeFormatting InFormatting) const override;
+#endif
+	//~ End FStateTreeNodeBase
+
+	//~ Begin FAvaTransitionSequenceTaskBase
+	virtual TArray<UAvaSequencePlayer*> ExecuteSequenceTask(FStateTreeExecutionContext& InContext) const override;
+	//~ End FAvaTransitionSequenceTaskBase
+
+	UE_DEPRECATED(5.5, "PlaySettings has been moved to Instance Data")
+	UPROPERTY(meta = (DeprecatedProperty, DeprecationMessage = "Use the Instance Data PlaySettings instead"))
+	FAvaSequencePlayParams PlaySettings_DEPRECATED;
+};
diff --git a/Engine/Plugins/VirtualProduction/Avalanche/Source/AvalancheSequence/Public/Transition/AvaTransitionSequenceTask.h b/Engine/Plugins/VirtualProduction/Avalanche/Source/AvalancheSequence/Public/Transition/AvaTransitionSequenceTask.h
new file mode 100644
--- /dev/null
+++ b/Engine/Plugins/VirtualProduction/Avalanche/Source/AvalancheSequence/Public/Transition/AvaTransitionSequenceTask.h
@@ -0,0 +1,47 @@
+// Copyright Epic Games, Inc. All Rights Reserved.
+
+#pragma once
+
+#include "AvaTransitionSequenceTaskBase.h"
+#include "AvaTransitionSequenceTask.generated.h"
+
+USTRUCT()
+struct FAvaTransitionSequenceTaskInstanceData : public FAvaTransitionSequenceTaskBaseInstanceData
+{
+	GENERATED_BODY()
+
+	/** The wait type before this task completes */
+	UPROPERTY(EditAnywhere, Category="Parameter")
+	EAvaTransitionSequenceWaitType WaitType = EAvaTransitionSequenceWaitType::WaitUntilStop;
+};
+
+/** Base Task but with additional Parameters */
+USTRUCT(meta=(Hidden))
+struct AVALANCHESEQUENCE_API FAvaTransitionSequenceTask : public FAvaTransitionSequenceTaskBase
+{
+	GENERATED_BODY()
+
+	using FInstanceDataType = FAvaTransitionSequenceTaskInstanceData;
+
+	PRAGMA_DISABLE_DEPRECATION_WARNINGS
+	FAvaTransitionSequenceTask() = default;
+	virtual ~FAvaTransitionSequenceTask() override = default;
+	FAvaTransitionSequenceTask(const FAvaTransitionSequenceTask&) = default;
+	FAvaTransitionSequenceTask(FAvaTransitionSequenceTask&&) = default;
+	FAvaTransitionSequenceTask& operator=(const FAvaTransitionSequenceTask&) = default;
+	FAvaTransitionSequenceTask& operator=(FAvaTransitionSequenceTask&&) = default;
+	PRAGMA_ENABLE_DEPRECATION_WARNINGS
+
+	//~ Begin FStateTreeNodeBase
+	virtual const UStruct* GetInstanceDataType() const override { return FInstanceDataType::StaticStruct(); }
+	virtual void PostLoad(FStateTreeDataView InInstanceDataView) override;
+	//~ End FStateTreeNodeBase
+
+	//~ Begin FAvaTransitionSequenceTaskBase
+	virtual EAvaTransitionSequenceWaitType GetWaitType(FStateTreeExecutionContext& InContext) const override;
+	//~ End FAvaTransitionSequenceTaskBase
+
+	UE_DEPRECATED(5.5, "WaitType has been moved to Instance Data")
+	UPROPERTY(meta = (DeprecatedProperty, DeprecationMessage = "Use the Instance Data WaitType instead"))
+	EAvaTransitionSequenceWaitType WaitType_DEPRECATED = EAvaTransitionSequenceWaitType::WaitUntilStop;
+};
diff --git a/Engine/Plugins/VirtualProduction/Avalanche/Source/AvalancheSequencer/Private/DetailsView/SAvaKeyFrameEdit.cpp b/Engine/Plugins/VirtualProduction/Avalanche/Source/AvalancheSequencer/Private/DetailsView/SAvaKeyFrameEdit.cpp
new file mode 100644
--- /dev/null
+++ b/Engine/Plugins/VirtualProduction/Avalanche/Source/AvalancheSequencer/Private/DetailsView/SAvaKeyFrameEdit.cpp
@@ -0,0 +1,59 @@
+// Copyright Epic Games, Inc. All Rights Reserved.
+
+#include "SAvaKeyFrameEdit.h"
+#include "AvaSequencer.h"
+#include "EaseCurveTool/AvaEaseCurveTool.h"
+#include "MVVM/Selection/Selection.h"
+#include "MVVM/ViewModels/SequencerEditorViewModel.h"
+#include "Widgets/Layout/SScaleBox.h"
+#include "Widgets/Layout/SScrollBox.h"
+
+using namespace UE::Sequencer;
+
+#define LOCTEXT_NAMESPACE "SAvaKeyFrameEdit"
+
+void SAvaKeyFrameEdit::Construct(const FArguments& InArgs, const TSharedRef<FAvaSequencer>& InSequencer)
+{
+	AvaSequencerWeak = InSequencer;
+
+	KeyEditData = InArgs._KeyEditData;
+
+	const TSharedRef<ISequencer> Sequencer = InSequencer->GetSequencer();
+
+	if (const TSharedPtr<FSequencerEditorViewModel> SequencerViewModel = Sequencer->GetViewModel())
+	{
+		SequencerSelectionWeak = SequencerViewModel->GetSelection();
+	}
+
+	if (!AvaSequencerWeak.IsValid())
+	{
+		return;
+	}
+
+	const TSharedPtr<FAvaSequencer> AvaSequencer = AvaSequencerWeak.Pin();
+
+	ChildSlot
+		[
+			SNew(SScrollBox)
+			+ SScrollBox::Slot()
+			.AutoSize()
+			.HAlign(HAlign_Fill)
+			.Padding(1.f, 1.f, 1.f, 3.f)
+			[
+				SNew(SScaleBox)
+				.Stretch(EStretch::Fill)
+				[
+					InSequencer->GetEaseCurveTool()->GenerateWidget()
+				]
+			]
+			+ SScrollBox::Slot()
+			.AutoSize()
+			.HAlign(HAlign_Fill)
+			[
+				SNew(SKeyEditInterface, AvaSequencer->GetSequencer())
+				.EditData(KeyEditData)
+			]
+		];
+}
+
+#undef LOCTEXT_NAMESPACE
diff --git a/Engine/Plugins/VirtualProduction/Avalanche/Source/AvalancheSequencer/Private/EaseCurveTool/AvaEaseCurve.cpp b/Engine/Plugins/VirtualProduction/Avalanche/Source/AvalancheSequencer/Private/EaseCurveTool/AvaEaseCurve.cpp
new file mode 100644
--- /dev/null
+++ b/Engine/Plugins/VirtualProduction/Avalanche/Source/AvalancheSequencer/Private/EaseCurveTool/AvaEaseCurve.cpp
@@ -0,0 +1,140 @@
+// Copyright Epic Games, Inc. All Rights Reserved.
+
+#include "AvaEaseCurve.h"
+
+UAvaEaseCurve::UAvaEaseCurve()
+{
+	StartKeyHandle = FloatCurve.AddKey(0, 0);
+	EndKeyHandle = FloatCurve.AddKey(1, 1);
+}
+
+void UAvaEaseCurve::OnCurveChanged(const TArray<FRichCurveEditInfo>& InRichCurveEditInfos) 
+{
+	UCurveFloat::OnCurveChanged(InRichCurveEditInfos);
+}
+
+FKeyHandle UAvaEaseCurve::GetStartKeyHandle() const
+{
+	return StartKeyHandle;
+}
+
+FKeyHandle UAvaEaseCurve::GetEndKeyHandle() const
+{
+	return EndKeyHandle;
+}
+
+FRichCurveKey& UAvaEaseCurve::GetStartKey()
+{
+	return FloatCurve.GetKey(StartKeyHandle);
+}
+
+FRichCurveKey& UAvaEaseCurve::GetEndKey()
+{
+	return FloatCurve.GetKey(EndKeyHandle);
+}
+
+float UAvaEaseCurve::GetStartKeyTime() const
+{
+	return FloatCurve.GetKey(StartKeyHandle).Time;
+}
+
+float UAvaEaseCurve::GetStartKeyValue() const
+{
+	return FloatCurve.GetKey(StartKeyHandle).Value;
+}
+
+float UAvaEaseCurve::GetEndKeyTime() const
+{
+	return FloatCurve.GetKey(EndKeyHandle).Time;
+}
+
+float UAvaEaseCurve::GetEndKeyValue() const
+{
+	return FloatCurve.GetKey(EndKeyHandle).Value;
+}
+
+FAvaEaseCurveTangents UAvaEaseCurve::GetTangents() const
+{
+	return FAvaEaseCurveTangents(FloatCurve.GetKey(StartKeyHandle), FloatCurve.GetKey(EndKeyHandle));
+}
+
+void UAvaEaseCurve::SetTangents(const FAvaEaseCurveTangents& InTangents)
+{
+	SetStartTangent(InTangents.Start, InTangents.StartWeight);
+	SetEndTangent(InTangents.End, InTangents.EndWeight);
+}
+
+void UAvaEaseCurve::SetStartTangent(const float InTangent, const float InTangentWeight)
+{
+	FloatCurve.SetKeyTime(StartKeyHandle, 0);
+	FloatCurve.SetKeyValue(StartKeyHandle, 0);
+	FloatCurve.SetKeyInterpMode(StartKeyHandle, ERichCurveInterpMode::RCIM_Cubic);
+	FloatCurve.SetKeyTangentWeightMode(StartKeyHandle, ERichCurveTangentWeightMode::RCTWM_WeightedBoth, false);
+	FloatCurve.SetKeyTangentMode(StartKeyHandle, ERichCurveTangentMode::RCTM_Break, false);
+
+	FRichCurveKey& StartKey = GetStartKey();
+	StartKey.ArriveTangent = 0;
+	StartKey.ArriveTangentWeight = 0;
+	StartKey.LeaveTangent = InTangent;
+	StartKey.LeaveTangentWeight = InTangentWeight;
+}
+
+void UAvaEaseCurve::SetEndTangent(const float InTangent, const float InTangentWeight)
+{
+	FloatCurve.SetKeyTime(EndKeyHandle, 1);
+	FloatCurve.SetKeyValue(EndKeyHandle, 1);
+	FloatCurve.SetKeyInterpMode(EndKeyHandle, ERichCurveInterpMode::RCIM_Cubic);
+	FloatCurve.SetKeyTangentWeightMode(EndKeyHandle, ERichCurveTangentWeightMode::RCTWM_WeightedBoth, false);
+	FloatCurve.SetKeyTangentMode(EndKeyHandle, ERichCurveTangentMode::RCTM_Break, false);
+
+	FRichCurveKey& EndKey = GetEndKey();
+	EndKey.LeaveTangent = 0;
+	EndKey.LeaveTangentWeight = 0;
+	EndKey.ArriveTangent = InTangent;
+	EndKey.ArriveTangentWeight = InTangentWeight;
+}
+
+void UAvaEaseCurve::FlattenOrStraightenTangents(const FKeyHandle InKeyHandle, const bool bInFlattenTangents)
+{
+	FRichCurveKey& Key = FloatCurve.GetKey(InKeyHandle);
+
+	float LeaveTangent = Key.LeaveTangent;
+	float ArriveTangent = Key.ArriveTangent;
+
+	if (bInFlattenTangents)
+	{
+		LeaveTangent = 0.f;
+		ArriveTangent = 0.f;
+	}
+	else
+	{
+		LeaveTangent = (LeaveTangent + ArriveTangent) * 0.5f;
+		ArriveTangent = LeaveTangent;
+	}
+
+	Key.LeaveTangent = LeaveTangent;
+	Key.ArriveTangent = ArriveTangent;
+	if (Key.InterpMode == RCIM_Cubic && Key.TangentMode == RCTM_Auto)
+	{
+		Key.TangentMode = RCTM_User;
+	}
+}
+
+void UAvaEaseCurve::SetInterpMode(const ERichCurveInterpMode InInterpMode, const ERichCurveTangentMode InTangentMode)
+{
+	if (InInterpMode != RCIM_Cubic)
+	{
+		FloatCurve.SetKeyInterpMode(StartKeyHandle, InInterpMode);
+		FloatCurve.SetKeyTangentMode(StartKeyHandle, InTangentMode);
+
+		FloatCurve.SetKeyInterpMode(EndKeyHandle, InInterpMode);
+		FloatCurve.SetKeyTangentMode(EndKeyHandle, InTangentMode);
+	}
+}
+
+void UAvaEaseCurve::BroadcastUpdate()
+{
+	TArray<FRichCurveEditInfo> ChangedCurveEditInfos;
+	ChangedCurveEditInfos.Add(FRichCurveEditInfo(&FloatCurve));
+	OnCurveChanged(ChangedCurveEditInfos);
+}
diff --git a/Engine/Plugins/VirtualProduction/Avalanche/Source/AvalancheSequencer/Private/EaseCurveTool/AvaEaseCurve.h b/Engine/Plugins/VirtualProduction/Avalanche/Source/AvalancheSequencer/Private/EaseCurveTool/AvaEaseCurve.h
new file mode 100644
--- /dev/null
+++ b/Engine/Plugins/VirtualProduction/Avalanche/Source/AvalancheSequencer/Private/EaseCurveTool/AvaEaseCurve.h
@@ -0,0 +1,46 @@
+// Copyright Epic Games, Inc. All Rights Reserved.
+
+#pragma once
+
+#include "AvaEaseCurveTangents.h"
+#include "Curves/CurveFloat.h"
+#include "AvaEaseCurve.generated.h"
+
+UCLASS()
+class UAvaEaseCurve : public UCurveFloat 
+{
+	GENERATED_BODY()
+
+public:
+	UAvaEaseCurve();
+
+	FKeyHandle GetStartKeyHandle() const;
+	FKeyHandle GetEndKeyHandle() const;
+
+	FRichCurveKey& GetStartKey();
+	FRichCurveKey& GetEndKey();
+
+	float GetStartKeyTime() const;
+	float GetStartKeyValue() const;
+	float GetEndKeyTime() const;
+	float GetEndKeyValue() const;
+
+	FAvaEaseCurveTangents GetTangents() const;
+	void SetTangents(const FAvaEaseCurveTangents& InTangents);
+	void SetStartTangent(const float InTangent, const float InTangentWeight);
+	void SetEndTangent(const float InTangent, const float InTangentWeight);
+
+	void FlattenOrStraightenTangents(const FKeyHandle InKeyHandle, const bool bInFlattenTangents);
+
+	void SetInterpMode(const ERichCurveInterpMode InInterpMode, const ERichCurveTangentMode InTangentMode);
+
+	void BroadcastUpdate();
+
+	//~ Begin UCurveFloat
+	virtual void OnCurveChanged(const TArray<FRichCurveEditInfo>& InRichCurveEditInfos) override;
+	//~ End UCurveFloat
+
+protected:
+	FKeyHandle StartKeyHandle;
+	FKeyHandle EndKeyHandle;
+};
diff --git a/Engine/Plugins/VirtualProduction/Avalanche/Source/AvalancheSequencer/Private/EaseCurveTool/AvaEaseCurvePreset.h b/Engine/Plugins/VirtualProduction/Avalanche/Source/AvalancheSequencer/Private/EaseCurveTool/AvaEaseCurvePreset.h
new file mode 100644
--- /dev/null
+++ b/Engine/Plugins/VirtualProduction/Avalanche/Source/AvalancheSequencer/Private/EaseCurveTool/AvaEaseCurvePreset.h
@@ -0,0 +1,42 @@
+// Copyright Epic Games, Inc. All Rights Reserved.
+
+#pragma once
+
+#include "Containers/UnrealString.h"
+#include "EaseCurveTool/AvaEaseCurveTangents.h"
+#include "Templates/SharedPointer.h"
+#include "AvaEaseCurvePreset.generated.h"
+
+USTRUCT()
+struct FAvaEaseCurvePreset
+{
+	GENERATED_BODY()
+
+	FAvaEaseCurvePreset() {}
+	FAvaEaseCurvePreset(const FString& InName, const FString& InCategory, const FAvaEaseCurveTangents& InTangents)
+		: Name(InName), Category(InCategory), Tangents(InTangents)
+	{}
+
+	FORCEINLINE bool operator==(const FAvaEaseCurvePreset& InRhs) const
+	{
+		return Name.Equals(InRhs.Name) && Category.Equals(InRhs.Category);
+	}
+	FORCEINLINE bool operator!=(const FAvaEaseCurvePreset& InRhs) const
+	{
+		return !(*this == InRhs);
+	}
+	FORCEINLINE bool operator<(const FAvaEaseCurvePreset& InRhs) const
+	{
+		// Order from soft to hard based on curve length
+		return Tangents.CalculateCurveLength() < InRhs.Tangents.CalculateCurveLength();
+	}
+
+	UPROPERTY(Config, EditAnywhere, Category = "EaseCurvePreset")
+	FString Name;
+
+	UPROPERTY(Config, EditAnywhere, Category = "EaseCurvePreset")
+	FString Category;
+
+	UPROPERTY(Config, EditAnywhere, Category = "EaseCurvePreset")
+	FAvaEaseCurveTangents Tangents;
+};
diff --git a/Engine/Plugins/VirtualProduction/Avalanche/Source/AvalancheSequencer/Private/EaseCurveTool/AvaEaseCurvePresetDragDropOp.cpp b/Engine/Plugins/VirtualProduction/Avalanche/Source/AvalancheSequencer/Private/EaseCurveTool/AvaEaseCurvePresetDragDropOp.cpp
new file mode 100644
--- /dev/null
+++ b/Engine/Plugins/VirtualProduction/Avalanche/Source/AvalancheSequencer/Private/EaseCurveTool/AvaEaseCurvePresetDragDropOp.cpp
@@ -0,0 +1,52 @@
+// Copyright Epic Games, Inc. All Rights Reserved.
+
+#include "AvaEaseCurvePresetDragDropOp.h"
+#include "EaseCurveTool/AvaEaseCurvePreset.h"
+#include "EaseCurveTool/Widgets/SAvaEaseCurvePresetGroup.h"
+#include "EaseCurveTool/Widgets/SAvaEaseCurvePresetGroupItem.h"
+#include "Styling/AppStyle.h"
+#include "Widgets/Layout/SBorder.h"
+
+FCursorReply FAvaEaseCurvePresetDragDropOperation::OnCursorQuery()
+{
+	return FCursorReply::Cursor(EMouseCursor::GrabHandClosed);
+}
+
+TSharedPtr<SWidget> FAvaEaseCurvePresetDragDropOperation::GetDefaultDecorator() const
+{
+	if (!WidgetWeak.IsValid())
+	{
+		return FDragDropOperation::GetDefaultDecorator();
+	}
+
+	return SNew(SBorder)
+		.Padding(2.f)
+		.BorderImage(FAppStyle::GetBrush(TEXT("Graph.ConnectorFeedback.Border")))
+		.Content()
+		[
+			WidgetWeak.Pin().ToSharedRef()
+		];
+}
+
+void FAvaEaseCurvePresetDragDropOperation::OnDrop(bool bInDropWasHandled, const FPointerEvent& InMouseEvent)
+{
+	FDragDropOperation::OnDrop(bInDropWasHandled, InMouseEvent);
+
+	if (const TSharedPtr<SAvaEaseCurvePresetGroupItem> Widget = WidgetWeak.Pin())
+	{
+		Widget->TriggerEndMove();
+	}
+
+	for (const TWeakPtr<SAvaEaseCurvePresetGroup>& GroupWidgetWeak : HoveredGroupWidgets)
+	{
+		if (const TSharedPtr<SAvaEaseCurvePresetGroup> GroupWidget = GroupWidgetWeak.Pin())
+		{
+			GroupWidget->ResetDragBorder();
+		}
+	}
+}
+
+void FAvaEaseCurvePresetDragDropOperation::AddHoveredGroup(const TSharedRef<SAvaEaseCurvePresetGroup>& InGroupWidget)
+{
+	HoveredGroupWidgets.Add(InGroupWidget);
+}
diff --git a/Engine/Plugins/VirtualProduction/Avalanche/Source/AvalancheSequencer/Private/EaseCurveTool/AvaEaseCurvePresetDragDropOp.h b/Engine/Plugins/VirtualProduction/Avalanche/Source/AvalancheSequencer/Private/EaseCurveTool/AvaEaseCurvePresetDragDropOp.h
new file mode 100644
--- /dev/null
+++ b/Engine/Plugins/VirtualProduction/Avalanche/Source/AvalancheSequencer/Private/EaseCurveTool/AvaEaseCurvePresetDragDropOp.h
@@ -0,0 +1,37 @@
+// Copyright Epic Games, Inc. All Rights Reserved.
+
+#pragma once
+
+#include "Input/DragAndDrop.h"
+#include "Templates/SharedPointer.h"
+
+class SAvaEaseCurvePresetGroup;
+class SAvaEaseCurvePresetGroupItem;
+struct FAvaEaseCurvePreset;
+
+class FAvaEaseCurvePresetDragDropOperation : public FDragDropOperation
+{
+public:
+	DRAG_DROP_OPERATOR_TYPE(FAvaEaseCurvePresetDragDropOperation, FDragDropOperation)
+
+	FAvaEaseCurvePresetDragDropOperation(const TSharedPtr<SAvaEaseCurvePresetGroupItem>& InWidget, const TSharedPtr<FAvaEaseCurvePreset>& InPreset)
+		: WidgetWeak(InWidget), Preset(InPreset)
+	{}
+
+	TSharedPtr<SAvaEaseCurvePresetGroupItem> GetWidget() const { return WidgetWeak.Pin(); }
+	TSharedPtr<FAvaEaseCurvePreset> GetPreset() const { return Preset; }
+
+	//~ Begin FDragDropOperation
+	virtual FCursorReply OnCursorQuery() override;
+	virtual TSharedPtr<SWidget> GetDefaultDecorator() const override;
+	virtual void OnDrop(bool bInDropWasHandled, const FPointerEvent& InMouseEvent) override;
+	//~ End FDragDropOperation
+
+	void AddHoveredGroup(const TSharedRef<SAvaEaseCurvePresetGroup>& InGroupWidget);
+
+protected:
+	TWeakPtr<SAvaEaseCurvePresetGroupItem> WidgetWeak;
+	TSharedPtr<FAvaEaseCurvePreset> Preset;
+
+	TSet<TWeakPtr<SAvaEaseCurvePresetGroup>> HoveredGroupWidgets;
+};
diff --git a/Engine/Plugins/VirtualProduction/Avalanche/Source/AvalancheSequencer/Private/EaseCurveTool/AvaEaseCurveStyle.cpp b/Engine/Plugins/VirtualProduction/Avalanche/Source/AvalancheSequencer/Private/EaseCurveTool/AvaEaseCurveStyle.cpp
new file mode 100644
--- /dev/null
+++ b/Engine/Plugins/VirtualProduction/Avalanche/Source/AvalancheSequencer/Private/EaseCurveTool/AvaEaseCurveStyle.cpp
@@ -0,0 +1,71 @@
+// Copyright Epic Games, Inc. All Rights Reserved.
+
+#include "AvaEaseCurveStyle.h"
+#include "Brushes/SlateImageBrush.h"
+#include "Brushes/SlateRoundedBoxBrush.h"
+#include "Interfaces/IPluginManager.h"
+#include "Fonts/SlateFontInfo.h"
+#include "Layout/Margin.h"
+#include "Styling/AppStyle.h"
+#include "Styling/CoreStyle.h"
+#include "Styling/StyleColors.h"
+#include "Styling/SlateStyleMacros.h"
+#include "Styling/SlateStyleRegistry.h"
+#include "Styling/SlateTypes.h"
+
+FAvaEaseCurveStyle::FAvaEaseCurveStyle()
+	: FSlateStyleSet(TEXT("EaseCurveTool"))
+{
+	const TSharedPtr<IPlugin> Plugin = IPluginManager::Get().FindPlugin(UE_PLUGIN_NAME);
+	check(Plugin.IsValid());
+	
+	SetContentRoot(Plugin->GetBaseDir() / TEXT("Resources"));
+
+	Set("Preset.Selected", new FSlateRoundedBoxBrush(FStyleColors::Transparent, FVector4(4.f, 0.f, 0.f, 4.f), FStyleColors::Select.GetSpecifiedColor(), 1.f));
+
+	Set("EditMode.Background", new FSlateRoundedBoxBrush(FStyleColors::Transparent, 4.f, FLinearColor(0.1f, 0.1f, 0.1f, 1.f), 1.f));
+	Set("EditMode.Background.Highlight", new FSlateRoundedBoxBrush(FStyleColors::Transparent, 4.f, FLinearColor(0.6f, 0.6f, 0.6f, 1.f), 1.f));
+	Set("EditMode.Background.Over", new FSlateRoundedBoxBrush(FStyleColors::Transparent, 4.f, FStyleColors::AccentBlue.GetSpecifiedColor(), 1.f));
+
+	const FButtonStyle& SimpleButtonStyle = FAppStyle::GetWidgetStyle<FButtonStyle>(TEXT("SimpleButton"));
+
+	constexpr float ToolButtonPadding = 2.f;
+	Set("ToolButton.Padding", ToolButtonPadding);
+
+	constexpr float ToolButtonImageSize = 12.f;
+	Set("ToolButton.ImageSize", ToolButtonImageSize);
+
+	const FButtonStyle ToolButtonStyle = FButtonStyle(SimpleButtonStyle)
+		.SetHovered(FSlateRoundedBoxBrush(FStyleColors::Hover, 4.f))
+		.SetPressed(FSlateRoundedBoxBrush(FStyleColors::Hover, 4.f))
+		.SetNormalPadding(FMargin(ToolButtonPadding))
+		.SetPressedPadding(FMargin(ToolButtonPadding, ToolButtonPadding + (ToolButtonPadding * 0.5f), ToolButtonPadding, ToolButtonPadding - (ToolButtonPadding * 0.5f)));
+	Set("ToolButton", ToolButtonStyle);
+
+	const FButtonStyle ToolButtonNoPadStyle = FButtonStyle(ToolButtonStyle)
+		.SetNormalPadding(FMargin(0.f))
+		.SetPressedPadding(FMargin(0.f));
+	Set("ToolButton.NoPad", ToolButtonNoPadStyle);
+
+	Set("ToolButton.Opaque", new FSlateRoundedBoxBrush(FStyleColors::Dropdown, 4.f));
+
+	const FCheckBoxStyle& ToggleButtonCheckboxStyle = FAppStyle::GetWidgetStyle<FCheckBoxStyle>(TEXT("ToggleButtonCheckbox"));
+
+	const FCheckBoxStyle ToolToggleButtonStyle = FCheckBoxStyle(ToggleButtonCheckboxStyle)
+		.SetUncheckedHoveredImage(FSlateRoundedBoxBrush(FStyleColors::Hover, 4.f))
+		.SetUncheckedPressedImage(FSlateRoundedBoxBrush(FStyleColors::Hover, 4.f))
+		.SetCheckedHoveredImage(FSlateRoundedBoxBrush(FStyleColors::PrimaryHover, 4.f))
+		.SetCheckedPressedImage(FSlateRoundedBoxBrush(FStyleColors::PrimaryPress, 4.f))
+		.SetPadding(ToolButtonPadding);
+	Set("ToolToggleButton", ToolToggleButtonStyle);
+
+	Set("Editor.LabelFont", FSlateFontInfo(FCoreStyle::GetDefaultFont(), 7, TEXT("Regular")));
+	Set("Editor.ErrorFont", FSlateFontInfo(FCoreStyle::GetDefaultFont(), 9, TEXT("Italic")));
+
+	FSlateStyleRegistry::RegisterSlateStyle(*this);
+}
+
+FAvaEaseCurveStyle::~FAvaEaseCurveStyle()
+{
+	FSlateStyleRegistry::UnRegisterSlateStyle(*this);
+}
diff --git a/Engine/Plugins/VirtualProduction/Avalanche/Source/AvalancheSequencer/Private/EaseCurveTool/AvaEaseCurveStyle.h b/Engine/Plugins/VirtualProduction/Avalanche/Source/AvalancheSequencer/Private/EaseCurveTool/AvaEaseCurveStyle.h
new file mode 100644
--- /dev/null
+++ b/Engine/Plugins/VirtualProduction/Avalanche/Source/AvalancheSequencer/Private/EaseCurveTool/AvaEaseCurveStyle.h
@@ -0,0 +1,18 @@
+// Copyright Epic Games, Inc. All Rights Reserved.
+
+#pragma once
+
+#include "Styling/SlateStyle.h"
+
+class FAvaEaseCurveStyle final : public FSlateStyleSet
+{
+public:
+	static FAvaEaseCurveStyle& Get()
+	{
+		static FAvaEaseCurveStyle Instance;
+		return Instance;
+	}
+
+	FAvaEaseCurveStyle();
+	virtual ~FAvaEaseCurveStyle() override;
+};
diff --git a/Engine/Plugins/VirtualProduction/Avalanche/Source/AvalancheSequencer/Private/EaseCurveTool/AvaEaseCurveSubsystem.cpp b/Engine/Plugins/VirtualProduction/Avalanche/Source/AvalancheSequencer/Private/EaseCurveTool/AvaEaseCurveSubsystem.cpp
new file mode 100644
--- /dev/null
+++ b/Engine/Plugins/VirtualProduction/Avalanche/Source/AvalancheSequencer/Private/EaseCurveTool/AvaEaseCurveSubsystem.cpp
@@ -0,0 +1,683 @@
+// Copyright Epic Games, Inc. All Rights Reserved.
+
+#include "AvaEaseCurveSubsystem.h"
+#include "Containers/UnrealString.h"
+#include "EaseCurveTool/AvaEaseCurvePreset.h"
+#include "EaseCurveTool/AvaEaseCurveToolSettings.h"
+#include "Editor.h"
+#include "Engine/World.h"
+#include "HAL/FileManager.h"
+#include "HAL/PlatformFileManager.h"
+#include "Interfaces/IPluginManager.h"
+#include "Misc/FileHelper.h"
+#include "Misc/Paths.h"
+#include "Serialization/JsonSerializer.h"
+
+DEFINE_LOG_CATEGORY_STATIC(LogAvaEaseCurveSubsystem, Log, All);
+
+namespace UE::EaseCurveTool::Json::Private
+{
+	bool GetPresetFilePath(FString InCategory, const bool bSetNewCategory, FString& OutFilePath)
+	{
+		OutFilePath.Empty();
+
+		const UAvaEaseCurveToolSettings* const Settings = GetDefault<UAvaEaseCurveToolSettings>();
+		check(Settings);
+
+		const FString PresetsPath = FPaths::ConvertRelativePathToFull(UAvaEaseCurveSubsystem::ProjectPresetPath());
+
+		if (!FPaths::DirectoryExists(PresetsPath))
+		{
+			UE_LOG(LogAvaEaseCurveSubsystem, Warning, TEXT("EaseCurveToolJson::GetPresetFilePath() - Configured presets directory does not exist: %s"), *PresetsPath);
+			return false;
+		}
+
+		if (InCategory.IsEmpty())
+		{
+			if (!bSetNewCategory)
+			{
+				return false;
+			}
+
+			const FString EaseCurveTool_NewPresetCategory = Settings->GetNewPresetCategory();
+			if (EaseCurveTool_NewPresetCategory.IsEmpty())
+			{
+				UE_LOG(LogAvaEaseCurveSubsystem, Warning, TEXT("EaseCurveToolJson::GetPresetFilePath() - No preset category and no configured default category"));
+				return false;
+			}
+
+			InCategory = EaseCurveTool_NewPresetCategory;
+		}
+
+		OutFilePath = FPaths::Combine(PresetsPath, InCategory) + TEXT(".json");
+
+		return FPaths::FileExists(OutFilePath);
+	}
+
+	bool LoadCurvePresetsJson(const FString& InFilePath, TSharedPtr<FJsonObject>& OutRootObject)
+	{
+		FString FileJsonString;
+		if (!FFileHelper::LoadFileToString(FileJsonString, *InFilePath))
+		{
+			UE_LOG(LogAvaEaseCurveSubsystem, Warning, TEXT("EaseCurveToolJson::LoadCurvePresetsJson() - Unable to load Json file: %s"), *InFilePath);
+			return false;
+		}
+
+		const TSharedPtr<TJsonReader<>> JsonReader = TJsonReaderFactory<>::Create(FileJsonString);
+		OutRootObject = MakeShared<FJsonObject>();
+		if (!FJsonSerializer::Deserialize(JsonReader.ToSharedRef(), OutRootObject))
+		{
+			UE_LOG(LogAvaEaseCurveSubsystem, Warning, TEXT("EaseCurveToolJson::LoadCurvePresetsJson() - Unable to parse file [%s]. Json=[%s]"), *InFilePath, *FileJsonString);
+			return false;
+		}
+
+		return true;
+	}
+
+	bool SaveCurvePresetsJson(const FString& InFilePath, const TSharedRef<FJsonObject>& InRootObject)
+	{
+		FString FileJsonString;
+		const TSharedRef<TJsonWriter<>> Writer = TJsonWriterFactory<>::Create(&FileJsonString);
+		if (!FJsonSerializer::Serialize(InRootObject, Writer))
+		{
+			UE_LOG(LogAvaEaseCurveSubsystem, Warning, TEXT("EaseCurveToolJson::SaveCurvePresetsJson() - Unable to serialize [%s]. Json=[%s]"), *InFilePath, *FileJsonString);
+			return false;
+		}
+
+		if (!FFileHelper::SaveStringToFile(FileJsonString, *InFilePath))
+		{
+			UE_LOG(LogAvaEaseCurveSubsystem, Warning, TEXT("EaseCurveToolJson::SaveCurvePresetsJson() - Unable to save Json file: %s"), *InFilePath);
+			return false;
+		}
+
+		return true;
+	}
+
+	bool DoesPresetCategoryExist(const FString& InCategory)
+	{
+		FString OutFilePath;
+		return GetPresetFilePath(InCategory, false, OutFilePath);
+	}
+
+	FString AddNewPresetCategory()
+	{
+		static const FString DefaultNewCategoryName = TEXT("New Category");
+
+		FString NewCategoryName;
+
+		int32 CurrentNum = 1;
+		FString FilePath;
+
+		while (true)
+		{
+			NewCategoryName = (CurrentNum == 1)
+				? DefaultNewCategoryName 
+				: DefaultNewCategoryName + TEXT(" ") + FString::FromInt(CurrentNum);
+
+			if (!GetPresetFilePath(NewCategoryName, false, FilePath))
+			{
+				break;
+			}
+
+			CurrentNum++;
+		}
+
+		const TSharedPtr<FJsonObject> RootObject = MakeShared<FJsonObject>();
+		if (!SaveCurvePresetsJson(FilePath, RootObject.ToSharedRef()))
+		{
+			return FString();
+		}
+
+		return NewCategoryName;
+	}
+
+	bool RemovePresetCategory(const FString& InCategory)
+	{
+		FString FilePath;
+		if (!GetPresetFilePath(InCategory, false, FilePath))
+		{
+			return false;
+		}
+
+		return IFileManager::Get().Delete(*FilePath);
+	}
+
+	bool RenamePresetCategory(const FString& InCategory, const FString& InNewCategory)
+	{
+		if (!DoesPresetCategoryExist(InCategory))
+		{
+			return false;
+		}
+
+		if (DoesPresetCategoryExist(InNewCategory))
+		{
+			return false;
+		}
+
+		FString OldFilePath;
+		if (!GetPresetFilePath(InCategory, false, OldFilePath))
+		{
+			return false;
+		}
+
+		FString NewFilePath;
+		GetPresetFilePath(InNewCategory, false, NewFilePath);
+
+		TSharedPtr<FJsonObject> RootObject;
+		if (!LoadCurvePresetsJson(OldFilePath, RootObject) || !RootObject.IsValid())
+		{
+			return false;
+		}
+
+		if (!SaveCurvePresetsJson(NewFilePath, RootObject.ToSharedRef()))
+		{
+			return false;
+		}
+
+		return RemovePresetCategory(InCategory);
+	}
+	
+	bool DoesPresetExist(const FString& InCategory, const FString& InPreset)
+	{
+		FString PresetsPath;
+		if (!GetPresetFilePath(InCategory, false, PresetsPath))
+		{
+			UE_LOG(LogAvaEaseCurveSubsystem, Warning, TEXT("EaseCurveToolJson::DoesPresetExist() - Could not find file for preset category: %s"), *InCategory);
+			return false;
+		}
+
+		TSharedPtr<FJsonObject> RootObject;
+		if (!LoadCurvePresetsJson(PresetsPath, RootObject) || !RootObject.IsValid())
+		{
+			UE_LOG(LogAvaEaseCurveSubsystem, Warning, TEXT("EaseCurveToolJson::DoesPresetExist() - Unable to load Json file: %s"), *PresetsPath);
+			return false;
+		}
+
+		return RootObject->Values.Contains(InPreset);
+	}
+
+	bool AddPreset(const FString& InCategory, const FString& InPreset, const FAvaEaseCurveTangents& InTangents)
+	{
+		FString PresetsPath;
+		const bool bCategoryFileExists = GetPresetFilePath(InCategory, true, PresetsPath);
+		if (PresetsPath.IsEmpty())
+		{
+			UE_LOG(LogAvaEaseCurveSubsystem, Warning, TEXT("EaseCurveToolJson::AddPreset() - Could not find file for preset category: %s"), *InCategory);
+			return false;
+		}
+
+		TSharedPtr<FJsonObject> RootObject;
+
+		if (bCategoryFileExists)
+		{
+			if (!UE::EaseCurveTool::Json::Private::LoadCurvePresetsJson(PresetsPath, RootObject) || !RootObject.IsValid())
+			{
+				UE_LOG(LogAvaEaseCurveSubsystem, Warning, TEXT("EaseCurveToolJson::AddPreset() - Unable to load Json file: %s"), *PresetsPath);
+				return false;
+			}
+		}
+		else
+		{
+			RootObject = MakeShared<FJsonObject>();
+		}
+
+		RootObject->Values.Add(InPreset, MakeShared<FJsonValueString>(InTangents.ToJson()));
+
+		if (!SaveCurvePresetsJson(PresetsPath, RootObject.ToSharedRef()))
+		{
+			UE_LOG(LogAvaEaseCurveSubsystem, Warning, TEXT("EaseCurveToolJson::AddPreset() - Unable to save Json file: %s"), *PresetsPath);
+			return false;
+		}
+
+		return true;
+	}
+
+	bool RemovePreset(const FString& InCategory, const FString& InPreset)
+	{
+		FString PresetsPath;
+		if (!GetPresetFilePath(InCategory, false, PresetsPath))
+		{
+			UE_LOG(LogAvaEaseCurveSubsystem, Warning, TEXT("EaseCurveToolJson::RemovePreset() - Could not find file for preset category: %s"), *InCategory);
+			return false;
+		}
+
+		TSharedPtr<FJsonObject> RootObject;
+		if (!LoadCurvePresetsJson(PresetsPath, RootObject) || !RootObject.IsValid())
+		{
+			UE_LOG(LogAvaEaseCurveSubsystem, Warning, TEXT("EaseCurveToolJson::RemovePreset() - Unable to load Json file: %s"), *PresetsPath);
+			return false;
+		}
+
+		RootObject->Values.Remove(InPreset);
+
+		if (!SaveCurvePresetsJson(PresetsPath, RootObject.ToSharedRef()))
+		{
+			UE_LOG(LogAvaEaseCurveSubsystem, Warning, TEXT("EaseCurveToolJson::RemovePreset() - Unable to save Json file: %s"), *PresetsPath);
+			return false;
+		}
+
+		return true;
+	}
+
+	bool RenamePreset(const FString& InCategory, const FString& InPreset, const FString& InNewPreset)
+	{
+		FString PresetsPath;
+		if (!GetPresetFilePath(InCategory, false, PresetsPath))
+		{
+			UE_LOG(LogAvaEaseCurveSubsystem, Warning, TEXT("EaseCurveToolJson::RemovePreset() - Could not find file for preset category: %s"), *InCategory);
+			return false;
+		}
+
+		TSharedPtr<FJsonObject> RootObject;
+		if (!LoadCurvePresetsJson(PresetsPath, RootObject) || !RootObject.IsValid())
+		{
+			UE_LOG(LogAvaEaseCurveSubsystem, Warning, TEXT("EaseCurveToolJson::RemovePreset() - Unable to load Json file: %s"), *PresetsPath);
+			return false;
+		}
+
+		if (const TSharedPtr<FJsonValue>* Value = RootObject->Values.Find(InPreset))
+		{
+			const TSharedPtr<FJsonValue>& ValueRef = *Value;
+			RootObject->Values.Add(InNewPreset, MakeShared<FJsonValueString>(ValueRef->AsString()));
+			RootObject->Values.Remove(InPreset);
+		}
+
+		if (!SaveCurvePresetsJson(PresetsPath, RootObject.ToSharedRef()))
+		{
+			UE_LOG(LogAvaEaseCurveSubsystem, Warning, TEXT("EaseCurveToolJson::RemovePreset() - Unable to save Json file: %s"), *PresetsPath);
+			return false;
+		}
+
+		return true;
+	}
+}
+
+UAvaEaseCurveSubsystem& UAvaEaseCurveSubsystem::Get()
+{
+	check(GEditor);
+	return *GEditor->GetEditorSubsystem<UAvaEaseCurveSubsystem>();
+}
+
+void UAvaEaseCurveSubsystem::Initialize(FSubsystemCollectionBase& InCollection)
+{
+	Super::Initialize(InCollection);
+
+	ReloadPresetsFromJson();
+}
+
+void UAvaEaseCurveSubsystem::ExploreJsonPresetsFolder()
+{
+	const FString EaseCurvePresetsPath = FPaths::ConvertRelativePathToFull(ProjectPresetPath());
+	FPlatformProcess::ExploreFolder(*EaseCurvePresetsPath);
+}
+
+FString UAvaEaseCurveSubsystem::ProjectPresetPath()
+{
+	return FPaths::ProjectConfigDir() / TEXT("EaseCurves");
+}
+
+FString UAvaEaseCurveSubsystem::PluginPresetPath()
+{
+	if (const TSharedPtr<IPlugin> Plugin = IPluginManager::Get().FindPlugin(UE_PLUGIN_NAME))
+	{
+		return Plugin->GetBaseDir() / TEXT("Config/EaseCurves");
+	}
+	return FString();
+}
+
+void UAvaEaseCurveSubsystem::ReloadPresetsFromJson()
+{
+	TArray<FString> JsonFiles;
+	IFileManager::Get().FindFilesRecursive(JsonFiles, *ProjectPresetPath(), TEXT("*.json"), true, false);
+
+	// Load all Json files found as their own separate category. Category Name = File Name
+	Presets.Empty();
+	for (const FString& File : JsonFiles)
+	{
+		TSharedPtr<FJsonObject> RootObject;
+		if (!UE::EaseCurveTool::Json::Private::LoadCurvePresetsJson(File, RootObject) || !RootObject.IsValid())
+		{
+			UE_LOG(LogAvaEaseCurveSubsystem, Warning, TEXT("UAvaEaseCurveSubsystem::LoadCurvePresetsJson() - Unable to load Json file [%s]"), *File);
+			continue;
+		}
+
+		const FString CategoryName = FPaths::GetBaseFilename(File);
+
+		TArray<TSharedPtr<FAvaEaseCurvePreset>>& PresetCategoryRef = Presets.FindOrAdd(CategoryName);
+
+		for (TMap<FString, TSharedPtr<FJsonValue>>::TConstIterator JsonValueIter = RootObject->Values.CreateConstIterator(); JsonValueIter; ++JsonValueIter)
+		{
+			const FString& PresetName = JsonValueIter.Key();
+			const FString TangentsString = (*JsonValueIter).Value->AsString();
+
+			FAvaEaseCurveTangents Tangents;
+			if (!FAvaEaseCurveTangents::FromString(TangentsString, Tangents))
+			{
+				continue;
+			}
+
+			const bool bExists = PresetCategoryRef.ContainsByPredicate([&PresetName](const TSharedPtr<FAvaEaseCurvePreset>& InPreset)
+				{
+					return InPreset->Name.Equals(PresetName);
+				});
+			if (!bExists)
+			{
+				PresetCategoryRef.Add(MakeShared<FAvaEaseCurvePreset>(PresetName, CategoryName, MoveTemp(Tangents)));
+			}
+		}
+
+		PresetCategoryRef.Sort([](const TSharedPtr<FAvaEaseCurvePreset>& InPresetA, const TSharedPtr<FAvaEaseCurvePreset>& InPresetB)
+			{
+				return *InPresetA < *InPresetB;
+			});
+	}
+}
+
+TArray<FString> UAvaEaseCurveSubsystem::GetEaseCurveCategories() const
+{
+	TArray<FString> OutCategories;
+	Presets.GenerateKeyArray(OutCategories);
+	return OutCategories;
+}
+
+TArray<TSharedPtr<FAvaEaseCurvePreset>> UAvaEaseCurveSubsystem::GetEaseCurvePresets() const
+{
+	TArray<TSharedPtr<FAvaEaseCurvePreset>> OutPresets;
+
+	for (const TPair<FString, TArray<TSharedPtr<FAvaEaseCurvePreset>>>& Preset : Presets)
+	{
+		OutPresets.Append(Preset.Value);
+	}
+
+	return OutPresets;
+}
+
+TArray<TSharedPtr<FAvaEaseCurvePreset>> UAvaEaseCurveSubsystem::GetEaseCurvePresets(const FString& InCategory) const
+{
+	if (!Presets.Contains(InCategory))
+	{
+		return {};
+	}
+	return Presets[InCategory];
+}
+
+bool UAvaEaseCurveSubsystem::DoesPresetExist(const FString& InName) const
+{
+	for (const TPair<FString, TArray<TSharedPtr<FAvaEaseCurvePreset>>>& Preset : Presets)
+	{
+		const bool bExists = Preset.Value.ContainsByPredicate([&InName](const TSharedPtr<FAvaEaseCurvePreset>& InPreset)
+			{
+				return InPreset->Name.Equals(InName);
+			});
+		if (bExists)
+		{
+			return true;
+		}
+	}
+	return false;
+}
+
+bool UAvaEaseCurveSubsystem::DoesPresetExist(const FString& InCategory, const FString& InName) const
+{
+	if (!Presets.Contains(InCategory))
+	{
+		return false;
+	}
+
+	return Presets[InCategory].ContainsByPredicate([&InName](const TSharedPtr<FAvaEaseCurvePreset>& InPreset)
+		{
+			return InPreset->Name.Equals(InName);
+		});
+}
+
+TSharedPtr<FAvaEaseCurvePreset> UAvaEaseCurveSubsystem::AddPreset(FAvaEaseCurvePreset InPreset)
+{
+	if (DoesPresetExist(InPreset.Category, InPreset.Name))
+	{
+		UE_LOG(LogAvaEaseCurveSubsystem, Warning, TEXT("UAvaEaseCurveSubsystem::AddPreset() - Preset already exists: %s"), *InPreset.Name);
+		return nullptr;
+	}
+
+	if (!UE::EaseCurveTool::Json::Private::AddPreset(InPreset.Category, InPreset.Name, InPreset.Tangents))
+	{
+		UE_LOG(LogAvaEaseCurveSubsystem, Warning, TEXT("UAvaEaseCurveSubsystem::AddPreset() - Unable to add preset %s to category %s Json file"), *InPreset.Name, *InPreset.Category);
+		return nullptr;
+	}
+
+	TSharedPtr<FAvaEaseCurvePreset> NewPresetPtr = MakeShared<FAvaEaseCurvePreset>(InPreset);
+	Presets.Add(InPreset.Category, { NewPresetPtr });
+
+	return NewPresetPtr;
+}
+
+TSharedPtr<FAvaEaseCurvePreset> UAvaEaseCurveSubsystem::AddPreset(const FString& InName, const FAvaEaseCurveTangents& InTangents)
+{
+	const FString CategoryName = GetDefault<UAvaEaseCurveToolSettings>()->GetNewPresetCategory();
+	if (CategoryName.IsEmpty())
+	{
+		return nullptr;
+	}
+
+	return AddPreset(FAvaEaseCurvePreset(InName, CategoryName, InTangents));
+}
+
+bool UAvaEaseCurveSubsystem::RemovePreset(const FAvaEaseCurvePreset& InPreset)
+{
+	if (!DoesPresetExist(InPreset.Name))
+	{
+		UE_LOG(LogAvaEaseCurveSubsystem, Warning, TEXT("UAvaEaseCurveSubsystem::RemovePreset() - Preset already exists: %s"), *InPreset.Name);
+		return false;
+	}
+
+	if (!UE::EaseCurveTool::Json::Private::RemovePreset(InPreset.Category, InPreset.Name))
+	{
+		UE_LOG(LogAvaEaseCurveSubsystem, Warning, TEXT("UAvaEaseCurveSubsystem::RemovePreset() - Failed to remove preset %s from category %s Json file"), *InPreset.Name, *InPreset.Category);
+		return false;
+	}
+
+	if (Presets.Contains(InPreset.Category))
+	{
+		TSharedPtr<FAvaEaseCurvePreset> PresetToRemove;
+		for (const TSharedPtr<FAvaEaseCurvePreset>& PresetPtr : Presets[InPreset.Category])
+		{
+			if (PresetPtr->Name.Equals(InPreset.Name, ESearchCase::IgnoreCase))
+			{
+				PresetToRemove = PresetPtr;
+				break;
+			}
+		}
+		if (PresetToRemove.IsValid())
+		{
+			Presets[InPreset.Category].Remove(PresetToRemove);
+		}
+	}
+
+	return true;
+}
+
+bool UAvaEaseCurveSubsystem::RemovePreset(const FString& InCategory, const FString& InName)
+{
+	return RemovePreset(FAvaEaseCurvePreset(InName, InCategory, FAvaEaseCurveTangents()));
+}
+
+bool UAvaEaseCurveSubsystem::ChangePresetCategory(const TSharedPtr<FAvaEaseCurvePreset>& InPreset, const FString& InNewCategory) const
+{
+	if (!Presets.Contains(InPreset->Category) || !Presets.Contains(InNewCategory))
+	{
+		return false;
+	}
+
+	// Verify old category exists and new category does not
+	if (!DoesPresetExist(InPreset->Category, InPreset->Name) || DoesPresetExist(InNewCategory, InPreset->Name))
+	{
+		return false;
+	}
+
+	if (!UAvaEaseCurveSubsystem::Get().RemovePreset(InPreset->Category, InPreset->Name))
+	{
+		return false;
+	}
+
+	InPreset->Category = InNewCategory;
+
+	return UAvaEaseCurveSubsystem::Get().AddPreset(*InPreset).IsValid();
+}
+
+TSharedPtr<FAvaEaseCurvePreset> UAvaEaseCurveSubsystem::FindPreset(const FString& InName)
+{
+	for (const TPair<FString, TArray<TSharedPtr<FAvaEaseCurvePreset>>>& Preset : Presets)
+	{
+		const TSharedPtr<FAvaEaseCurvePreset>* PresetPtr = Preset.Value.FindByPredicate([&InName](const TSharedPtr<FAvaEaseCurvePreset>& InPreset)
+			{
+				return InPreset->Name == InName;
+			});
+		if (PresetPtr)
+		{
+			return *PresetPtr;
+		}
+	}
+	return nullptr;
+}
+
+TSharedPtr<FAvaEaseCurvePreset> UAvaEaseCurveSubsystem::FindPresetByTangents(const FAvaEaseCurveTangents& InTangents, const double InErrorTolerance)
+{
+	for (const TPair<FString, TArray<TSharedPtr<FAvaEaseCurvePreset>>>& Preset : Presets)
+	{
+		const TSharedPtr<FAvaEaseCurvePreset>* PresetPtr = Preset.Value.FindByPredicate([&InTangents, InErrorTolerance]
+			(const TSharedPtr<FAvaEaseCurvePreset>& InPreset)
+			{
+				return InPreset->Tangents.IsNearlyEqual(InTangents, InErrorTolerance);
+			});
+		if (PresetPtr)
+		{
+			return *PresetPtr;
+		}
+	}
+	return nullptr;
+}
+
+bool UAvaEaseCurveSubsystem::DoesPresetCategoryExist(const FString& InCategory)
+{
+	return UE::EaseCurveTool::Json::Private::DoesPresetCategoryExist(InCategory);
+}
+
+bool UAvaEaseCurveSubsystem::RenamePresetCategory(const FString& InCategory, const FString& InNewCategoryName)
+{
+	if (!UE::EaseCurveTool::Json::Private::RenamePresetCategory(InCategory, InNewCategoryName))
+	{
+		return false;
+	}
+
+	if (Presets.Remove(InCategory) > 0)
+	{
+		Presets.Add(InNewCategoryName, Presets[InCategory]);
+	}
+
+	return true;
+}
+
+bool UAvaEaseCurveSubsystem::AddNewPresetCategory()
+{
+	const FString NewCategoryName = UE::EaseCurveTool::Json::Private::AddNewPresetCategory();
+
+	if (NewCategoryName.IsEmpty())
+	{
+		return false;
+	}
+
+	Presets.Add(NewCategoryName);
+
+	return true;
+}
+
+bool UAvaEaseCurveSubsystem::RemovePresetCategory(const FString& InCategory)
+{
+	if (!UE::EaseCurveTool::Json::Private::RemovePresetCategory(InCategory))
+	{
+		return false;
+	}
+
+	return Presets.Remove(InCategory) > 0;
+}
+
+bool UAvaEaseCurveSubsystem::RenamePreset(const FString& InCategory, const FString& InPreset, const FString& InNewPresetName)
+{
+	if (!UE::EaseCurveTool::Json::Private::RenamePreset(InCategory, InPreset, InNewPresetName))
+	{
+		UE_LOG(LogAvaEaseCurveSubsystem, Warning, TEXT("UAvaEaseCurveSubsystem::RenamePreset() - Failed to rename preset %s from category %s Json file"), *InPreset, *InCategory);
+		return false;
+	}
+
+	if (Presets.Contains(InCategory))
+	{
+		TSharedPtr<FAvaEaseCurvePreset> PresetToRename;
+		for (const TSharedPtr<FAvaEaseCurvePreset>& PresetPtr : Presets[InCategory])
+		{
+			if (PresetPtr->Name.Equals(InPreset, ESearchCase::IgnoreCase))
+			{
+				PresetToRename = PresetPtr;
+				break;
+			}
+		}
+		if (PresetToRename.IsValid())
+		{
+			Presets[InCategory].Remove(PresetToRename);
+
+			PresetToRename->Name = InNewPresetName;
+
+			Presets[InCategory].Add(PresetToRename);
+		}
+	}
+
+	return true;
+}
+
+void UAvaEaseCurveSubsystem::ResetToDefaultPresets(const bool bInOnlyIfNoProjectPresets)
+{
+	IPlatformFile& PlatformFile = FPlatformFileManager::Get().GetPlatformFile();
+
+	const FString SourceProjectPresetPath = FPaths::ConvertRelativePathToFull(UAvaEaseCurveSubsystem::ProjectPresetPath());
+
+	if (!bInOnlyIfNoProjectPresets)
+	{
+		PlatformFile.DeleteDirectoryRecursively(*SourceProjectPresetPath);
+		PlatformFile.CreateDirectory(*SourceProjectPresetPath);
+	}
+
+	TArray<FString> JsonFiles;
+	IFileManager::Get().FindFilesRecursive(JsonFiles, *SourceProjectPresetPath, TEXT("*.json"), true, false);
+
+	if (JsonFiles.IsEmpty())
+	{
+		const FString DestPluginPresetPath = FPaths::ConvertRelativePathToFull(UAvaEaseCurveSubsystem::PluginPresetPath());
+
+		if (PlatformFile.CopyDirectoryTree(*SourceProjectPresetPath, *DestPluginPresetPath, false))
+		{
+			UE_LOG(LogAvaEaseCurveSubsystem, Warning, TEXT("Motion Design ease curve tool project presets are empty. "
+				"Copied default presets to [Project]/Config/EaseCurves"));
+
+			ReloadPresetsFromJson();
+		}
+	}
+}
+
+const TMap<FString, TArray<FString>>& UAvaEaseCurveSubsystem::GetDefaultCategoryPresetNames()
+{
+	static const TArray<FString> Names = {
+		TEXT("Sine"), 
+		TEXT("Cubic"), 
+		TEXT("Quintic"), 
+		TEXT("Circular"), 
+		TEXT("Quadratic"), 
+		TEXT("Quartic"), 
+		TEXT("Exponential") 
+	};
+
+	static const TMap<FString, TArray<FString>> CategoryPresetNames = {
+		{ TEXT("Ease In Out"), Names },
+		{ TEXT("Ease In"), Names },
+		{ TEXT("Ease Out"), Names }
+	};
+
+	return CategoryPresetNames;
+}
diff --git a/Engine/Plugins/VirtualProduction/Avalanche/Source/AvalancheSequencer/Private/EaseCurveTool/AvaEaseCurveSubsystem.h b/Engine/Plugins/VirtualProduction/Avalanche/Source/AvalancheSequencer/Private/EaseCurveTool/AvaEaseCurveSubsystem.h
new file mode 100644
--- /dev/null
+++ b/Engine/Plugins/VirtualProduction/Avalanche/Source/AvalancheSequencer/Private/EaseCurveTool/AvaEaseCurveSubsystem.h
@@ -0,0 +1,71 @@
+// Copyright Epic Games, Inc. All Rights Reserved.
+
+#pragma once
+
+#include "AvaEaseCurveTangents.h"
+#include "Containers/Array.h"
+#include "Containers/Map.h"
+#include "Containers/StringFwd.h"
+#include "EditorSubsystem.h"
+#include "Templates/SharedPointer.h"
+#include "AvaEaseCurveSubsystem.generated.h"
+
+struct FAvaEaseCurvePreset;
+
+UCLASS()
+class UAvaEaseCurveSubsystem : public UEditorSubsystem
+{
+	GENERATED_BODY()
+
+public:
+	static UAvaEaseCurveSubsystem& Get();
+
+	static void ExploreJsonPresetsFolder();
+
+	static FString ProjectPresetPath();
+	static FString PluginPresetPath();
+
+	/**
+	 * Mainly used to compare clicked presets names for engine analytics, thus returns static category
+	 * and preset names of the default presets instead of looking on disk for the actual default names.
+	 */
+	static const TMap<FString, TArray<FString>>& GetDefaultCategoryPresetNames();
+
+	void ReloadPresetsFromJson();
+
+	TArray<FString> GetEaseCurveCategories() const;
+
+	TArray<TSharedPtr<FAvaEaseCurvePreset>> GetEaseCurvePresets() const;
+	TArray<TSharedPtr<FAvaEaseCurvePreset>> GetEaseCurvePresets(const FString& InCategory) const;
+
+	bool DoesPresetExist(const FString& InName) const;
+	bool DoesPresetExist(const FString& InCategory, const FString& InName) const;
+
+	TSharedPtr<FAvaEaseCurvePreset> AddPreset(FAvaEaseCurvePreset InPreset);
+	TSharedPtr<FAvaEaseCurvePreset> AddPreset(const FString& InName, const FAvaEaseCurveTangents& InTangents);
+
+	bool RemovePreset(const FAvaEaseCurvePreset& InPreset);
+	bool RemovePreset(const FString& InCategory, const FString& InName);
+
+	bool ChangePresetCategory(const TSharedPtr<FAvaEaseCurvePreset>& InPreset, const FString& InNewCategory) const;
+
+	TSharedPtr<FAvaEaseCurvePreset> FindPreset(const FString& InName);
+	TSharedPtr<FAvaEaseCurvePreset> FindPresetByTangents(const FAvaEaseCurveTangents& InTangents, const double InErrorTolerance = 0.01);
+
+	bool DoesPresetCategoryExist(const FString& InCategory);
+	bool RenamePresetCategory(const FString& InCategory, const FString& InNewCategoryName);
+	bool AddNewPresetCategory();
+	bool RemovePresetCategory(const FString& InCategory);
+
+	bool RenamePreset(const FString& InCategory, const FString& InPreset, const FString& InNewPresetName);
+
+	void ResetToDefaultPresets(const bool bInOnlyIfNoProjectPresets);
+
+	//~ Begin USubsystem
+	virtual void Initialize(FSubsystemCollectionBase& InCollection) override;
+	//~ End USubsystem
+
+protected:
+	/** Cached presets from Json files. Each Json file represents a category. */
+	TSortedMap<FString, TArray<TSharedPtr<FAvaEaseCurvePreset>>> Presets;
+};
diff --git a/Engine/Plugins/VirtualProduction/Avalanche/Source/AvalancheSequencer/Private/EaseCurveTool/AvaEaseCurveTangents.cpp b/Engine/Plugins/VirtualProduction/Avalanche/Source/AvalancheSequencer/Private/EaseCurveTool/AvaEaseCurveTangents.cpp
new file mode 100644
--- /dev/null
+++ b/Engine/Plugins/VirtualProduction/Avalanche/Source/AvalancheSequencer/Private/EaseCurveTool/AvaEaseCurveTangents.cpp
@@ -0,0 +1,335 @@
+// Copyright Epic Games, Inc. All Rights Reserved.
+
+#include "AvaEaseCurveTangents.h"
+#include "EaseCurveTool/Widgets/SAvaEaseCurveEditor.h"
+#include "Internationalization/Text.h"
+#include "Misc/DefaultValueHelper.h"
+#include "Misc/FrameRate.h"
+#include "Misc/QualifiedFrameTime.h"
+
+#define LOCTEXT_NAMESPACE "AvaEaseCurveTangents"
+
+bool FAvaEaseCurveTangents::FromString(const FString& InString, FAvaEaseCurveTangents& OutTangents)
+{
+	const FString TangentsString = InString.Replace(TEXT(" "), TEXT(""));
+
+	TArray<FString> ValueStrings;
+	if (TangentsString.ParseIntoArray(ValueStrings, TEXT(",")) != 4)
+	{
+		return false;
+	}
+
+	// Convert strings to doubles
+	TArray<double> Values = { 0.0, 0.0, 0.0, 0.0 };
+	for (int32 Index = 0; Index < 4; ++Index)
+	{
+		if (!FDefaultValueHelper::ParseDouble(ValueStrings[Index], Values[Index]))
+		{
+			return false;
+		}
+	}
+
+	// Convert four cubic bezier points to two points/tangents
+	OutTangents.FromCubicBezier(Values);
+
+	return true;
+}
+
+bool FAvaEaseCurveTangents::CanParseString(const FString& InString)
+{
+	FAvaEaseCurveTangents Tangents;
+	return FromString(InString, Tangents);
+}
+
+FNumberFormattingOptions FAvaEaseCurveTangents::DefaultNumberFormattingOptions()
+{
+	FNumberFormattingOptions NumberFormat;
+	NumberFormat.MinimumIntegralDigits = 1;
+	NumberFormat.MinimumFractionalDigits = 2;
+	NumberFormat.MaximumFractionalDigits = 2;
+	NumberFormat.UseGrouping = false;
+	return NumberFormat;
+}
+
+FAvaEaseCurveTangents::FAvaEaseCurveTangents(const FString& InTangentsString)
+{
+	FAvaEaseCurveTangents::FromString(InTangentsString, *this);
+}
+
+bool FAvaEaseCurveTangents::IsNearlyEqual(const FAvaEaseCurveTangents& InOther, const double InErrorTolerance) const
+{
+	return FMath::IsNearlyEqual(Start, InOther.Start, InErrorTolerance)
+		&& FMath::IsNearlyEqual(StartWeight, InOther.StartWeight, InErrorTolerance)
+		&& FMath::IsNearlyEqual(End, InOther.End, InErrorTolerance)
+		&& FMath::IsNearlyEqual(EndWeight, InOther.EndWeight, InErrorTolerance);
+}
+
+FText FAvaEaseCurveTangents::ToDisplayText() const
+{
+	const FNumberFormattingOptions NumberFormat = DefaultNumberFormattingOptions();
+	const FText StartTangentText = FText::AsNumber(Start, &NumberFormat);
+	const FText StartTangentWeightText = FText::AsNumber(StartWeight, &NumberFormat);
+	const FText EndTangentText = FText::AsNumber(End, &NumberFormat);
+	const FText EndTangentWeightText = FText::AsNumber(EndWeight, &NumberFormat);
+	return FText::Format(LOCTEXT("TangentText", "{0}, {1} - {2}, {3}")
+		, { StartTangentText, StartTangentWeightText, EndTangentText, EndTangentWeightText });
+}
+
+FString FAvaEaseCurveTangents::ToDisplayString() const
+{
+	return ToDisplayText().ToString();
+}
+
+FString FAvaEaseCurveTangents::ToJson() const
+{
+	// Convert two points/tangents to four cubic bezier points 
+	TArray<double> CubicBezierPoints;
+	ToCubicBezier(CubicBezierPoints);
+
+	const FNumberFormattingOptions NumberFormat = DefaultNumberFormattingOptions();
+	const FText PointAText = FText::AsNumber(CubicBezierPoints[0], &NumberFormat);
+	const FText PointBText = FText::AsNumber(CubicBezierPoints[1], &NumberFormat);
+	const FText PointCText = FText::AsNumber(CubicBezierPoints[2], &NumberFormat);
+	const FText PointDText = FText::AsNumber(CubicBezierPoints[3], &NumberFormat);
+
+	return FString::Printf(TEXT("%s, %s, %s, %s")
+		, *PointAText.ToString(), *PointBText.ToString(), *PointCText.ToString(), *PointDText.ToString());
+}
+
+FText FAvaEaseCurveTangents::GetStartTangentText() const
+{
+	const FNumberFormattingOptions NumberFormat = DefaultNumberFormattingOptions();
+	const FText StartTangentText = FText::AsNumber(Start, &NumberFormat);
+	const FText StartTangentWeightText = FText::AsNumber(StartWeight, &NumberFormat);
+	return FText::Format(LOCTEXT("StartTangentText", "Start: {0}, {1}"), { StartTangentText, StartTangentWeightText });
+}
+
+FText FAvaEaseCurveTangents::GetEndTangentText() const
+{
+	const FNumberFormattingOptions NumberFormat = DefaultNumberFormattingOptions();
+	const FText EndTangentText = FText::AsNumber(End, &NumberFormat);
+	const FText EndTangentWeightText = FText::AsNumber(EndWeight, &NumberFormat);
+	return FText::Format(LOCTEXT("EndTangentText", "End: {0}, {1}"), { EndTangentText, EndTangentWeightText });
+}
+
+FText FAvaEaseCurveTangents::GetCubicBezierText() const
+{
+	return FText::FromString(ToJson());
+}
+
+bool FAvaEaseCurveTangents::FromCubicBezier(const TArray<double>& InPoints)
+{
+	if (InPoints.Num() != 4)
+	{
+		return false;
+	}
+
+	const FVector2d StartDir(InPoints[0], InPoints[1]);
+	const FVector2d EndDir(1.0 - InPoints[2], 1.0 - InPoints[3]);
+
+	Start = SAvaEaseCurveEditor::CalcTangent(StartDir);
+	End = SAvaEaseCurveEditor::CalcTangent(EndDir);
+
+	StartWeight = StartDir.Size();
+	EndWeight = EndDir.Size();
+
+	return true;
+}
+
+bool FAvaEaseCurveTangents::ToCubicBezier(TArray<double>& OutPoints) const
+{
+	const FVector2d StartDir = SAvaEaseCurveEditor::CalcTangentDir(Start) * StartWeight;
+	const FVector2d EndDir = SAvaEaseCurveEditor::CalcTangentDir(End) * EndWeight;
+
+	OutPoints.Reset(4);
+	OutPoints.Add(StartDir.X);
+	OutPoints.Add(StartDir.Y);
+	OutPoints.Add(1.0 - EndDir.X);
+	OutPoints.Add(1.0 - EndDir.Y);
+
+	return true;
+}
+
+void FAvaEaseCurveTangents::Normalize(const FFrameNumber& InFrameNumber, const double InValue
+	, const FFrameNumber& InNextFrameNumber, const double InNextValue
+	, const FFrameRate& InDisplayRate, const FFrameRate& InTickResolution)
+{
+	// Convert frame time from tick resolution to display rate
+	const FFrameTime FrameTimeDifference = InNextFrameNumber - InFrameNumber;
+	const FQualifiedFrameTime QualifiedFrameTime(FrameTimeDifference, InTickResolution);
+	const FFrameTime ConvertedFrameTime = QualifiedFrameTime.ConvertTo(InDisplayRate);
+
+	// Create time/value range scale factor
+	FVector2d TimeValueRange;
+	TimeValueRange.X = InDisplayRate.AsSeconds(ConvertedFrameTime);
+	TimeValueRange.Y = FMath::Abs(InNextValue - InValue);
+
+	const double ScaleFactor = InTickResolution.AsDecimal();
+
+	// Convert tangent angles to grid coordinates
+	FVector2d StartDir = SAvaEaseCurveEditor::CalcTangentDir(Start * ScaleFactor);
+	FVector2d EndDir = SAvaEaseCurveEditor::CalcTangentDir(End * ScaleFactor);
+
+	StartDir *= StartWeight;
+	EndDir *= EndWeight;
+
+	// Scale down tangent grid coordinates to normalized range of InTimeValueRange
+	auto SafeDivide = [](const double InNumerator, const double InDenominator)
+		{
+			return FMath::IsNearlyZero(InDenominator) ? InNumerator : (InNumerator / InDenominator);
+		};
+	StartDir.X = SafeDivide(StartDir.X, TimeValueRange.X);
+	StartDir.Y = SafeDivide(StartDir.Y, TimeValueRange.Y);
+	EndDir.X = SafeDivide(EndDir.X, TimeValueRange.X);
+	EndDir.Y = SafeDivide(EndDir.Y, TimeValueRange.Y);
+
+	// Set new weights and tangents
+	StartWeight = StartDir.Size();
+	EndWeight = EndDir.Size();
+
+	StartDir.Normalize();
+	EndDir.Normalize();
+
+	Start = SAvaEaseCurveEditor::CalcTangent(StartDir);
+	End = SAvaEaseCurveEditor::CalcTangent(EndDir);
+}
+
+void FAvaEaseCurveTangents::ScaleUp(const FFrameNumber& InFrameNumber, const double InValue
+	, const FFrameNumber& InNextFrameNumber, const double InNextValue
+	, const FFrameRate& InDisplayRate, const FFrameRate& InTickResolution)
+{
+	// Convert frame time from tick resolution to display rate
+	const FFrameTime FrameTimeDifference = InNextFrameNumber - InFrameNumber;
+	const FQualifiedFrameTime QualifiedFrameTime(FrameTimeDifference, InTickResolution);
+	const FFrameTime ConvertedFrameTime = QualifiedFrameTime.ConvertTo(InDisplayRate);
+
+	// Create time/value range scale factor
+	FVector2d TimeValueRange;
+	TimeValueRange.X = InDisplayRate.AsSeconds(ConvertedFrameTime);
+	TimeValueRange.Y = FMath::Abs(InNextValue - InValue);
+
+	FVector2d StartDir = SAvaEaseCurveEditor::CalcTangentDir(Start);
+	FVector2d EndDir = SAvaEaseCurveEditor::CalcTangentDir(End);
+
+	StartDir *= StartWeight;
+	EndDir *= EndWeight;
+
+	StartDir *= TimeValueRange;
+	EndDir *= TimeValueRange;
+
+	// Set new weights
+	StartWeight = StartDir.Size();
+	EndWeight = EndDir.Size();
+
+	// Normalize vector and set new tangents
+	FVector2d NormalizedStartDir = StartDir;
+	FVector2d NormalizedEndDir = EndDir;
+
+	const double ScaleFactor = InTickResolution.AsDecimal();
+	NormalizedStartDir.X *= ScaleFactor;
+	NormalizedEndDir.X *= ScaleFactor;
+
+	NormalizedStartDir.Normalize();
+	NormalizedEndDir.Normalize();
+
+	Start = SAvaEaseCurveEditor::CalcTangent(NormalizedStartDir);
+	End = SAvaEaseCurveEditor::CalcTangent(NormalizedEndDir);
+}
+
+FAvaEaseCurveTangents FAvaEaseCurveTangents::Average(const TArray<FAvaEaseCurveTangents>& InTangentArray)
+{
+	const int32 KeySetTangentsCount = InTangentArray.Num();
+	if (KeySetTangentsCount == 0)
+	{
+		return FAvaEaseCurveTangents();
+	}
+
+	double TotalStartTangent = 0.0;
+	double TotalStartWeight = 0.0;
+	double TotalEndTangent = 0.0;
+	double TotalEndWeight = 0.0;
+
+	for (const FAvaEaseCurveTangents& Tangent : InTangentArray)
+	{
+		TotalStartTangent += Tangent.Start;
+		TotalStartWeight += Tangent.StartWeight;
+		TotalEndTangent += Tangent.End;
+		TotalEndWeight += Tangent.EndWeight;
+	}
+
+	FAvaEaseCurveTangents AverageTangents;
+	AverageTangents.Start = TotalStartTangent / KeySetTangentsCount;
+	AverageTangents.StartWeight = TotalStartWeight / KeySetTangentsCount;
+	AverageTangents.End = TotalEndTangent / KeySetTangentsCount;
+	AverageTangents.EndWeight = TotalEndWeight / KeySetTangentsCount;
+
+	return AverageTangents;
+}
+
+double FAvaEaseCurveTangents::CalculateCurveLength(const int32 SampleCount) const
+{
+	TArray<double> CubicBezierPoints;
+	ToCubicBezier(CubicBezierPoints);
+
+	const TArray<FVector2d> Points = { 
+			FVector2d(0.0, 0.0), 
+			FVector2d(CubicBezierPoints[0], CubicBezierPoints[1]), 
+			FVector2d(CubicBezierPoints[2], CubicBezierPoints[3]), 
+			FVector2d(1.0, 1.0)
+		};
+
+	auto EvaluatePoint = [](const double InTime, double InStart, double InControlA, double InControlB, double InEnd) -> double
+		{
+			return InStart * (1.0 - InTime) * (1.0 - InTime) * (1.0 - InTime)
+				+ 3.0 * InControlA * (1.0 - InTime) * (1.0 - InTime) * InTime
+				+ 3.0 * InControlB * (1.0 - InTime) * InTime * InTime
+				+ InEnd * InTime * InTime * InTime;
+		};
+
+	// Sample points along the curve and use those points to calculate the length of the curve
+	FVector2d SamplePoint;
+	FVector2d PreviousSamplePoint = FVector2d::ZeroVector;
+	double Length = 0.0;
+
+	for (int32 Index = 0; Index <= SampleCount; Index++)
+	{
+		const double Time = (double)Index / (double)SampleCount;
+				
+		SamplePoint.X = EvaluatePoint(Time, Points[0].X, Points[1].X, Points[2].X, Points[3].X);
+		SamplePoint.Y = EvaluatePoint(Time, Points[0].Y, Points[1].Y, Points[2].Y, Points[3].Y);
+
+		if (Index > 0)
+		{
+			const double DifferenceX = SamplePoint.X - PreviousSamplePoint.X;
+			const double DifferenceY = SamplePoint.Y - PreviousSamplePoint.Y;
+			Length += FMath::Sqrt(DifferenceX * DifferenceX + DifferenceY * DifferenceY);
+		}
+
+		PreviousSamplePoint = SamplePoint;
+	}
+
+	return Length;
+}
+
+bool FAvaEaseCurveTangents::IsEaseCurveKey(const FRichCurveKey& InKey)
+{
+	return InKey.TangentMode == RCTM_Break
+		&& InKey.TangentWeightMode == RCTWM_WeightedBoth
+		&& InKey.InterpMode == RCIM_Cubic;
+}
+
+bool FAvaEaseCurveTangents::IsEaseCurveKey(const FMovieSceneDoubleValue& InValue)
+{
+	return InValue.TangentMode == RCTM_Break
+		&& InValue.Tangent.TangentWeightMode == RCTWM_WeightedBoth
+		&& InValue.InterpMode == RCIM_Cubic;
+}
+
+bool FAvaEaseCurveTangents::IsEaseCurveKey(const FMovieSceneFloatValue& InValue)
+{
+	return InValue.TangentMode == RCTM_Break
+		&& InValue.Tangent.TangentWeightMode == RCTWM_WeightedBoth
+		&& InValue.InterpMode == RCIM_Cubic;
+}
+
+#undef LOCTEXT_NAMESPACE
diff --git a/Engine/Plugins/VirtualProduction/Avalanche/Source/AvalancheSequencer/Private/EaseCurveTool/AvaEaseCurveTangents.h b/Engine/Plugins/VirtualProduction/Avalanche/Source/AvalancheSequencer/Private/EaseCurveTool/AvaEaseCurveTangents.h
new file mode 100644
--- /dev/null
+++ b/Engine/Plugins/VirtualProduction/Avalanche/Source/AvalancheSequencer/Private/EaseCurveTool/AvaEaseCurveTangents.h
@@ -0,0 +1,110 @@
+// Copyright Epic Games, Inc. All Rights Reserved.
+
+#pragma once
+
+#include "Curves/RichCurve.h"
+#include "Channels/MovieSceneDoubleChannel.h"
+#include "Channels/MovieSceneFloatChannel.h"
+#include "Math/UnrealMathUtility.h"
+#include "AvaEaseCurveTangents.generated.h"
+
+class FString;
+class FText;
+struct FFrameRate;
+
+/** Represents a curve constructed of two tangents and their weights, one for each end of the curve. */
+USTRUCT()
+struct FAvaEaseCurveTangents
+{
+	GENERATED_BODY()
+
+public:
+	static bool FromString(const FString& InString, FAvaEaseCurveTangents& OutTangents);
+	static bool CanParseString(const FString& InString);
+
+	static FNumberFormattingOptions DefaultNumberFormattingOptions();
+
+	static FAvaEaseCurveTangents Average(const TArray<FAvaEaseCurveTangents>& InTangentArray);
+
+	static bool IsEaseCurveKey(const FRichCurveKey& InKey);
+	static bool IsEaseCurveKey(const FMovieSceneDoubleValue& InValue);
+	static bool IsEaseCurveKey(const FMovieSceneFloatValue& InValue);
+
+	FAvaEaseCurveTangents() {}
+	FAvaEaseCurveTangents(const double InStart, const double InStartWeight, const double InEnd, const double InEndWeight)
+		: Start(InStart), StartWeight(InStartWeight)
+		, End(InEnd), EndWeight(InEndWeight)
+	{}
+	FAvaEaseCurveTangents(const FRichCurveKey& InRichCurveKey)
+		: Start(InRichCurveKey.LeaveTangent), StartWeight(InRichCurveKey.LeaveTangentWeight)
+		, End(InRichCurveKey.ArriveTangent), EndWeight(InRichCurveKey.ArriveTangentWeight)
+	{}
+	FAvaEaseCurveTangents(const FRichCurveKey& InStartRichCurveKey, const FRichCurveKey& InEndRichCurveKey)
+		: Start(InStartRichCurveKey.LeaveTangent), StartWeight(InStartRichCurveKey.LeaveTangentWeight)
+		, End(InEndRichCurveKey.ArriveTangent), EndWeight(InEndRichCurveKey.ArriveTangentWeight)
+	{}
+	FAvaEaseCurveTangents(const FMovieSceneDoubleValue& InMovieSceneDoubleValue)
+		: Start(InMovieSceneDoubleValue.Tangent.LeaveTangent), StartWeight(InMovieSceneDoubleValue.Tangent.LeaveTangentWeight)
+		, End(InMovieSceneDoubleValue.Tangent.ArriveTangent), EndWeight(InMovieSceneDoubleValue.Tangent.ArriveTangentWeight)
+	{}
+	FAvaEaseCurveTangents(const FMovieSceneDoubleValue& InStartMovieSceneDoubleValue, const FMovieSceneDoubleValue& InEndMovieSceneDoubleValue)
+		: Start(InStartMovieSceneDoubleValue.Tangent.LeaveTangent), StartWeight(InStartMovieSceneDoubleValue.Tangent.LeaveTangentWeight)
+		, End(InEndMovieSceneDoubleValue.Tangent.ArriveTangent), EndWeight(InEndMovieSceneDoubleValue.Tangent.ArriveTangentWeight)
+	{}
+	FAvaEaseCurveTangents(const FMovieSceneFloatValue& InMovieSceneFloatValue)
+		: Start(InMovieSceneFloatValue.Tangent.LeaveTangent), StartWeight(InMovieSceneFloatValue.Tangent.LeaveTangentWeight)
+		, End(InMovieSceneFloatValue.Tangent.ArriveTangent), EndWeight(InMovieSceneFloatValue.Tangent.ArriveTangentWeight)
+	{}
+	FAvaEaseCurveTangents(const FMovieSceneFloatValue& InStartMovieSceneFloatValue, const FMovieSceneFloatValue& InEndMovieSceneFloatValue)
+		: Start(InStartMovieSceneFloatValue.Tangent.LeaveTangent), StartWeight(InStartMovieSceneFloatValue.Tangent.LeaveTangentWeight)
+		, End(InEndMovieSceneFloatValue.Tangent.ArriveTangent), EndWeight(InEndMovieSceneFloatValue.Tangent.ArriveTangentWeight)
+	{}
+	/** Constructor from string consisting of cubic bezier points. Ex. "0.45, 0.34, 0.0, 1.00" */
+	explicit FAvaEaseCurveTangents(const FString& InTangentsString);
+
+	FORCEINLINE bool operator==(const FAvaEaseCurveTangents& InRhs) const
+	{
+		return Start == InRhs.Start && StartWeight == InRhs.StartWeight && End == InRhs.End && EndWeight == InRhs.EndWeight;
+	}
+	FORCEINLINE bool operator!=(const FAvaEaseCurveTangents& InRhs) const
+	{
+		return !(*this == InRhs);
+	}
+
+	bool IsNearlyEqual(const FAvaEaseCurveTangents& InOther, const double InErrorTolerance = UE_DOUBLE_SMALL_NUMBER) const;
+
+	FText ToDisplayText() const;
+	FString ToDisplayString() const;
+
+	FString ToJson() const;
+
+	FText GetStartTangentText() const;
+	FText GetEndTangentText() const;
+	FText GetCubicBezierText() const;
+
+	bool FromCubicBezier(const TArray<double>& InPoints);
+	bool ToCubicBezier(TArray<double>& OutPoints) const;
+
+	void Normalize(const FFrameNumber& InFrameNumber, const double InValue
+		, const FFrameNumber& InNextFrameNumber, const double InNextValue
+		, const FFrameRate& InDisplayRate, const FFrameRate& InTickResolution);
+
+	void ScaleUp(const FFrameNumber& InFrameNumber, const double InValue
+		, const FFrameNumber& InNextFrameNumber, const double InNextValue
+		, const FFrameRate& InDisplayRate, const FFrameRate& InTickResolution);
+
+	/** Calculates the length of this curve. The higher the sample count, the higher the calculation precision. */
+	double CalculateCurveLength(const int32 InSampleCount = 10) const;
+
+	UPROPERTY(Config, EditAnywhere, Category = "EaseCurveTangents")
+	double Start = 0.0;
+
+	UPROPERTY(Config, EditAnywhere, Category = "EaseCurveTangents")
+	double StartWeight = 0.0;
+
+	UPROPERTY(Config, EditAnywhere, Category = "EaseCurveTangents")
+	double End = 0.0;
+
+	UPROPERTY(Config, EditAnywhere, Category = "EaseCurveTangents")
+	double EndWeight = 0.0;
+};
