Motion Design: moved the following plugins from /Plugins/Experimental to /Plugins/VirtualProduction:

diff --git a/Engine/Plugins/VirtualProduction/Avalanche/Source/AvalancheSequencer/Private/EaseCurveTool/AvaEaseCurveTool.cpp b/Engine/Plugins/VirtualProduction/Avalanche/Source/AvalancheSequencer/Private/EaseCurveTool/AvaEaseCurveTool.cpp
new file mode 100644
--- /dev/null
+++ b/Engine/Plugins/VirtualProduction/Avalanche/Source/AvalancheSequencer/Private/EaseCurveTool/AvaEaseCurveTool.cpp
@@ -0,0 +1,705 @@
+// Copyright Epic Games, Inc. All Rights Reserved.
+
+#include "AvaEaseCurveTool.h"
+#include "AssetRegistry/AssetRegistryModule.h"
+#include "AssetToolsModule.h"
+#include "AvaEaseCurveSubsystem.h"
+#include "AvaEaseCurveToolCommands.h"
+#include "Dialogs/DlgPickAssetPath.h"
+#include "EaseCurveTool/AvaEaseCurveKeySelection.h"
+#include "EaseCurveTool/AvaEaseCurveToolSettings.h"
+#include "EaseCurveTool/Widgets/SAvaEaseCurveTool.h"
+#include "Editor.h"
+#include "EngineAnalytics.h"
+#include "Factories/CurveFactory.h"
+#include "Framework/Commands/GenericCommands.h"
+#include "Framework/Notifications/NotificationManager.h"
+#include "HAL/PlatformApplicationMisc.h"
+#include "ISequencer.h"
+#include "ISettingsModule.h"
+#include "Math/UnrealMathUtility.h"
+#include "Misc/MessageDialog.h"
+#include "ScopedTransaction.h"
+#include "Settings/AvaSequencerSettings.h"
+#include "Widgets/Notifications/SNotificationList.h"
+
+using namespace UE::Sequencer;
+
+#define LOCTEXT_NAMESPACE "AvaEaseCurveTool"
+
+FAvaEaseCurveTool::FAvaEaseCurveTool(const TSharedRef<ISequencer>& InSequencer)
+	: SequencerWeak(InSequencer)
+{
+	EaseCurve = NewObject<UAvaEaseCurve>(GetTransientPackage(), NAME_None, RF_Transient | RF_Transactional);
+
+	ToolSettings = GetMutableDefault<UAvaEaseCurveToolSettings>();
+
+	UpdateEaseCurveFromSequencerKeySelections();
+
+	InSequencer->GetSelectionChangedObjectGuids().AddRaw(this, &FAvaEaseCurveTool::OnSequencerSelectionChanged);
+}
+
+FAvaEaseCurveTool::~FAvaEaseCurveTool()
+{
+	if (const TSharedPtr<ISequencer> Sequencer = SequencerWeak.Pin())
+	{
+		Sequencer->GetSelectionChangedObjectGuids().RemoveAll(this);
+	}
+}
+
+void FAvaEaseCurveTool::OnSequencerSelectionChanged(TArray<FGuid> InObjectGuids)
+{
+	UpdateEaseCurveFromSequencerKeySelections();
+}
+
+void FAvaEaseCurveTool::AddReferencedObjects(FReferenceCollector& Collector)
+{
+	Collector.AddReferencedObject(EaseCurve);
+}
+
+FString FAvaEaseCurveTool::GetReferencerName() const
+{
+	return TEXT("AvaEaseCurveTool");
+}
+
+TSharedRef<SWidget> FAvaEaseCurveTool::GenerateWidget()
+{
+	UpdateEaseCurveFromSequencerKeySelections();
+
+	if (!ToolWidget.IsValid())
+	{
+		BindCommands();
+
+		ToolWidget = SNew(SAvaEaseCurveTool, SharedThis(this))
+			.InitialTangents(GetEaseCurveTangents())
+			.Visibility(this, &FAvaEaseCurveTool::GetVisibility)
+			.ToolOperation(this, &FAvaEaseCurveTool::GetToolOperation);
+	}
+
+	return ToolWidget.ToSharedRef();
+}
+
+EVisibility FAvaEaseCurveTool::GetVisibility() const
+{
+	return (KeyCache.GetTotalSelectedKeys() > 0 && !KeyCache.IsLastOnlySelectedKey()) ? EVisibility::Visible : EVisibility::Collapsed;
+}
+
+TObjectPtr<UAvaEaseCurve> FAvaEaseCurveTool::GetToolCurve() const
+{
+	return EaseCurve;
+}
+
+FRichCurve* FAvaEaseCurveTool::GetToolRichCurve() const
+{
+	return &EaseCurve->FloatCurve;
+}
+
+FAvaEaseCurveTangents FAvaEaseCurveTool::GetEaseCurveTangents() const
+{
+	return EaseCurve->GetTangents();
+}
+
+void FAvaEaseCurveTool::SetEaseCurveTangents_Internal(const FAvaEaseCurveTangents& InTangents, const EAvaEaseCurveToolOperation InOperation, const bool bInBroadcastUpdate) const
+{
+	switch (InOperation)
+	{
+	case EAvaEaseCurveToolOperation::InOut:
+		EaseCurve->SetTangents(InTangents);
+		break;
+	case EAvaEaseCurveToolOperation::In:
+		EaseCurve->SetEndTangent(InTangents.End, InTangents.EndWeight);
+		break;
+	case EAvaEaseCurveToolOperation::Out:
+		EaseCurve->SetStartTangent(InTangents.Start, InTangents.StartWeight);
+		break;
+	}
+
+	if (bInBroadcastUpdate)
+	{
+		EaseCurve->BroadcastUpdate();
+	}
+}
+
+void FAvaEaseCurveTool::SetEaseCurveTangents(const FAvaEaseCurveTangents& InTangents
+	, const EAvaEaseCurveToolOperation InOperation
+	, const bool bInBroadcastUpdate
+	, const bool bInSetSequencerTangents
+	, const FText& InTransactionText)
+{
+	if (InTangents == GetEaseCurveTangents())
+	{
+		return;
+	}
+
+	const FScopedTransaction Transaction(InTransactionText, !GIsTransacting);
+	EaseCurve->Modify();
+
+	SetEaseCurveTangents_Internal(InTangents, InOperation, bInBroadcastUpdate);
+
+	if (bInSetSequencerTangents)
+	{
+		SetSequencerKeySelectionTangents(InTangents, InOperation);
+	}
+
+	KeyCache = FAvaEaseCurveKeySelection(SequencerWeak.Pin());
+}
+
+void FAvaEaseCurveTool::ResetEaseCurveTangents(const EAvaEaseCurveToolOperation InOperation)
+{
+	FText TransactionText;
+
+	switch (InOperation)
+	{
+	case EAvaEaseCurveToolOperation::InOut:
+		TransactionText = LOCTEXT("ResetTangents", "Reset Tangents");
+		break;
+	case EAvaEaseCurveToolOperation::In:
+		TransactionText = LOCTEXT("ResetEndTangents", "Reset End Tangents");
+		break;
+	case EAvaEaseCurveToolOperation::Out:
+		TransactionText = LOCTEXT("ResetStartTangents", "Reset Start Tangents");
+		break;
+	}
+
+	const FScopedTransaction Transaction(TransactionText);
+	EaseCurve->ModifyOwner();
+
+	const FAvaEaseCurveTangents ZeroTangents;
+	SetEaseCurveTangents(ZeroTangents, InOperation, /*bInBroadcastUpdate=*/true, /*bInSetSequencerTangents=*/true, TransactionText);
+
+	// Update ease curve combobox widget and zoom graph editor
+	if (ToolWidget.IsValid())
+	{
+		ToolWidget->SetTangents(ZeroTangents, InOperation, /*bInSetEaseCurve=*/false, /*bInBroadcastUpdate=*/false, /*bInSetSequencerTangents=*/false);
+	}
+}
+
+void FAvaEaseCurveTool::FlattenOrStraightenTangents(const EAvaEaseCurveToolOperation InOperation, const bool bInFlattenTangents)
+{
+	FText TransactionText;
+	if (bInFlattenTangents)
+	{
+		switch (InOperation)
+		{
+		case EAvaEaseCurveToolOperation::InOut:
+			TransactionText = LOCTEXT("FlattenTangents", "Flatten Tangents");
+			break;
+		case EAvaEaseCurveToolOperation::In:
+			TransactionText = LOCTEXT("FlattenEndTangents", "Flatten End Tangents");
+			break;
+		case EAvaEaseCurveToolOperation::Out:
+			TransactionText = LOCTEXT("FlattenStartTangents", "Flatten Start Tangents");
+			break;
+		}
+	}
+	else
+	{
+		switch (InOperation)
+		{
+		case EAvaEaseCurveToolOperation::InOut:
+			TransactionText = LOCTEXT("StraightenTangents", "Straighten Tangents");
+			break;
+		case EAvaEaseCurveToolOperation::In:
+			TransactionText = LOCTEXT("StraightenEndTangents", "Straighten End Tangents");
+			break;
+		case EAvaEaseCurveToolOperation::Out:
+			TransactionText = LOCTEXT("StraightenStartTangents", "Straighten Start Tangents");
+			break;
+		}
+	}
+	const FScopedTransaction Transaction(TransactionText);
+	EaseCurve->ModifyOwner();
+
+	if (InOperation == EAvaEaseCurveToolOperation::Out || InOperation == EAvaEaseCurveToolOperation::InOut)
+	{
+		EaseCurve->FlattenOrStraightenTangents(EaseCurve->GetStartKeyHandle(), bInFlattenTangents);
+	}
+	if (InOperation == EAvaEaseCurveToolOperation::In || InOperation == EAvaEaseCurveToolOperation::InOut)
+	{
+		EaseCurve->FlattenOrStraightenTangents(EaseCurve->GetEndKeyHandle(), bInFlattenTangents);
+	}
+
+	const FAvaEaseCurveTangents NewTangents = EaseCurve->GetTangents();
+	SetEaseCurveTangents(NewTangents, InOperation, /*bInBroadcastUpdate=*/true, /*bInSetSequencerTangents=*/true, TransactionText);
+
+	// Update ease curve combobox widget and zoom graph editor
+	if (ToolWidget.IsValid())
+	{
+		ToolWidget->SetTangents(NewTangents, InOperation, /*bInSetEaseCurve=*/false, /*bInBroadcastUpdate=*/false, /*bInSetSequencerTangents=*/false);
+	}
+}
+
+bool FAvaEaseCurveTool::CanApplyQuickEaseToSequencerKeySelections() const
+{
+	FAvaEaseCurveTangents Tangents;
+	return FAvaEaseCurveTangents::FromString(ToolSettings->GetQuickEaseTangents(), Tangents);
+}
+
+void FAvaEaseCurveTool::ApplyQuickEaseToSequencerKeySelections(const EAvaEaseCurveToolOperation InOperation)
+{
+	FAvaEaseCurveTangents Tangents;
+	if (!FAvaEaseCurveTangents::FromString(ToolSettings->GetQuickEaseTangents(), Tangents))
+	{
+		UE_LOG(LogTemp, Warning, TEXT("Ease curve tool failed to apply quick ease tangents: "
+			"Could not parse configured quick ease tangent string."));
+		return;
+	}
+
+	SetEaseCurveTangents(Tangents, InOperation, /*bInBroadcastUpdate=*/true, /*bInSetSequencerTangents=*/true);
+
+	// Update ease curve combobox widget and zoom graph editor
+	if (ToolWidget.IsValid())
+	{
+		ToolWidget->SetTangents(Tangents, InOperation, /*bInSetEaseCurve=*/false, /*bInBroadcastUpdate=*/false, /*bInSetSequencerTangents=*/false);
+	}
+
+	if (FEngineAnalytics::IsAvailable())
+	{
+		FString ParamValue;
+		switch (InOperation)
+		{
+		case EAvaEaseCurveToolOperation::InOut:
+			ParamValue = TEXT("InOut");
+			break;
+		case EAvaEaseCurveToolOperation::In:
+			ParamValue = TEXT("In");
+			break;
+		case EAvaEaseCurveToolOperation::Out:
+			ParamValue = TEXT("Out");
+			break;
+		}
+		FEngineAnalytics::GetProvider().RecordEvent(TEXT("Editor.Usage.MotionDesign.EaseCurveTool"), TEXT("QuickEase"), ParamValue);
+	}
+}
+
+void FAvaEaseCurveTool::SetSequencerKeySelectionTangents(const FAvaEaseCurveTangents& InTangents, const EAvaEaseCurveToolOperation InOperation)
+{
+	KeyCache = FAvaEaseCurveKeySelection(SequencerWeak.Pin());
+
+	if (KeyCache.GetTotalSelectedKeys() == 0)
+	{
+		return;
+	}
+
+	KeyCache.SetTangents(InTangents, InOperation, GetDisplayRate(), GetTickResolution(), ToolSettings->GetAutoFlipTangents());
+}
+
+void FAvaEaseCurveTool::UpdateEaseCurveFromSequencerKeySelections()
+{
+	KeyCache = FAvaEaseCurveKeySelection(SequencerWeak.Pin());
+
+	const FAvaEaseCurveTangents AverageTangents = KeyCache.AverageTangents(GetDisplayRate(), GetTickResolution(), ToolSettings->GetAutoFlipTangents());
+
+	SetEaseCurveTangents(AverageTangents, EAvaEaseCurveToolOperation::InOut, /*bInBroadcastUpdate=*/true, false);
+
+	// Update the preset combobox widget
+	if (ToolWidget.IsValid())
+	{
+		ToolWidget->SetTangents(AverageTangents, EAvaEaseCurveToolOperation::InOut
+			, /*bInSetEaseCurve=*/false, /*bInBroadcastUpdate=*/false, /*bInSetSequencerTangents*/false);
+	}
+}
+
+UCurveBase* FAvaEaseCurveTool::CreateCurveAsset() const
+{
+	const FAssetToolsModule& AssetToolsModule = FModuleManager::LoadModuleChecked<FAssetToolsModule>(TEXT("AssetTools"));
+
+	FString OutNewPackageName;
+	FString OutNewAssetName;
+	AssetToolsModule.Get().CreateUniqueAssetName(TEXT("/Game/NewCurve"), TEXT(""), OutNewPackageName, OutNewAssetName);
+
+	const TSharedRef<SDlgPickAssetPath> NewAssetDialog =
+		SNew(SDlgPickAssetPath)
+		.Title(LOCTEXT("CreateExternalCurve", "Create External Curve"))
+		.DefaultAssetPath(FText::FromString(OutNewPackageName));
+
+	if (NewAssetDialog->ShowModal() != EAppReturnType::Cancel)
+	{
+		const FString PackageName = NewAssetDialog->GetFullAssetPath().ToString();
+		const FName AssetName = FName(*NewAssetDialog->GetAssetName().ToString());
+
+		UPackage* const Package = CreatePackage(*PackageName);
+		
+		// Create curve object
+		UObject* NewCurveObject = nullptr;
+		
+		if (UCurveFactory* const CurveFactory = Cast<UCurveFactory>(NewObject<UFactory>(GetTransientPackage(), UCurveFactory::StaticClass())))
+		{
+			CurveFactory->CurveClass = UCurveFloat::StaticClass();
+			NewCurveObject = CurveFactory->FactoryCreateNew(CurveFactory->GetSupportedClass(), Package, AssetName, RF_Public | RF_Standalone, nullptr, GWarn);
+		}
+
+		if (NewCurveObject)
+		{
+			UCurveBase* AssetCurve = nullptr;
+			
+			// Copy curve data from current curve to newly create curve
+			UCurveFloat* const DestCurve = CastChecked<UCurveFloat>(NewCurveObject);
+			if (EaseCurve && DestCurve)
+			{
+				DestCurve->bIsEventCurve = false;
+
+				AssetCurve = DestCurve;
+
+				for (auto It(EaseCurve->FloatCurve.GetKeyIterator()); It; ++It)
+				{
+					const FRichCurveKey& Key = *It;
+					const FKeyHandle KeyHandle = DestCurve->FloatCurve.AddKey(Key.Time, Key.Value);
+					DestCurve->FloatCurve.GetKey(KeyHandle) = Key;
+				}
+			}
+
+			FAssetRegistryModule::AssetCreated(NewCurveObject);
+
+			Package->GetOutermost()->MarkPackageDirty();
+
+			return AssetCurve;
+		}
+	}
+
+	return nullptr;
+}
+
+EAvaEaseCurveToolOperation FAvaEaseCurveTool::GetToolOperation() const
+{
+	return OperationMode;
+}
+
+void FAvaEaseCurveTool::SetToolOperation(const EAvaEaseCurveToolOperation InNewOperation)
+{
+	OperationMode = InNewOperation;
+}
+
+bool FAvaEaseCurveTool::IsToolOperation(const EAvaEaseCurveToolOperation InNewOperation) const
+{
+	return OperationMode == InNewOperation;
+}
+
+bool FAvaEaseCurveTool::CanCopyTangentsToClipboard() const
+{
+	return true;
+}
+
+void FAvaEaseCurveTool::CopyTangentsToClipboard() const
+{
+	FPlatformApplicationMisc::ClipboardCopy(*EaseCurve->GetTangents().ToJson());
+
+	ShowNotificationMessage(LOCTEXT("EaseCurveToolTangentsCopied", "Ease Curve Tool Tangents Copied!"));
+}
+
+bool FAvaEaseCurveTool::CanPasteTangentsFromClipboard() const
+{
+	FAvaEaseCurveTangents Tangents;
+	return TangentsFromClipboardPaste(Tangents);
+}
+
+void FAvaEaseCurveTool::PasteTangentsFromClipboard() const
+{
+	FAvaEaseCurveTangents Tangents;
+	if (TangentsFromClipboardPaste(Tangents))
+	{
+		EaseCurve->SetTangents(Tangents);
+	}
+}
+
+bool FAvaEaseCurveTool::TangentsFromClipboardPaste(FAvaEaseCurveTangents& OutTangents)
+{
+	FString ClipboardString;
+	FPlatformApplicationMisc::ClipboardPaste(ClipboardString);
+
+	// Expects four comma separated cubic bezier points that define the curve
+	return FAvaEaseCurveTangents::FromString(ClipboardString, OutTangents);
+}
+
+bool FAvaEaseCurveTool::IsKeyInterpMode(const ERichCurveInterpMode InInterpMode, const ERichCurveTangentMode InTangentMode) const
+{
+	const FKeyHandle StartKeyHandle = EaseCurve->GetStartKeyHandle();
+	return (EaseCurve->FloatCurve.GetKeyInterpMode(StartKeyHandle) == InInterpMode
+		&& EaseCurve->FloatCurve.GetKeyTangentMode(StartKeyHandle) == InTangentMode);
+}
+
+void FAvaEaseCurveTool::SetKeyInterpMode(const ERichCurveInterpMode InInterpMode, const ERichCurveTangentMode InTangentMode) const
+{
+	const FKeyHandle StartKeyHandle = EaseCurve->GetStartKeyHandle();
+
+	const FScopedTransaction Transaction(LOCTEXT("CurveEditor_SetInterpolationMode", "Select Interpolation Mode"));
+	EaseCurve->ModifyOwner();
+
+	EaseCurve->FloatCurve.SetKeyInterpMode(StartKeyHandle, InInterpMode);
+	EaseCurve->FloatCurve.SetKeyTangentMode(StartKeyHandle, InTangentMode);
+
+	if (InInterpMode != ERichCurveInterpMode::RCIM_Cubic)
+	{
+		FRichCurveKey& StartKey = EaseCurve->GetStartKey();
+		StartKey.LeaveTangentWeight = 0.f;
+		
+		FRichCurveKey& EndKey = EaseCurve->GetEndKey();
+		EndKey.ArriveTangentWeight = 0.f;
+	}
+
+	TArray<FRichCurveEditInfo> ChangedCurveEditInfos;
+	ChangedCurveEditInfos.Add(FRichCurveEditInfo(&EaseCurve->FloatCurve));
+	EaseCurve->OnCurveChanged(ChangedCurveEditInfos);
+}
+
+void FAvaEaseCurveTool::BeginTransaction(const FText& InDescription) const
+{
+	if (GEditor)
+	{
+		EaseCurve->ModifyOwnerChange();
+
+		GEditor->BeginTransaction(InDescription);
+	}
+}
+
+void FAvaEaseCurveTool::EndTransaction() const
+{
+	if (GEditor)
+	{
+		GEditor->EndTransaction();
+	}
+}
+
+void FAvaEaseCurveTool::UndoAction()
+{
+	if (GEditor && GEditor->UndoTransaction())
+	{
+		UpdateEaseCurveFromSequencerKeySelections();
+	}
+}
+
+void FAvaEaseCurveTool::RedoAction()
+{
+	if (GEditor && GEditor->RedoTransaction())
+	{
+		UpdateEaseCurveFromSequencerKeySelections();
+	}
+}
+
+void FAvaEaseCurveTool::PostUndo(bool bInSuccess)
+{
+	UndoAction();
+}
+
+void FAvaEaseCurveTool::PostRedo(bool bInSuccess)
+{
+	RedoAction();
+}
+
+void FAvaEaseCurveTool::OpenToolSettings() const
+{
+	if (ISettingsModule* const SettingsModule = FModuleManager::GetModulePtr<ISettingsModule>(TEXT("Settings")))
+	{
+		SettingsModule->ShowViewer(ToolSettings->GetContainerName(), ToolSettings->GetCategoryName(), ToolSettings->GetSectionName());
+	}
+}
+
+FFrameRate FAvaEaseCurveTool::GetTickResolution() const
+{
+	if (const TSharedPtr<ISequencer> Sequencer = SequencerWeak.Pin())
+	{
+		return Sequencer->GetFocusedTickResolution();
+	}
+
+	return FFrameRate();
+}
+
+FFrameRate FAvaEaseCurveTool::GetDisplayRate() const
+{
+	if (const TSharedPtr<ISequencer> Sequencer = SequencerWeak.Pin())
+	{
+		return Sequencer->GetFocusedDisplayRate();
+	}
+
+	// Fallback to using config display rate if tool is being used outside sequencer
+	return GetDefault<UAvaSequencerSettings>()->GetDisplayRate();
+}
+
+void FAvaEaseCurveTool::ShowNotificationMessage(const FText& InMessageText)
+{
+	FNotificationInfo Info(InMessageText);
+	Info.ExpireDuration = 3.0f;
+	FSlateNotificationManager::Get().AddNotification(Info);
+}
+
+bool FAvaEaseCurveTool::HasCachedKeysToEase()
+{
+	bool bEaseableKeys = false;
+
+	KeyCache.ForEachEaseableKey(/*bInIncludeEqualValueKeys=*/false, [&bEaseableKeys](const FKeyHandle& InKeyHandle
+		, const FKeyHandle& InNextKeyHandle, const FAvaEaseCurveKeySelection::FChannelData& InChannelData)
+		{
+			bEaseableKeys = true;
+			return false;
+		});
+
+	return bEaseableKeys;
+}
+
+bool FAvaEaseCurveTool::AreAllEaseCurves()
+{
+	return KeyCache.AreAllEaseCurves();
+}
+
+void FAvaEaseCurveTool::ResetToDefaultPresets()
+{
+	const FText MessageBoxTitle = LOCTEXT("ResetToDefaultPresets", "Reset To Default Presets");
+	const EAppReturnType::Type Response = FMessageDialog::Open(EAppMsgType::YesNoCancel
+		, LOCTEXT("ConfirmResetToDefaultPresets", "Are you sure you want to reset to default presets?\n\n"
+			"*CAUTION* All directories and files inside '[Project]/Config/EaseCurves' will be lost!")
+		, MessageBoxTitle);
+	if (Response == EAppReturnType::Yes)
+	{
+		UAvaEaseCurveSubsystem::Get().ResetToDefaultPresets(false);
+	}
+}
+
+void FAvaEaseCurveTool::ApplyTangents()
+{
+	SetEaseCurveTangents(GetEaseCurveTangents(), GetToolOperation(), /*bInBroadcastUpdate=*/true, /*bInSetSequencerTangents=*/true);
+}
+
+void FAvaEaseCurveTool::ZoomToFit() const
+{
+	if (ToolWidget.IsValid())
+	{
+		ToolWidget->ZoomToFit();
+	}
+}
+
+TSharedPtr<FUICommandList> FAvaEaseCurveTool::GetCommandList() const
+{
+	return CommandList;
+}
+
+void FAvaEaseCurveTool::BindCommands()
+{
+	const FAvaEaseCurveToolCommands& EaseCurveToolCommands = FAvaEaseCurveToolCommands::Get();
+
+	CommandList = MakeShared<FUICommandList>();
+
+	CommandList->MapAction(FGenericCommands::Get().Undo
+		, FExecuteAction::CreateSP(this, &FAvaEaseCurveTool::UndoAction));
+	
+	CommandList->MapAction(FGenericCommands::Get().Redo
+		, FExecuteAction::CreateSP(this, &FAvaEaseCurveTool::RedoAction));
+
+	CommandList->MapAction(EaseCurveToolCommands.OpenToolSettings
+		, FExecuteAction::CreateSP(this, &FAvaEaseCurveTool::OpenToolSettings));
+
+	CommandList->MapAction(EaseCurveToolCommands.ResetToDefaultPresets
+		, FExecuteAction::CreateSP(this, &FAvaEaseCurveTool::ResetToDefaultPresets));
+
+	CommandList->MapAction(EaseCurveToolCommands.Refresh
+		, FExecuteAction::CreateSP(this, &FAvaEaseCurveTool::UpdateEaseCurveFromSequencerKeySelections));
+	
+	CommandList->MapAction(EaseCurveToolCommands.Apply
+		, FExecuteAction::CreateSP(this, &FAvaEaseCurveTool::ApplyTangents));
+
+	CommandList->MapAction(EaseCurveToolCommands.ZoomToFit
+		, FExecuteAction::CreateSP(this, &FAvaEaseCurveTool::ZoomToFit));
+
+	CommandList->MapAction(EaseCurveToolCommands.ToggleGridSnap
+		, FExecuteAction::CreateUObject(ToolSettings, &UAvaEaseCurveToolSettings::ToggleGridSnap)
+		, FCanExecuteAction()
+		, FIsActionChecked::CreateUObject(ToolSettings, &UAvaEaseCurveToolSettings::GetGridSnap));
+
+	CommandList->MapAction(EaseCurveToolCommands.ToggleAutoFlipTangents
+		, FExecuteAction::CreateUObject(ToolSettings, &UAvaEaseCurveToolSettings::ToggleAutoFlipTangents)
+		, FCanExecuteAction()
+		, FIsActionChecked::CreateUObject(ToolSettings, &UAvaEaseCurveToolSettings::GetAutoFlipTangents));
+
+	CommandList->MapAction(EaseCurveToolCommands.ToggleAutoZoomToFit
+		, FExecuteAction::CreateUObject(ToolSettings, &UAvaEaseCurveToolSettings::ToggleAutoZoomToFit)
+		, FCanExecuteAction()
+		, FIsActionChecked::CreateUObject(ToolSettings, &UAvaEaseCurveToolSettings::GetAutoZoomToFit));
+
+	CommandList->MapAction(EaseCurveToolCommands.SetOperationToEaseOut
+		, FExecuteAction::CreateSP(this, &FAvaEaseCurveTool::SetToolOperation, EAvaEaseCurveToolOperation::Out)
+		, FCanExecuteAction()
+		, FIsActionChecked::CreateSP(this, &FAvaEaseCurveTool::IsToolOperation, EAvaEaseCurveToolOperation::Out));
+
+	CommandList->MapAction(EaseCurveToolCommands.SetOperationToEaseInOut
+		, FExecuteAction::CreateSP(this, &FAvaEaseCurveTool::SetToolOperation, EAvaEaseCurveToolOperation::InOut)
+		, FCanExecuteAction()
+		, FIsActionChecked::CreateSP(this, &FAvaEaseCurveTool::IsToolOperation, EAvaEaseCurveToolOperation::InOut));
+
+	CommandList->MapAction(EaseCurveToolCommands.SetOperationToEaseIn
+		, FExecuteAction::CreateSP(this, &FAvaEaseCurveTool::SetToolOperation, EAvaEaseCurveToolOperation::In)
+		, FCanExecuteAction()
+		, FIsActionChecked::CreateSP(this, &FAvaEaseCurveTool::IsToolOperation, EAvaEaseCurveToolOperation::In));
+
+	CommandList->MapAction(EaseCurveToolCommands.ResetTangents
+		, FExecuteAction::CreateSP(this, &FAvaEaseCurveTool::ResetEaseCurveTangents, EAvaEaseCurveToolOperation::InOut));
+
+	CommandList->MapAction(EaseCurveToolCommands.ResetStartTangent
+		, FExecuteAction::CreateSP(this, &FAvaEaseCurveTool::ResetEaseCurveTangents, EAvaEaseCurveToolOperation::Out));
+
+	CommandList->MapAction(EaseCurveToolCommands.ResetEndTangent
+		, FExecuteAction::CreateSP(this, &FAvaEaseCurveTool::ResetEaseCurveTangents, EAvaEaseCurveToolOperation::In));
+
+	CommandList->MapAction(EaseCurveToolCommands.FlattenTangents
+		, FExecuteAction::CreateSP(this, &FAvaEaseCurveTool::FlattenOrStraightenTangents, EAvaEaseCurveToolOperation::InOut, true));
+
+	CommandList->MapAction(EaseCurveToolCommands.FlattenStartTangent
+		, FExecuteAction::CreateSP(this, &FAvaEaseCurveTool::FlattenOrStraightenTangents, EAvaEaseCurveToolOperation::Out, true));
+
+	CommandList->MapAction(EaseCurveToolCommands.FlattenEndTangent
+		, FExecuteAction::CreateSP(this, &FAvaEaseCurveTool::FlattenOrStraightenTangents, EAvaEaseCurveToolOperation::In, true));
+
+	CommandList->MapAction(EaseCurveToolCommands.StraightenTangents
+		, FExecuteAction::CreateSP(this, &FAvaEaseCurveTool::FlattenOrStraightenTangents, EAvaEaseCurveToolOperation::InOut, false));
+
+	CommandList->MapAction(EaseCurveToolCommands.StraightenStartTangent
+		, FExecuteAction::CreateSP(this, &FAvaEaseCurveTool::FlattenOrStraightenTangents, EAvaEaseCurveToolOperation::Out, false));
+
+	CommandList->MapAction(EaseCurveToolCommands.StraightenEndTangent
+		, FExecuteAction::CreateSP(this, &FAvaEaseCurveTool::FlattenOrStraightenTangents, EAvaEaseCurveToolOperation::In, false));
+
+	CommandList->MapAction(EaseCurveToolCommands.CopyTangents
+		, FExecuteAction::CreateSP(this, &FAvaEaseCurveTool::CopyTangentsToClipboard)
+		, FCanExecuteAction::CreateSP(this, &FAvaEaseCurveTool::CanCopyTangentsToClipboard));
+
+	CommandList->MapAction(EaseCurveToolCommands.PasteTangents
+		, FExecuteAction::CreateSP(this, &FAvaEaseCurveTool::PasteTangentsFromClipboard)
+		, FCanExecuteAction::CreateSP(this, &FAvaEaseCurveTool::CanPasteTangentsFromClipboard));
+
+	CommandList->MapAction(EaseCurveToolCommands.CreateExternalCurveAsset
+		, FExecuteAction::CreateSPLambda(this, [this]()
+			{
+				CreateCurveAsset();
+			})
+		, FCanExecuteAction());
+
+	CommandList->MapAction(EaseCurveToolCommands.SetKeyInterpConstant,
+		FExecuteAction::CreateSP(this, &FAvaEaseCurveTool::SetKeyInterpMode, RCIM_Constant, RCTM_Auto),
+		FCanExecuteAction(),
+		FIsActionChecked::CreateSP(this, &FAvaEaseCurveTool::IsKeyInterpMode, RCIM_Constant, RCTM_Auto));
+
+	CommandList->MapAction(EaseCurveToolCommands.SetKeyInterpLinear,
+		FExecuteAction::CreateSP(this, &FAvaEaseCurveTool::SetKeyInterpMode, RCIM_Linear, RCTM_Auto),
+		FCanExecuteAction(),
+		FIsActionChecked::CreateSP(this, &FAvaEaseCurveTool::IsKeyInterpMode, RCIM_Linear, RCTM_Auto));
+
+	CommandList->MapAction(EaseCurveToolCommands.SetKeyInterpCubicAuto,
+		FExecuteAction::CreateSP(this, &FAvaEaseCurveTool::SetKeyInterpMode, RCIM_Cubic, RCTM_Auto),
+		FCanExecuteAction(),
+		FIsActionChecked::CreateSP(this, &FAvaEaseCurveTool::IsKeyInterpMode, RCIM_Cubic, RCTM_Auto));
+
+	CommandList->MapAction(EaseCurveToolCommands.SetKeyInterpCubicSmartAuto,
+		FExecuteAction::CreateSP(this, &FAvaEaseCurveTool::SetKeyInterpMode, RCIM_Cubic, RCTM_SmartAuto),
+		FCanExecuteAction(),
+		FIsActionChecked::CreateSP(this, &FAvaEaseCurveTool::IsKeyInterpMode, RCIM_Cubic, RCTM_SmartAuto));
+
+	CommandList->MapAction(EaseCurveToolCommands.SetKeyInterpCubicUser,
+		FExecuteAction::CreateSP(this, &FAvaEaseCurveTool::SetKeyInterpMode, RCIM_Cubic, RCTM_User),
+		FCanExecuteAction(),
+		FIsActionChecked::CreateSP(this, &FAvaEaseCurveTool::IsKeyInterpMode, RCIM_Cubic, RCTM_User));
+
+	CommandList->MapAction(EaseCurveToolCommands.SetKeyInterpCubicBreak,
+		FExecuteAction::CreateSP(this, &FAvaEaseCurveTool::SetKeyInterpMode, RCIM_Cubic, RCTM_Break),
+		FCanExecuteAction(),
+		FIsActionChecked::CreateSP(this, &FAvaEaseCurveTool::IsKeyInterpMode, RCIM_Cubic, RCTM_Break));
+}
+
+#undef LOCTEXT_NAMESPACE
diff --git a/Engine/Plugins/VirtualProduction/Avalanche/Source/AvalancheSequencer/Private/EaseCurveTool/AvaEaseCurveTool.h b/Engine/Plugins/VirtualProduction/Avalanche/Source/AvalancheSequencer/Private/EaseCurveTool/AvaEaseCurveTool.h
new file mode 100644
--- /dev/null
+++ b/Engine/Plugins/VirtualProduction/Avalanche/Source/AvalancheSequencer/Private/EaseCurveTool/AvaEaseCurveTool.h
@@ -0,0 +1,175 @@
+// Copyright Epic Games, Inc. All Rights Reserved.
+
+#pragma once
+
+#include "EaseCurveTool/AvaEaseCurve.h"
+#include "EaseCurveTool/AvaEaseCurveKeySelection.h"
+#include "Curves/RichCurve.h"
+#include "Curves/KeyHandle.h"
+#include "EaseCurveTool/Widgets/SAvaEaseCurvePreset.h"
+#include "EditorUndoClient.h"
+#include "Internationalization/Text.h"
+#include "Templates/SharedPointer.h"
+#include "UObject/WeakObjectPtrTemplates.h"
+#include "Widgets/DeclarativeSyntaxSupport.h"
+#include "Widgets/SCompoundWidget.h"
+
+class FCurveEditor;
+class FUICommandList;
+class ISequencer;
+class SAvaEaseCurveTool;
+class SCurveEditor;
+class SCurveEditorPanel;
+class SEditableTextBox;
+class UAvaEaseCurve;
+class UAvaEaseCurveToolSettings;
+class UCurveFloat;
+class UMovieSceneSection;
+struct FAvaEaseCurveTangents;
+struct FGuid;
+
+namespace UE::Sequencer
+{
+	class FChannelModel;
+	class FSequencerSelection;
+}
+
+/** Current default and only implemented is DualKeyEdit. */
+enum class EAvaEaseCurveToolMode : uint8
+{
+	/** Edits the selected key's leave tangent and the next key's arrive tangent in the curve editor graph. */
+	DualKeyEdit,
+	/** Edits only the selected key.
+	 * The leave tangent in the curve editor graph will set the sequence key arrive tangent.
+	 * The arrive tangent in the curve editor graph will set the sequence key leave tangent. */
+	SingleKeyEdit
+};
+
+enum class EAvaEaseCurveToolOperation : uint8
+{
+	InOut,
+	In,
+	Out
+};
+
+class FAvaEaseCurveTool
+	: public TSharedFromThis<FAvaEaseCurveTool>
+	, public FGCObject
+	, public FSelfRegisteringEditorUndoClient
+{
+public:
+	static void ShowNotificationMessage(const FText& InMessageText);
+
+	/** Returns true if the clipboard paste data contains tangent information. */
+	static bool TangentsFromClipboardPaste(FAvaEaseCurveTangents& OutTangents);
+
+	FAvaEaseCurveTool(const TSharedRef<ISequencer>& InSequencer);
+
+	virtual ~FAvaEaseCurveTool() override;
+
+	TSharedRef<SWidget> GenerateWidget();
+
+	EVisibility GetVisibility() const;
+
+	TObjectPtr<UAvaEaseCurve> GetToolCurve() const;
+	FRichCurve* GetToolRichCurve() const;
+
+	FAvaEaseCurveTangents GetEaseCurveTangents() const;
+
+	/**
+	 * Sets the internal ease curve tangents and optionally broadcasts a change event for the curve object.
+	 * This is different from SetEaseCurveTangents_Internal in that it performs undo/redo transactions and
+	 * optionally sets the selected tangents in the actual sequence.
+	 */
+	void SetEaseCurveTangents(const FAvaEaseCurveTangents& InTangents
+		, const EAvaEaseCurveToolOperation InOperation
+		, const bool bInBroadcastUpdate
+		, const bool bInSetSequencerTangents
+		, const FText& InTransactionText = NSLOCTEXT("AvaEaseCurveTool", "SetEaseCurveTangents", "Set Ease Curve Tangents"));
+
+	void ResetEaseCurveTangents(const EAvaEaseCurveToolOperation InOperation);
+
+	void FlattenOrStraightenTangents(const EAvaEaseCurveToolOperation InOperation, const bool bInFlattenTangents);
+
+	/** Creates a new external float curve from the internet curve editor curve. */
+	UCurveBase* CreateCurveAsset() const;
+
+	void SetSequencerKeySelectionTangents(const FAvaEaseCurveTangents& InTangents, const EAvaEaseCurveToolOperation InOperation = EAvaEaseCurveToolOperation::InOut);
+
+	bool CanApplyQuickEaseToSequencerKeySelections() const;
+	void ApplyQuickEaseToSequencerKeySelections(const EAvaEaseCurveToolOperation InOperation = EAvaEaseCurveToolOperation::InOut);
+
+	/** Updates the ease curve graph view based on the active sequencer key selection. */
+	void UpdateEaseCurveFromSequencerKeySelections();
+
+	EAvaEaseCurveToolOperation GetToolOperation() const;
+	void SetToolOperation(const EAvaEaseCurveToolOperation InNewOperation);
+	bool IsToolOperation(const EAvaEaseCurveToolOperation InNewOperation) const;
+
+	bool CanCopyTangentsToClipboard() const;
+	void CopyTangentsToClipboard() const;
+	bool CanPasteTangentsFromClipboard() const;
+	void PasteTangentsFromClipboard() const;
+
+	bool IsKeyInterpMode(const ERichCurveInterpMode InInterpMode, const ERichCurveTangentMode InTangentMode) const;
+	void SetKeyInterpMode(const ERichCurveInterpMode InInterpMode, const ERichCurveTangentMode InTangentMode) const;
+
+	void BeginTransaction(const FText& InDescription) const;
+	void EndTransaction() const;
+
+	void UndoAction();
+	void RedoAction();
+
+	void OpenToolSettings() const;
+
+	FFrameRate GetTickResolution() const;
+	FFrameRate GetDisplayRate() const;
+
+	bool HasCachedKeysToEase();
+
+	bool AreAllEaseCurves();
+
+	//~ Begin FGCObject
+	virtual void AddReferencedObjects(FReferenceCollector& Collector) override;
+	virtual FString GetReferencerName() const override;
+	//~ End FGCObject
+
+	//~ Begin FEditorUndoClient
+	virtual void PostUndo(bool bInSuccess) override;
+	virtual void PostRedo(bool bInSuccess) override;
+	//~ End FEditorUndoClient
+
+	TSharedPtr<FUICommandList> GetCommandList() const;
+
+protected:
+	void ResetToDefaultPresets();
+	void ApplyTangents();
+	void ZoomToFit() const;
+
+	void BindCommands();
+
+	/** 
+	 * Sets the internal ease curve tangents and optionally broadcasts a change event for the curve object.
+	 * Changing the internal ease curve tangents will be directly reflected in the ease curve editor graph.
+	 */
+	void SetEaseCurveTangents_Internal(const FAvaEaseCurveTangents& InTangents, const EAvaEaseCurveToolOperation InOperation, const bool bInBroadcastUpdate) const;
+
+	TWeakPtr<ISequencer> SequencerWeak;
+
+	TSharedPtr<FUICommandList> CommandList;
+
+	TObjectPtr<UAvaEaseCurve> EaseCurve;
+
+	TObjectPtr<UAvaEaseCurveToolSettings> ToolSettings;
+
+	EAvaEaseCurveToolMode ToolMode = EAvaEaseCurveToolMode::DualKeyEdit;
+	EAvaEaseCurveToolOperation OperationMode = EAvaEaseCurveToolOperation::InOut;
+
+	TSharedPtr<SAvaEaseCurveTool> ToolWidget;
+
+	/** Cached data set when a new sequencer selection is made. */
+	FAvaEaseCurveKeySelection KeyCache;
+
+private:
+	void OnSequencerSelectionChanged(TArray<FGuid> InObjectGuids);
+};
diff --git a/Engine/Plugins/VirtualProduction/Avalanche/Source/AvalancheSequencer/Private/EaseCurveTool/AvaEaseCurveToolCommands.cpp b/Engine/Plugins/VirtualProduction/Avalanche/Source/AvalancheSequencer/Private/EaseCurveTool/AvaEaseCurveToolCommands.cpp
new file mode 100644
--- /dev/null
+++ b/Engine/Plugins/VirtualProduction/Avalanche/Source/AvalancheSequencer/Private/EaseCurveTool/AvaEaseCurveToolCommands.cpp
@@ -0,0 +1,250 @@
+// Copyright Epic Games, Inc. All Rights Reserved.
+
+#include "AvaEaseCurveToolCommands.h"
+
+#define LOCTEXT_NAMESPACE "AvaEaseCurveToolCommands"
+
+void FAvaEaseCurveToolCommands::RegisterCommands()
+{
+	const TSharedRef<FBindingContext> SharedThis = AsShared();
+
+	FUICommandInfo::MakeCommandInfo(SharedThis
+		, OpenToolSettings
+		, TEXT("OpenToolSettings")
+		, LOCTEXT("OpenToolSettings_Label", "Open Tool Settings...")
+		, LOCTEXT("OpenToolSettings_ToolTip", "Open the editor settings for the the curve ease tool.")
+		, FSlateIcon()
+		, EUserInterfaceActionType::Button
+		, FInputChord());
+
+	
+	UI_COMMAND(ResetToDefaultPresets
+		, "Reset To Default Presets"
+		, "Reset presets to the defaults.\n\n"
+			"*CAUTION* All directories and files inside '[Project]/Config/EaseCurves' will be replaced with defaults!"
+		, EUserInterfaceActionType::Button
+		, FInputChord());
+
+	FUICommandInfo::MakeCommandInfo(SharedThis
+		, Refresh
+		, TEXT("Refresh")
+		, LOCTEXT("Refresh_Label", "Refresh")
+		, LOCTEXT("Refresh_ToolTip", "Refresh the ease curve editor graph tangents from the selected sequencer key(s).")
+		, FSlateIcon(FAppStyle::Get().GetStyleSetName(), TEXT("Icons.Refresh"))
+		, EUserInterfaceActionType::Button
+		, FInputChord(EKeys::F5));
+
+	FUICommandInfo::MakeCommandInfo(SharedThis
+		, Apply
+		, TEXT("Apply")
+		, LOCTEXT("Apply_Label", "Apply")
+		, LOCTEXT("Apply_ToolTip", "Applies the ease curve to the selected sequencer key(s).")
+		, FSlateIcon(FAppStyle::Get().GetStyleSetName(), TEXT("AssetEditor.Apply"))
+		, EUserInterfaceActionType::Button
+		, FInputChord());
+
+	FUICommandInfo::MakeCommandInfo(SharedThis
+		, ZoomToFit
+		, TEXT("ZoomToFit")
+		, LOCTEXT("ZoomToFit_Label", "Zoom to Fit")
+		, LOCTEXT("ZoomToFit_ToolTip", "Zoom the curve graph view to fit tangent handles.")
+		, FSlateIcon(FAppStyle::Get().GetStyleSetName(), TEXT("CurveEditorTools.SetFocusPlaybackRange"))
+		, EUserInterfaceActionType::Button
+		, FInputChord(EKeys::F));
+
+	FUICommandInfo::MakeCommandInfo(SharedThis
+		, SetOperationToEaseOut
+		, TEXT("SetOperationToEaseOut")
+		, LOCTEXT("SetOperationToEaseOut_Label", "Ease Out")
+		, LOCTEXT("SetOperationToEaseOut_ToolTip", "Set the operation of the Curve Ease Tool to Ease Out only. This will only set the Out (Leave) tangent.")
+		, FSlateIcon(FAppStyle::Get().GetStyleSetName(), TEXT("Icons.ChevronLeft"))
+		, EUserInterfaceActionType::RadioButton
+		, FInputChord(EKeys::O));
+
+	FUICommandInfo::MakeCommandInfo(SharedThis
+		, SetOperationToEaseInOut
+		, TEXT("SetOperationToEaseInOut")
+		, LOCTEXT("SetOperationToEaseInOut_Label", "Ease In Out")
+		, LOCTEXT("SetOperationToEaseInOut_ToolTip", "Set the operation of the Curve Ease Tool to both Ease Out and Ease In. This will set the Out (Leave) and In (Arrive) tangents.")
+		, FSlateIcon(FAppStyle::Get().GetStyleSetName(), TEXT("Icons.BulletPoint"))
+		, EUserInterfaceActionType::RadioButton
+		, FInputChord(EKeys::U));
+
+	FUICommandInfo::MakeCommandInfo(SharedThis
+		, SetOperationToEaseIn
+		, TEXT("SetOperationToEaseIn")
+		, LOCTEXT("SetOperationToEaseIn_Label", "Ease In")
+		, LOCTEXT("SetOperationToEaseIn_ToolTip", "Set the operation of the Curve Ease Tool to Ease In only. This will only set the In (Arrive) tangent.")
+		, FSlateIcon(FAppStyle::Get().GetStyleSetName(), TEXT("Icons.ChevronRight"))
+		, EUserInterfaceActionType::RadioButton
+		, FInputChord(EKeys::I));
+
+	FUICommandInfo::MakeCommandInfo(SharedThis
+		, ToggleGridSnap
+		, TEXT("ToggleGridSnap")
+		, LOCTEXT("ToggleGridSnap_Label", "Snap To Grid")
+		, LOCTEXT("ToggleGridSnap_ToolTip", "Toggle snapping of tangents to grid.")
+		, FSlateIcon(FAppStyle::Get().GetStyleSetName(), TEXT("Icons.Snap"))
+		, EUserInterfaceActionType::Check
+		, FInputChord(EKeys::G));
+
+	FUICommandInfo::MakeCommandInfo(SharedThis
+		, ToggleAutoFlipTangents
+		, TEXT("ToggleAutoFlipTangents")
+		, LOCTEXT("ToggleAutoFlipTangents_Label", "Auto Flip Tangents")
+		, LOCTEXT("ToggleAutoFlipTangents_ToolTip", "Toggle flipping of tangents when sequential key frame curve values are descending.")
+		, FSlateIcon(FAppStyle::Get().GetStyleSetName(), TEXT("GenericCurveEditor.InterpolationCubicAuto.Small"))
+		, EUserInterfaceActionType::Check
+		, FInputChord());
+
+	UI_COMMAND(ToggleAutoZoomToFit
+		, "Auto Zoom To Fit"
+		, "Auto zoom the graph editor to fit the tangent handles after they have been changed."
+		, EUserInterfaceActionType::Check
+		, FInputChord());
+
+	FUICommandInfo::MakeCommandInfo(SharedThis
+		, ResetTangents
+		, TEXT("ResetBothTangents")
+		, LOCTEXT("ResetBothTangents_Label", "Both Tangents")
+		, LOCTEXT("ResetBothTangents_ToolTip", "Set the In / Out tangent values and weights to 0.")
+		, FSlateIcon(FAppStyle::Get().GetStyleSetName(), TEXT("PropertyWindow.DiffersFromDefault"))
+		, EUserInterfaceActionType::Button
+		, FInputChord(EKeys::R));
+
+	FUICommandInfo::MakeCommandInfo(SharedThis
+		, ResetStartTangent
+		, TEXT("ResetStartTangent")
+		, LOCTEXT("ResetStartTangent_Label", "Start Tangent")
+		, LOCTEXT("ResetStartTangent_ToolTip", "Set the start tangent value and weight to 0.")
+		, FSlateIcon(FAppStyle::Get().GetStyleSetName(), TEXT("PropertyWindow.DiffersFromDefault"))
+		, EUserInterfaceActionType::Button
+		, FInputChord(EKeys::R, EModifierKey::Control));
+
+	FUICommandInfo::MakeCommandInfo(SharedThis
+		, ResetEndTangent
+		, TEXT("ResetEndTangent")
+		, LOCTEXT("ResetEndTangent_Label", "End Tangent")
+		, LOCTEXT("ResetEndTangent_ToolTip", "Set the end tangent value and weight to 0.")
+		, FSlateIcon(FAppStyle::Get().GetStyleSetName(), TEXT("PropertyWindow.DiffersFromDefault"))
+		, EUserInterfaceActionType::Button
+		, FInputChord(EKeys::R, EModifierKey::Alt));
+
+	FUICommandInfo::MakeCommandInfo(SharedThis
+		, FlattenTangents
+		, TEXT("FlattenBothTangents")
+		, LOCTEXT("FlattenBothTangents_Label", "Both Tangents")
+		, LOCTEXT("FlattenBothTangents_ToolTip", "Flatten the start and end tangents.")
+		, FSlateIcon(FAppStyle::Get().GetStyleSetName(), TEXT("GenericCurveEditor.FlattenTangents"))
+		, EUserInterfaceActionType::Button
+		, FInputChord(EKeys::T));
+
+	FUICommandInfo::MakeCommandInfo(SharedThis
+		, FlattenStartTangent
+		, TEXT("FlattenStartTangent")
+		, LOCTEXT("FlattenStartTangent_Label", "Start Tangent")
+		, LOCTEXT("FlattenStartTangent_ToolTip", "Flatten the start tangent.")
+		, FSlateIcon(FAppStyle::Get().GetStyleSetName(), TEXT("GenericCurveEditor.FlattenTangents"))
+		, EUserInterfaceActionType::Button
+		, FInputChord(EKeys::T, EModifierKey::Control));
+
+	FUICommandInfo::MakeCommandInfo(SharedThis
+		, FlattenEndTangent
+		, TEXT("FlattenEndTangent")
+		, LOCTEXT("FlattenEndTangent_Label", "End Tangent")
+		, LOCTEXT("FlattenEndTangent_ToolTip", "Flatten the end tangent.")
+		, FSlateIcon(FAppStyle::Get().GetStyleSetName(), TEXT("GenericCurveEditor.FlattenTangents"))
+		, EUserInterfaceActionType::Button
+		, FInputChord(EKeys::T, EModifierKey::Alt));
+
+	FUICommandInfo::MakeCommandInfo(SharedThis
+		, StraightenTangents
+		, TEXT("StraightenBothTangents")
+		, LOCTEXT("StraightenBothTangents_Label", "Both Tangents")
+		, LOCTEXT("StraightenBothTangents_ToolTip", "Straighten the start and end tangents.")
+		, FSlateIcon(FAppStyle::Get().GetStyleSetName(), TEXT("GenericCurveEditor.StraightenTangents"))
+		, EUserInterfaceActionType::Button
+		, FInputChord(EKeys::S));
+
+	FUICommandInfo::MakeCommandInfo(SharedThis
+		, StraightenStartTangent
+		, TEXT("StraightenStartTangent")
+		, LOCTEXT("StraightenStartTangent_Label", "Start Tangent")
+		, LOCTEXT("StraightenStartTangent_ToolTip", "Straighten the start tangent.")
+		, FSlateIcon(FAppStyle::Get().GetStyleSetName(), TEXT("GenericCurveEditor.StraightenTangents"))
+		, EUserInterfaceActionType::Button
+		, FInputChord(EKeys::S, EModifierKey::Control));
+
+	FUICommandInfo::MakeCommandInfo(SharedThis
+		, StraightenEndTangent
+		, TEXT("StraightenEndTangent")
+		, LOCTEXT("StraightenEndTangent_Label", "End Tangent")
+		, LOCTEXT("StraightenEndTangent_ToolTip", "Straighten the end tangent.")
+		, FSlateIcon(FAppStyle::Get().GetStyleSetName(), TEXT("GenericCurveEditor.StraightenTangents"))
+		, EUserInterfaceActionType::Button
+		, FInputChord(EKeys::S, EModifierKey::Alt));
+
+	FUICommandInfo::MakeCommandInfo(SharedThis
+		, CopyTangents
+		, TEXT("CopyTangents")
+		, LOCTEXT("CopyTangents_Label", "Copy Tangents")
+		, LOCTEXT("CopyTangents_ToolTip", "Copy the set of In / Out tangent values and weights.")
+		, FSlateIcon(FAppStyle::Get().GetStyleSetName(), TEXT("GenericCommands.Copy"))
+		, EUserInterfaceActionType::Button
+		, FInputChord(EKeys::C, EModifierKey::Control));
+
+	FUICommandInfo::MakeCommandInfo(SharedThis
+		, PasteTangents
+		, TEXT("PasteTangents")
+		, LOCTEXT("PasteTangents_Label", "Paste Tangents")
+		, LOCTEXT("PasteTangents_ToolTip", "Paste the set of In / Out tangent values and weights.")
+		, FSlateIcon(FAppStyle::Get().GetStyleSetName(), TEXT("GenericCommands.Paste"))
+		, EUserInterfaceActionType::Button
+		, FInputChord(EKeys::V, EModifierKey::Control));
+
+	FUICommandInfo::MakeCommandInfo(SharedThis
+		, CreateExternalCurveAsset
+		, TEXT("CreateExternalCurve")
+		, LOCTEXT("CreateExternalCurve_Label", "Create External Curve...")
+		, LOCTEXT("CreateExternalCurve_ToolTip", "Create an external curve asset using this internal curve.")
+		, FSlateIcon(FAppStyle::Get().GetStyleSetName(), "TimelineEditor.AddCurveAssetTrack")
+		, EUserInterfaceActionType::Button
+		, FInputChord());
+
+	UI_COMMAND(SetKeyInterpConstant, "Constant", "Constant interpolation", EUserInterfaceActionType::ToggleButton, FInputChord(EKeys::Five));
+	UI_COMMAND(SetKeyInterpLinear, "Linear", "Linear interpolation", EUserInterfaceActionType::ToggleButton, FInputChord(EKeys::Four));
+	UI_COMMAND(SetKeyInterpCubicSmartAuto, "Smart Auto", "Cubic interpolation - Smart Automatic tangents", EUserInterfaceActionType::ToggleButton, FInputChord(EKeys::Zero));
+	UI_COMMAND(SetKeyInterpCubicAuto, "Auto", "Cubic interpolation - Automatic tangents", EUserInterfaceActionType::ToggleButton, FInputChord(EKeys::One));
+	UI_COMMAND(SetKeyInterpCubicUser, "User", "Cubic interpolation - User flat tangents", EUserInterfaceActionType::ToggleButton, FInputChord(EKeys::Two));
+	UI_COMMAND(SetKeyInterpCubicBreak, "Break", "Cubic interpolation - User broken tangents", EUserInterfaceActionType::ToggleButton, FInputChord(EKeys::Three));
+	UI_COMMAND(SetKeyInterpToggleWeighted, "Weighted Tangents", "Toggle weighted tangents for cubic interpolation modes. Only supported on some curve types", EUserInterfaceActionType::ToggleButton, FInputChord(EModifierKey::Control, EKeys::W));
+
+	FUICommandInfo::MakeCommandInfo(SharedThis
+		, QuickEase
+		, TEXT("QuickEase")
+		, LOCTEXT("QuickEase_Label", "Quick Ease")
+		, LOCTEXT("QuickEase_ToolTip", "Apply the quick ease preset. This will be disabled if there is no valid quick ease preset setup.")
+		, FSlateIcon(FAppStyle::Get().GetStyleSetName(), TEXT("Icons.BulletPoint"))
+		, EUserInterfaceActionType::Button
+		, FInputChord(EKeys::F8));
+
+	FUICommandInfo::MakeCommandInfo(SharedThis
+		, QuickEaseIn
+		, TEXT("QuickEaseIn")
+		, LOCTEXT("QuickEaseIn_Label", "Quick Ease In")
+		, LOCTEXT("QuickEaseIn_ToolTip", "Apply the quick ease in preset. This will be disabled if there is no valid quick ease preset setup.")
+		, FSlateIcon(FAppStyle::Get().GetStyleSetName(), TEXT("Icons.ChevronRight"))
+		, EUserInterfaceActionType::Button
+		, FInputChord(EKeys::F8, EModifierKey::Shift));
+
+	FUICommandInfo::MakeCommandInfo(SharedThis
+		, QuickEaseOut
+		, TEXT("QuickEaseOut")
+		, LOCTEXT("QuickEaseOut_Label", "Quick Ease Out")
+		, LOCTEXT("QuickEaseOut_ToolTip", "Apply the quick ease out preset. This will be disabled if there is no valid quick ease preset setup.")
+		, FSlateIcon(FAppStyle::Get().GetStyleSetName(), TEXT("Icons.ChevronLeft"))
+		, EUserInterfaceActionType::Button
+		, FInputChord(EKeys::F8, EModifierKey::Command | EModifierKey::Shift));
+}
+
+#undef LOCTEXT_NAMESPACE
diff --git a/Engine/Plugins/VirtualProduction/Avalanche/Source/AvalancheSequencer/Private/EaseCurveTool/AvaEaseCurveToolCommands.h b/Engine/Plugins/VirtualProduction/Avalanche/Source/AvalancheSequencer/Private/EaseCurveTool/AvaEaseCurveToolCommands.h
new file mode 100644
--- /dev/null
+++ b/Engine/Plugins/VirtualProduction/Avalanche/Source/AvalancheSequencer/Private/EaseCurveTool/AvaEaseCurveToolCommands.h
@@ -0,0 +1,68 @@
+// Copyright Epic Games, Inc. All Rights Reserved.
+
+#pragma once
+
+#include "Framework/Commands/Commands.h"
+#include "Styling/AppStyle.h"
+
+class FAvaEaseCurveToolCommands : public TCommands<FAvaEaseCurveToolCommands>
+{
+public:
+	FAvaEaseCurveToolCommands()
+		: TCommands<FAvaEaseCurveToolCommands>(TEXT("AvaEaseCurveTool")
+			, NSLOCTEXT("MotionDesignEaseCurveToolCommands", "MotionDesignEaseCurveToolCommands", "Motion Design Curve Ease Tool Commands")
+			, NAME_None
+			, FAppStyle::GetAppStyleSetName())
+	{
+	}
+
+	virtual void RegisterCommands() override;
+
+	TSharedPtr<FUICommandInfo> OpenToolSettings;
+
+	TSharedPtr<FUICommandInfo> ResetToDefaultPresets;
+
+	TSharedPtr<FUICommandInfo> Refresh;
+	TSharedPtr<FUICommandInfo> Apply;
+
+	TSharedPtr<FUICommandInfo> ZoomToFit;
+
+	TSharedPtr<FUICommandInfo> SetOperationToEaseOut;
+	TSharedPtr<FUICommandInfo> SetOperationToEaseInOut;
+	TSharedPtr<FUICommandInfo> SetOperationToEaseIn;
+
+	TSharedPtr<FUICommandInfo> ToggleGridSnap;
+
+	TSharedPtr<FUICommandInfo> ToggleAutoFlipTangents;
+
+	TSharedPtr<FUICommandInfo> ToggleAutoZoomToFit;
+
+	TSharedPtr<FUICommandInfo> ResetTangents;
+	TSharedPtr<FUICommandInfo> ResetStartTangent;
+	TSharedPtr<FUICommandInfo> ResetEndTangent;
+
+	TSharedPtr<FUICommandInfo> FlattenTangents;
+	TSharedPtr<FUICommandInfo> FlattenStartTangent;
+	TSharedPtr<FUICommandInfo> FlattenEndTangent;
+
+	TSharedPtr<FUICommandInfo> StraightenTangents;
+	TSharedPtr<FUICommandInfo> StraightenStartTangent;
+	TSharedPtr<FUICommandInfo> StraightenEndTangent;
+
+	TSharedPtr<FUICommandInfo> CopyTangents;
+	TSharedPtr<FUICommandInfo> PasteTangents;
+
+	TSharedPtr<FUICommandInfo> CreateExternalCurveAsset;
+
+	TSharedPtr<FUICommandInfo> SetKeyInterpConstant;
+	TSharedPtr<FUICommandInfo> SetKeyInterpLinear;
+	TSharedPtr<FUICommandInfo> SetKeyInterpCubicAuto;
+	TSharedPtr<FUICommandInfo> SetKeyInterpCubicSmartAuto;
+	TSharedPtr<FUICommandInfo> SetKeyInterpCubicUser;
+	TSharedPtr<FUICommandInfo> SetKeyInterpCubicBreak;
+	TSharedPtr<FUICommandInfo> SetKeyInterpToggleWeighted;
+
+	TSharedPtr<FUICommandInfo> QuickEase;
+	TSharedPtr<FUICommandInfo> QuickEaseIn;
+	TSharedPtr<FUICommandInfo> QuickEaseOut;
+};
diff --git a/Engine/Plugins/VirtualProduction/Avalanche/Source/AvalancheSequencer/Private/EaseCurveTool/AvaEaseCurveToolSettings.cpp b/Engine/Plugins/VirtualProduction/Avalanche/Source/AvalancheSequencer/Private/EaseCurveTool/AvaEaseCurveToolSettings.cpp
new file mode 100644
--- /dev/null
+++ b/Engine/Plugins/VirtualProduction/Avalanche/Source/AvalancheSequencer/Private/EaseCurveTool/AvaEaseCurveToolSettings.cpp
@@ -0,0 +1,11 @@
+// Copyright Epic Games, Inc. All Rights Reserved.
+
+#include "AvaEaseCurveToolSettings.h"
+
+UAvaEaseCurveToolSettings::UAvaEaseCurveToolSettings()
+{
+	CategoryName = TEXT("Motion Design");
+	SectionName = TEXT("Ease Curve Tool");
+
+	NewPresetCategory = TEXT("Custom");
+}
diff --git a/Engine/Plugins/VirtualProduction/Avalanche/Source/AvalancheSequencer/Private/EaseCurveTool/AvaEaseCurveToolSettings.h b/Engine/Plugins/VirtualProduction/Avalanche/Source/AvalancheSequencer/Private/EaseCurveTool/AvaEaseCurveToolSettings.h
new file mode 100644
--- /dev/null
+++ b/Engine/Plugins/VirtualProduction/Avalanche/Source/AvalancheSequencer/Private/EaseCurveTool/AvaEaseCurveToolSettings.h
@@ -0,0 +1,71 @@
+// Copyright Epic Games, Inc. All Rights Reserved.
+
+#pragma once
+
+#include "Containers/EnumAsByte.h"
+#include "Engine/DeveloperSettings.h"
+#include "Types/SlateEnums.h"
+#include "UObject/Object.h"
+#include "AvaEaseCurveToolSettings.generated.h"
+
+UCLASS(config=EditorPerProjectUserSettings, meta = (DisplayName = "Ease Curve Tool"))
+class UAvaEaseCurveToolSettings : public UDeveloperSettings
+{
+	GENERATED_BODY()
+
+public:
+	UAvaEaseCurveToolSettings();
+
+	int32 GetGraphSize() const { return GraphSize; }
+	void SetGraphSize(const int32 InSize) { GraphSize = InSize; }
+
+	bool GetGridSnap() const { return bGridSnap; }
+	void SetGridSnap(const bool bInGridSnap) { bGridSnap = bInGridSnap; }
+	void ToggleGridSnap() { bGridSnap = !bGridSnap; }
+
+	int32 GetGridSize() const { return GridSize; }
+	void SetGridSize(const int32 InSize) { GridSize = InSize; }
+
+	bool GetAutoZoomToFit() const { return bAutoZoomToFit; }
+	void SetAutoZoomToFit(const bool bInAutoZoomToFit) { bAutoZoomToFit = bInAutoZoomToFit; }
+	void ToggleAutoZoomToFit() { bAutoZoomToFit = !bAutoZoomToFit; }
+
+	bool GetAutoFlipTangents() const { return bAutoFlipTangents; }
+	void SetAutoFlipTangents(const bool bInAutoFlipTangents) { bAutoFlipTangents = bInAutoFlipTangents; }
+	void ToggleAutoFlipTangents() { bAutoFlipTangents = !bAutoFlipTangents; }
+
+	FString GetNewPresetCategory() const { return NewPresetCategory; }
+	void SetNewPresetCategory(const FString& InNewPresetCategory) { NewPresetCategory = InNewPresetCategory; }
+
+	FString GetQuickEaseTangents() const { return QuickEaseTangents; }
+	void SetQuickEaseTangents(const FString& InTangents) { QuickEaseTangents = InTangents; }
+
+private:
+	/** The height of the curve ease tool in the details panel. */
+	UPROPERTY(Config, EditAnywhere, Category = "EaseCurveTool", meta = (DisplayName = "Graph Size", UIMin = 64, ClampMin = 64, UIMax = 256, ClampMax = 256, Delta = 1))
+	int32 GraphSize = 140;
+
+	/** If true, snaps tangents to grid. */
+	UPROPERTY(Config, EditAnywhere, Category = "EaseCurveTool", meta = (DisplayName = "Grid Snap"))
+	bool bGridSnap = false;
+
+	/** The height of the curve ease tool in the details panel. */
+	UPROPERTY(Config, EditAnywhere, Category = "EaseCurveTool", meta = (DisplayName = "Grid Size", UIMin = 4, ClampMin = 4, UIMax = 24, ClampMax = 24, Delta = 1))
+	int32 GridSize = 8;
+
+	/** If true, will auto zoom the graph editor to fit the tangent handles after they have been changed. */
+	UPROPERTY(Config, EditAnywhere, Category = "EaseCurveTool", meta = (DisplayName = "Auto Zoom To Fit"))
+	bool bAutoZoomToFit = true;
+
+	/** If true, auto flips tangents when sequential key frame curve values are descending. */
+	UPROPERTY(Config, EditAnywhere, Category = "EaseCurveTool", meta = (DisplayName = "Auto Flip Tangents"))
+	bool bAutoFlipTangents = true;
+
+	/** The name of the category to place newly created curve presets. */
+	UPROPERTY(Config, EditAnywhere, Category = "EaseCurveTool", meta = (DisplayName = "New Preset Category"))
+	FString NewPresetCategory;
+
+	/** The tangents to apply for quick ease. Should be in the format of four comma-separated cubic bezier points. Ex. "0.45, 0.34, 0.0, 1.00" */
+	UPROPERTY(Config, EditAnywhere, Category = "EaseCurveTool", meta = (DisplayName = "Quick Ease Tangents"))
+	FString QuickEaseTangents;
+};
