Motion Design: moved the following plugins from /Plugins/Experimental to /Plugins/VirtualProduction:

diff --git a/Engine/Plugins/VirtualProduction/Avalanche/Source/AvalancheSequencer/Private/EaseCurveTool/Widgets/SAvaEaseCurvePresetGroupItem.h b/Engine/Plugins/VirtualProduction/Avalanche/Source/AvalancheSequencer/Private/EaseCurveTool/Widgets/SAvaEaseCurvePresetGroupItem.h
new file mode 100644
--- /dev/null
+++ b/Engine/Plugins/VirtualProduction/Avalanche/Source/AvalancheSequencer/Private/EaseCurveTool/Widgets/SAvaEaseCurvePresetGroupItem.h
@@ -0,0 +1,87 @@
+// Copyright Epic Games, Inc. All Rights Reserved.
+
+#pragma once
+
+#include "Misc/FrameRate.h"
+#include "Templates/SharedPointer.h"
+#include "Widgets/Views/SListView.h"
+#include "Widgets/Views/STableRow.h"
+
+class FText;
+class SEditableTextBox;
+class STableViewBase;
+struct FAvaEaseCurvePreset;
+
+DECLARE_DELEGATE_RetVal_OneParam(bool, FAvaEaseCurvePresetDelegate, const TSharedPtr<FAvaEaseCurvePreset>& /*InPreset*/)
+DECLARE_DELEGATE_RetVal_OneParam(bool, FAvaEaseCurvePresetClickDelegate, const TSharedPtr<FAvaEaseCurvePreset>& /*InPreset*/)
+DECLARE_DELEGATE_RetVal_TwoParams(bool, FAvaEaseCurvePresetRenameDelegate, const TSharedPtr<FAvaEaseCurvePreset>& /*InPreset*/, const FString& /*InNewName*/)
+DECLARE_DELEGATE_RetVal_TwoParams(bool, FAvaEaseCurvePresetMoveDelegate, const TSharedPtr<FAvaEaseCurvePreset>& /*InPreset*/, const FString& /*InNewCategoryName*/)
+
+class SAvaEaseCurvePresetGroupItem : public STableRow<TSharedPtr<FAvaEaseCurvePreset>>
+{
+public:
+	SLATE_BEGIN_ARGS(SAvaEaseCurvePresetGroupItem)
+		: _DisplayRate(FFrameRate(30, 1))
+	{}
+		SLATE_ARGUMENT(TSharedPtr<FAvaEaseCurvePreset>, Preset)
+		SLATE_ATTRIBUTE(bool, IsEditMode)
+		SLATE_ARGUMENT(FFrameRate, DisplayRate)
+		SLATE_ATTRIBUTE(bool, IsSelected)
+		SLATE_EVENT(FAvaEaseCurvePresetClickDelegate, OnClick)
+		SLATE_EVENT(FAvaEaseCurvePresetDelegate, OnDelete)
+		SLATE_EVENT(FAvaEaseCurvePresetRenameDelegate, OnRename)
+		SLATE_EVENT(FAvaEaseCurvePresetDelegate, OnSetQuickEase)
+		SLATE_EVENT(FAvaEaseCurvePresetMoveDelegate, OnBeginMove)
+		SLATE_EVENT(FAvaEaseCurvePresetMoveDelegate, OnEndMove)
+	SLATE_END_ARGS()
+
+	void Construct(const FArguments& InArgs, const TSharedPtr<STableViewBase>& InOwnerTableView);
+
+	void SetPreset(const TSharedPtr<FAvaEaseCurvePreset>& InPreset);
+
+	bool IsEditMode() const;
+
+	void TriggerBeginMove();
+	void TriggerEndMove();
+
+protected:
+	void HandleRenameTextCommitted(const FText& InNewText, ETextCommit::Type InCommitType) const;
+
+	FReply HandleDeleteClick() const;
+
+	EVisibility GetEditModeVisibility() const;
+
+	EVisibility GetBorderVisibility() const;
+	const FSlateBrush* GetBackgroundImage() const;
+
+	FSlateColor GetQuickPresetIconColor() const;
+	EVisibility GetQuickPresetIconVisibility() const;
+	FText GetQuickPresetIconToolTip() const;
+
+	bool IsQuickEasePreset() const;
+
+	FReply HandleSetQuickEase();
+
+	//~ Begin SWidget
+	virtual FReply OnMouseButtonDown(const FGeometry& InGeometry, const FPointerEvent& InMouseEvent) override;
+	virtual FReply OnMouseButtonUp(const FGeometry& InGeometry, const FPointerEvent& InMouseEvent) override;
+	virtual FReply OnDragDetected(const FGeometry& InGeometry, const FPointerEvent& InMouseEvent) override;
+	virtual FReply OnDrop(const FGeometry& InGeometry, const FDragDropEvent& InDragDropEvent) override;
+	//~ End SWidget
+
+	TSharedPtr<FAvaEaseCurvePreset> Preset;
+	TAttribute<bool> bIsEditMode;
+	TAttribute<bool> IsSelected;
+	
+	FAvaEaseCurvePresetClickDelegate OnClick;
+	FAvaEaseCurvePresetDelegate OnDelete;
+	FAvaEaseCurvePresetRenameDelegate OnRename;
+	FAvaEaseCurvePresetDelegate OnSetQuickEase;
+	
+	FAvaEaseCurvePresetMoveDelegate OnBeginMove;
+	FAvaEaseCurvePresetMoveDelegate OnEndMove;
+
+	TSharedPtr<SEditableTextBox> RenameTextBox;
+
+	bool bIsDragging = false;
+};
diff --git a/Engine/Plugins/VirtualProduction/Avalanche/Source/AvalancheSequencer/Private/EaseCurveTool/Widgets/SAvaEaseCurvePreview.cpp b/Engine/Plugins/VirtualProduction/Avalanche/Source/AvalancheSequencer/Private/EaseCurveTool/Widgets/SAvaEaseCurvePreview.cpp
new file mode 100644
--- /dev/null
+++ b/Engine/Plugins/VirtualProduction/Avalanche/Source/AvalancheSequencer/Private/EaseCurveTool/Widgets/SAvaEaseCurvePreview.cpp
@@ -0,0 +1,369 @@
+// Copyright Epic Games, Inc. All Rights Reserved.
+
+#include "SAvaEaseCurvePreview.h"
+#include "EaseCurveTool/AvaEaseCurveTool.h"
+#include "EaseCurveTool/Widgets/SAvaEaseCurvePreviewToolTip.h"
+#include "Curves/CurveEvaluation.h"
+#include "HAL/PlatformApplicationMisc.h"
+#include "Kismet/KismetMathLibrary.h"
+#include "SCurveEditor.h"
+#include "Styling/AppStyle.h"
+#include "Widgets/Accessibility/SlateCoreAccessibleWidgets.h"
+
+#define LOCTEXT_NAMESPACE "SAvaEaseCurvePreview"
+
+void SAvaEaseCurvePreview::Construct(const FArguments& InArgs)
+{ 
+	SetCanTick(InArgs._Animate);
+
+	Tangents = InArgs._Tangents;
+	PreviewSize = InArgs._PreviewSize;
+	CanExpandPreview = InArgs._CanExpandPreview;
+	CurveThickness = InArgs._CurveThickness;
+	BackgroundColor = InArgs._BackgroundColor;
+	CurveColor = InArgs._CurveColor;
+	StraightColor = InArgs._StraightColor;
+	UnderCurveColor = InArgs._UnderCurveColor;
+	bAnimate = InArgs._Animate;
+	IdleAnimationLength = InArgs._IdleAnimationLength;
+	DisplayRate = InArgs._DisplayRate;
+	AnimationColor = InArgs._AnimationColor;
+	AnimationSize = InArgs._AnimationSize;
+	bDrawMotionTrails = InArgs._DrawMotionTrails;
+	MotionTrailColor = InArgs._MotionTrailColor;
+	MotionTrailColorFade = InArgs._MotionTrailColorFade;
+	MotionTrailSize = InArgs._MotionTrailSize;
+	MotionTrailFadeLength = InArgs._MotionTrailFadeLength;
+	MotionTrailOffset = InArgs._MotionTrailOffset;
+
+	if (InArgs._CustomToolTip)
+	{
+		ToolTipWidget = MakeShared<SAvaEaseCurvePreviewToolTip>(InArgs);
+		SetToolTip(ToolTipWidget);
+	}
+	else
+	{
+		SetToolTipText(SAvaEaseCurvePreviewToolTip::GetToolTipText(Tangents.Get(FAvaEaseCurveTangents())));
+	}
+
+	SetTangents(InArgs._Tangents);
+}
+
+FAvaEaseCurveTangents SAvaEaseCurvePreview::GetTangents() const
+{
+	return Tangents.Get(FAvaEaseCurveTangents());
+}
+
+void SAvaEaseCurvePreview::SetTangents(const TAttribute<FAvaEaseCurveTangents>& InTangents)
+{
+	Tangents = InTangents;
+
+	// Save min and max to compute extra screen space if the curve goes beyond normal limits (0 - 1).
+	float MinInputValue = 0.f;
+	float MaxInputValue = 1.f;
+
+	// Cache motion trails
+	const FAvaEaseCurveTangents DrawTangents = Tangents.Get(FAvaEaseCurveTangents());
+	const float SecondsPerFrame = DisplayRate.AsInterval();
+	const int32 FrameCount = DisplayRate.AsFrameNumber(1.f).Value;
+
+	MotionTrails.Empty(FrameCount + 1);
+
+	for (float Time = 0; Time < 1.f; Time += SecondsPerFrame)
+	{
+		const float CurrentValue = WeightedEvalForTwoKeys(DrawTangents, Time);
+
+		FMotionTrail NewMotionTrail;
+		NewMotionTrail.Time = Time;
+		NewMotionTrail.Color = FLinearColor(1.f, 1.f, 1.f, 1.f);
+		NewMotionTrail.Location = FVector2D(CurrentValue, 0.f);
+		MotionTrails.Add(MoveTemp(NewMotionTrail));
+
+		MinInputValue = FMath::Min(CurrentValue, MinInputValue);
+		MaxInputValue = FMath::Max(CurrentValue, MaxInputValue);
+	}
+
+	FMotionTrail LastMotionTrail;
+	LastMotionTrail.Time = 1.f;
+	LastMotionTrail.Color = FLinearColor(1.f, 1.f, 1.f, 1.f);
+	LastMotionTrail.Location = FVector2D(1.f, 0.f);
+	MotionTrails.Add(MoveTemp(LastMotionTrail));
+
+	BelowZeroValue = FMath::Abs(MinInputValue);
+	AboveOneValue = MaxInputValue - 1.f;
+}
+
+FVector2D SAvaEaseCurvePreview::ComputeDesiredSize(const float InLayoutScaleMultiplier) const
+{
+	float AdditionalHeight = CanExpandPreview ? (BelowZeroValue + AboveOneValue) : 0.f;
+	AdditionalHeight = FMath::Min(PreviewSize, AdditionalHeight);
+	return FVector2D(PreviewSize) + FVector2D(AdditionalHeight * PreviewSize);
+}
+
+FReply SAvaEaseCurvePreview::OnMouseButtonDown(const FGeometry& InMyGeometry, const FPointerEvent& InMouseEvent)
+{
+	if (InMouseEvent.GetEffectingButton() == EKeys::RightMouseButton)
+	{
+		FPlatformApplicationMisc::ClipboardCopy(*Tangents.Get(FAvaEaseCurveTangents()).ToJson());
+
+		FAvaEaseCurveTool::ShowNotificationMessage(LOCTEXT("EaseCurveToolTangentsCopied", "Ease Curve Tool Tangents Copied!"));
+
+		return FReply::Handled();
+	}
+
+	return SLeafWidget::OnMouseButtonDown(InMyGeometry, InMouseEvent);
+}
+
+void SAvaEaseCurvePreview::Tick(const FGeometry& InAllottedGeometry, const double InCurrentTime, const float InDeltaTime)
+{
+	SLeafWidget::Tick(InAllottedGeometry, InCurrentTime, InDeltaTime);
+
+	if (IdleTimerHandle.IsValid()) // Fade out all animations
+	{
+		CurrentAnimateTime = 1.f;
+
+		// Fade out animation and motion trails
+		const double IdlePlayTime = FMath::Min(IdleAnimationLength, InCurrentTime - MotionTrailFadeStartTime);
+		const double NewOpacity = UKismetMathLibrary::MapRangeClamped(IdlePlayTime, 0.f, IdleAnimationLength, 1.f, 0.f);
+
+		AnimationColor.A = NewOpacity;
+
+		if (bDrawMotionTrails)
+		{
+			const FLinearColor LerpColor = FLinearColor::LerpUsingHSV(FLinearColor(0.f, 0.f, 0.f, 0.f), AnimationColor, NewOpacity);
+
+			for (FMotionTrail& MotionTrail : MotionTrails)
+			{
+				MotionTrail.Color = LerpColor;
+			}
+		}
+	}
+	else // Animate curve along path
+	{
+		if (CurrentAnimateTime >= 1.f)
+		{
+			CurrentAnimateTime = 1.f;
+			MotionTrailFadeStartTime = InCurrentTime;
+
+			IdleTimerHandle = RegisterActiveTimer(IdleAnimationLength, FWidgetActiveTimerDelegate::CreateLambda([this](double InCurrentTime, float InDeltaTime) -> EActiveTimerReturnType
+				{
+					IdleTimerHandle.Reset();
+
+					CurrentAnimateTime = 0.f;
+
+					if (bDrawMotionTrails)
+					{
+						for (FMotionTrail& MotionTrail : MotionTrails)
+						{
+							MotionTrail.Color = FLinearColor::White;
+						}
+					}
+
+					return EActiveTimerReturnType::Stop;
+				}));
+		}
+		else
+		{
+			CurrentAnimateTime += InDeltaTime;
+			CurrentAnimateTime = FMath::Min(1.f, CurrentAnimateTime);
+
+			const FTrackScaleInfo ScaleInfo(0.f, 1.f, 0.f, 1.f, FVector2f(PreviewSize)/*InAllottedGeometry.GetLocalSize()*/);
+			const FAvaEaseCurveTangents DrawTangents = Tangents.Get(FAvaEaseCurveTangents());
+			const float CurrentValue = WeightedEvalForTwoKeys(DrawTangents, CurrentAnimateTime);
+
+			CurrentAnimateLocation = FVector2D(ScaleInfo.InputToLocalX(CurrentAnimateTime), ScaleInfo.OutputToLocalY(CurrentValue));
+
+			AnimationColor.A = 1.f;
+
+			if (bDrawMotionTrails)
+			{
+				for (FMotionTrail& MotionTrail : MotionTrails)
+				{
+					const float TimeSinceTrail = CurrentAnimateTime - MotionTrail.Time;
+					const float ClampedTimeSinceTrail = FMath::Min(MotionTrailFadeLength, TimeSinceTrail);
+					const float FadeOutAlpha = TimeSinceTrail <= MotionTrailFadeLength
+						? UKismetMathLibrary::MapRangeClamped(ClampedTimeSinceTrail, 0.f, MotionTrailFadeLength, 0.f, 1.f)
+						: 1.f;
+
+					FLinearColor NewColor = FLinearColor::LerpUsingHSV(MotionTrailColor, MotionTrailColorFade, FadeOutAlpha);
+					NewColor.A = MotionTrail.Color.A;
+
+					MotionTrail.Color = NewColor;
+				}
+			}
+		}
+	}
+}
+
+#if WITH_ACCESSIBILITY
+TSharedRef<FSlateAccessibleWidget> SAvaEaseCurvePreview::CreateAccessibleWidget()
+{
+	return MakeShareable<FSlateAccessibleWidget>(new FSlateAccessibleImage(SharedThis(this)));
+}
+#endif
+
+int32 SAvaEaseCurvePreview::OnPaint(const FPaintArgs& InArgs, const FGeometry& InAllottedGeometry, const FSlateRect& InMyCullingRect, FSlateWindowElementList& OutDrawElements, int32 InLayerId, const FWidgetStyle& InWidgetStyle, bool bInParentEnabled) const
+{
+	const bool bEnabled = ShouldBeEnabled(bInParentEnabled);
+	const ESlateDrawEffect DrawEffects = bEnabled ? ESlateDrawEffect::None : ESlateDrawEffect::DisabledEffect;
+	const FVector2D LocalSize = InAllottedGeometry.GetLocalSize();
+	const FTrackScaleInfo ScaleInfo(0.f, 1.f, 0.f, 1.f, FVector2f(PreviewSize));
+
+	const float ZeroOutputY = ScaleInfo.OutputToLocalY(0.f);
+	const float OneOutputY = ScaleInfo.OutputToLocalY(1.f);
+
+	int32 LayerId = InLayerId;
+
+	const FVector2D AdditionalAboveHeight(BelowZeroValue * PreviewSize, AboveOneValue * PreviewSize);
+	const FPaintGeometry AbovePaintGeometry = InAllottedGeometry.ToPaintGeometry(InAllottedGeometry.GetLocalSize(), FSlateLayoutTransform(AdditionalAboveHeight));
+
+	// Draw background color
+	if (BackgroundColor.IsSet())
+	{
+		FSlateDrawElement::MakeBox(OutDrawElements, LayerId++, InAllottedGeometry.ToPaintGeometry()
+			, FAppStyle::GetBrush(TEXT("WhiteTexture")), DrawEffects, BackgroundColor.GetValue());
+	}
+
+	// Draw 0 and 1 line
+	{
+		constexpr FLinearColor NormalLimitLineColor = FLinearColor(0.1f, 0.1f, 0.1f, 1.f);
+		const TArray<FLinearColor> ValueLineColors = { FLinearColor::Transparent, NormalLimitLineColor, NormalLimitLineColor, FLinearColor::Transparent};
+		const float LocalWidth = InAllottedGeometry.GetLocalSize().X;
+		const float OneFifthWidth = LocalWidth * 0.2f;
+
+		// Value = 0 line
+		if (BelowZeroValue > 0.f)
+		{
+			TArray<FVector2D> ZeroValueLinePoints;
+			ZeroValueLinePoints.Add(FVector2D(0.f, ZeroOutputY));
+			ZeroValueLinePoints.Add(FVector2D(OneFifthWidth, ZeroOutputY));
+			ZeroValueLinePoints.Add(FVector2D(LocalWidth - OneFifthWidth, ZeroOutputY));
+			ZeroValueLinePoints.Add(FVector2D(LocalWidth, ZeroOutputY));
+			TArray<FLinearColor> ZeroValueLineColors;
+			FSlateDrawElement::MakeLines(OutDrawElements, LayerId++, AbovePaintGeometry
+				, ZeroValueLinePoints, ValueLineColors, DrawEffects, FLinearColor::White, false, InAllottedGeometry.Scale);
+		}
+
+		// Value = 1 line
+		if (AboveOneValue > 0.f)
+		{
+			TArray<FVector2D> OneValueLinePoints;
+			OneValueLinePoints.Add(FVector2D(0.f, OneOutputY));
+			OneValueLinePoints.Add(FVector2D(OneFifthWidth, OneOutputY));
+			OneValueLinePoints.Add(FVector2D(LocalWidth - OneFifthWidth, OneOutputY));
+			OneValueLinePoints.Add(FVector2D(LocalWidth, OneOutputY));
+			FSlateDrawElement::MakeLines(OutDrawElements, LayerId++, AbovePaintGeometry
+				, OneValueLinePoints, ValueLineColors, DrawEffects, FLinearColor::White, false, InAllottedGeometry.Scale);
+		}
+	}
+
+	TArray<FVector2D> CurvePoints;
+	TArray<FLinearColor> CurveColors;
+	MakeCurveData(InAllottedGeometry, CurvePoints, CurveColors);
+
+	// Draw under curve fill
+	if (UnderCurveColor.IsSet())
+	{
+		for (const FVector2D& Point : CurvePoints)
+		{
+			const TArray<FVector2D> LinePoints = { FVector2D(Point.X, LocalSize.Y), Point };
+			const TArray<FLinearColor> LineColors = { UnderCurveColor.GetValue(), UnderCurveColor.GetValue() };
+
+			FSlateDrawElement::MakeLines(OutDrawElements, LayerId++, AbovePaintGeometry,
+				LinePoints, LineColors, DrawEffects, FLinearColor::White, true, 2.0f * InAllottedGeometry.Scale);
+		}
+	}
+
+	// Draw curve
+	FSlateDrawElement::MakeLines(OutDrawElements, LayerId++, AbovePaintGeometry,
+		CurvePoints, CurveColors, DrawEffects, FLinearColor::White, true, CurveThickness * InAllottedGeometry.Scale);
+
+	// Draw animation
+	if (bAnimate)
+	{
+		const FSlateBrush* ImageBrush = FAppStyle::GetBrush(TEXT("Icons.BulletPoint"));
+
+		// Draw past motion trails
+		if (bDrawMotionTrails)
+		{
+			for (const FMotionTrail& MotionTrail : MotionTrails)
+			{
+				if (MotionTrail.Time >= CurrentAnimateTime)
+				{
+					continue;
+				}
+
+				FVector2D MotionTrailLocation(ScaleInfo.InputToLocalX(MotionTrail.Location.X), ScaleInfo.OutputToLocalY(MotionTrail.Location.Y));
+				MotionTrailLocation -= FVector2D(MotionTrailSize * 0.5f);
+				MotionTrailLocation += MotionTrailOffset;
+				MotionTrailLocation += AdditionalAboveHeight;
+
+				FSlateDrawElement::MakeBox(OutDrawElements, LayerId++
+					, InAllottedGeometry.ToPaintGeometry(FVector2D(MotionTrailSize), FSlateLayoutTransform(MotionTrailLocation))
+					, ImageBrush, DrawEffects, MotionTrail.Color * InWidgetStyle.GetColorAndOpacityTint());
+			}
+		}
+
+		// Draw current animation point
+		const FVector2D TangentIconLocation = CurrentAnimateLocation - (AnimationSize * 0.5f);
+
+		FSlateDrawElement::MakeBox(OutDrawElements, LayerId++
+			, InAllottedGeometry.ToPaintGeometry(FVector2D(AnimationSize), FSlateLayoutTransform(TangentIconLocation + AdditionalAboveHeight))
+			, ImageBrush, DrawEffects, AnimationColor * InWidgetStyle.GetColorAndOpacityTint());
+	}
+
+	return LayerId;
+}
+
+void SAvaEaseCurvePreview::MakeCurveData(const FGeometry& InGeometry, TArray<FVector2D>& OutCurvePoints, TArray<FLinearColor>& OutCurveColors) const
+{
+	const FAvaEaseCurveTangents DrawTangents = Tangents.Get(FAvaEaseCurveTangents());
+	const FTrackScaleInfo ScaleInfo(0.f, 1.f, 0.f, 1.f, FVector2f(PreviewSize)/*InAllottedGeometry.GetLocalSize()InGeometry.GetLocalSize()*/);
+
+	// Clamp to screen to avoid massive slowdown when zoomed in
+	float StartX = FMath::Max(ScaleInfo.InputToLocalX(0.f), 0.0f);
+	const float EndX = FMath::Min(ScaleInfo.InputToLocalX(1.f), ScaleInfo.WidgetSize.X);
+	
+	constexpr float StepSize = 1.0f;
+	float LastOutputY = 0.0f;
+	
+	for (; StartX < EndX; StartX += StepSize)
+	{
+		const float CurveIn = ScaleInfo.LocalXToInput(FMath::Min(StartX, EndX));
+		const float CurveOut = WeightedEvalForTwoKeys(DrawTangents, CurveIn);
+
+		const FVector2D ScreenPoint(ScaleInfo.InputToLocalX(CurveIn), ScaleInfo.OutputToLocalY(CurveOut));
+		OutCurvePoints.Add(ScreenPoint);
+
+		const float NewCurveStrength = CurveOut - LastOutputY;
+		LastOutputY = CurveOut;
+
+		const float Alpha = UKismetMathLibrary::MapRangeClamped(FMath::Abs(NewCurveStrength), 0.f, 0.03f, 0.f, 1.f);
+		OutCurveColors.Add(FLinearColor::LerpUsingHSV(StraightColor, CurveColor, Alpha));
+	}
+}
+
+float SAvaEaseCurvePreview::WeightedEvalForTwoKeys(const FAvaEaseCurveTangents& InTangents, const float InTime) const
+{
+	FRichCurveKey StartKey(0.f, 0.f);
+	StartKey.InterpMode = ERichCurveInterpMode::RCIM_Cubic;
+	StartKey.TangentMode = ERichCurveTangentMode::RCTM_Break;
+	StartKey.TangentWeightMode = ERichCurveTangentWeightMode::RCTWM_WeightedLeave;
+	StartKey.ArriveTangent = 0.f;
+	StartKey.ArriveTangentWeight = 0.f;
+	StartKey.LeaveTangent = InTangents.Start;
+	StartKey.LeaveTangentWeight = InTangents.StartWeight;
+
+	FRichCurveKey EndKey(1.f, 1.f);
+	EndKey.InterpMode = ERichCurveInterpMode::RCIM_Cubic;
+	EndKey.TangentMode = ERichCurveTangentMode::RCTM_Break;
+	EndKey.TangentWeightMode = ERichCurveTangentWeightMode::RCTWM_WeightedArrive;
+	EndKey.ArriveTangent = InTangents.End;
+	EndKey.ArriveTangentWeight = InTangents.EndWeight;
+	EndKey.LeaveTangent = 0.f;
+	EndKey.LeaveTangentWeight = 0.f;
+	
+	return UE::Curves::EvalForTwoKeys(StartKey, EndKey, InTime);
+}
+
+#undef LOCTEXT_NAMESPACE
diff --git a/Engine/Plugins/VirtualProduction/Avalanche/Source/AvalancheSequencer/Private/EaseCurveTool/Widgets/SAvaEaseCurvePreview.h b/Engine/Plugins/VirtualProduction/Avalanche/Source/AvalancheSequencer/Private/EaseCurveTool/Widgets/SAvaEaseCurvePreview.h
new file mode 100644
--- /dev/null
+++ b/Engine/Plugins/VirtualProduction/Avalanche/Source/AvalancheSequencer/Private/EaseCurveTool/Widgets/SAvaEaseCurvePreview.h
@@ -0,0 +1,135 @@
+// Copyright Epic Games, Inc. All Rights Reserved.
+
+#pragma once
+
+#include "EaseCurveTool/Widgets/SAvaEaseCurvePreset.h"
+#include "Misc/Optional.h"
+#include "Styling/SlateTypes.h"
+#include "Styling/StyleColors.h"
+#include "Templates/SharedPointer.h"
+#include "Widgets/DeclarativeSyntaxSupport.h"
+#include "Widgets/SLeafWidget.h"
+
+class FPaintArgs;
+class FSlateWindowElementList;
+class SAvaEaseCurvePreviewToolTip;
+struct FGeometry;
+
+class SAvaEaseCurvePreview : public SLeafWidget
+{
+public:
+	SLATE_BEGIN_ARGS(SAvaEaseCurvePreview)
+		: _Tangents()
+		, _PreviewSize(32.f)
+		, _CanExpandPreview(false)
+		, _CurveThickness(2.f)
+		, _CurveColor(FLinearColor::White)
+		, _StraightColor(FStyleColors::AccentBlue.GetSpecifiedColor())
+		, _CustomToolTip(false)
+		, _Animate(false)
+		, _IdleAnimationLength(1.f)
+		, _DisplayRate(FFrameRate(30, 1))
+		, _AnimationColor(FLinearColor::White)
+		, _AnimationSize(12.f)
+		, _DrawMotionTrails(false)
+		, _MotionTrailColor(FLinearColor(0.5f, 0.5f, 0.5f, 1.f))
+		, _MotionTrailColorFade(FLinearColor(0.15f, 0.15f, 0.15f, 1.f))
+		, _MotionTrailSize(10.f)
+		, _MotionTrailFadeLength(0.2f)
+		, _MotionTrailOffset(FVector2D(0.f, 10.f))
+	{}
+		SLATE_ATTRIBUTE(FAvaEaseCurveTangents, Tangents)
+		SLATE_ARGUMENT(float, PreviewSize)
+		/** If true, will expand the preview size to expand for curves that go outside of the normalized range. */
+		SLATE_ARGUMENT(bool, CanExpandPreview)
+		SLATE_ARGUMENT(float, CurveThickness)
+	
+		SLATE_ARGUMENT(TOptional<FLinearColor>, BackgroundColor)
+		SLATE_ARGUMENT(FLinearColor, CurveColor)
+		SLATE_ARGUMENT(FLinearColor, StraightColor)
+		SLATE_ARGUMENT(TOptional<FLinearColor>, UnderCurveColor)
+
+		/** If true, will show a custom tooltip showing an animated preview of the curve. */
+		SLATE_ARGUMENT(bool, CustomToolTip)
+
+		SLATE_ARGUMENT(bool, Animate)
+		SLATE_ARGUMENT(bool, FlatAnimation)
+		SLATE_ARGUMENT(float, IdleAnimationLength)
+		SLATE_ARGUMENT(FFrameRate, DisplayRate)
+		SLATE_ARGUMENT(FLinearColor, AnimationColor)
+		SLATE_ARGUMENT(float, AnimationSize)
+		SLATE_ARGUMENT(bool, DrawMotionTrails)
+		SLATE_ARGUMENT(FLinearColor, MotionTrailColor)
+		SLATE_ARGUMENT(FLinearColor, MotionTrailColorFade)
+		SLATE_ARGUMENT(float, MotionTrailSize)
+		SLATE_ARGUMENT(float, MotionTrailFadeLength)
+		SLATE_ARGUMENT(FVector2D, MotionTrailOffset)
+	SLATE_END_ARGS()
+
+	void Construct(const FArguments& InArgs);
+
+	FAvaEaseCurveTangents GetTangents() const;
+	void SetTangents(const TAttribute<FAvaEaseCurveTangents>& InTangents);
+
+	//~ Begin SWidget
+#if WITH_ACCESSIBILITY
+	virtual TSharedRef<FSlateAccessibleWidget> CreateAccessibleWidget() override;
+#endif
+	virtual int32 OnPaint(const FPaintArgs& InArgs, const FGeometry& InAllottedGeometry, const FSlateRect& InMyCullingRect, FSlateWindowElementList& OutDrawElements, int32 InLayerId, const FWidgetStyle& InWidgetStyle, bool bInParentEnabled) const override;
+	//~ End SWidget
+
+protected:
+	void MakeCurveData(const FGeometry& InGeometry, TArray<FVector2D>& OutCurvePoints, TArray<FLinearColor>& OutCurveColors) const;
+
+	float WeightedEvalForTwoKeys(const FAvaEaseCurveTangents& InTangents, const float InTime) const;
+
+	//~ Begin SWidget
+	virtual FVector2D ComputeDesiredSize(const float InLayoutScaleMultiplier) const override;
+	virtual FReply OnMouseButtonDown(const FGeometry& InMyGeometry, const FPointerEvent& InMouseEvent) override;
+	virtual void Tick(const FGeometry& InAllottedGeometry, const double InCurrentTime, const float InDeltaTime) override;
+	//~ End SWidget
+
+	TAttribute<FAvaEaseCurveTangents> Tangents;
+	float PreviewSize = 0.f;
+	bool CanExpandPreview = false;
+	float CurveThickness = 1.f;
+
+	TOptional<FLinearColor> BackgroundColor;
+	FLinearColor CurveColor;
+	FLinearColor StraightColor;
+	TOptional<FLinearColor> UnderCurveColor;
+
+	bool bAnimate = false;
+	float IdleAnimationLength = 0.f;
+	FFrameRate DisplayRate;
+	FLinearColor AnimationColor;
+	float AnimationSize = 0.f;
+	bool bDrawMotionTrails = false;
+	FLinearColor MotionTrailColor;
+	FLinearColor MotionTrailColorFade;
+	float MotionTrailSize = 0.f;
+	float MotionTrailFadeLength = 0.f;
+	FVector2D MotionTrailOffset;
+
+	TSharedPtr<SAvaEaseCurvePreviewToolTip> ToolTipWidget;
+
+	/** The amount the curve is below value 0 in curve space. */
+	float BelowZeroValue = 0.f;
+	/** The amount the curve is above value 1 in curve space. */
+	float AboveOneValue = 0.f;
+
+	TSharedPtr<FActiveTimerHandle> IdleTimerHandle;
+
+	double CurrentAnimateTime = 0.f;
+	FVector2D CurrentAnimateLocation = FVector2D::ZeroVector;
+
+	struct FMotionTrail
+	{
+		FVector2D Location = FVector2D::ZeroVector;
+		float Time = 0.f;
+		FLinearColor Color = FLinearColor(1.f, 1.f, 1.f, 1.f);
+	};
+	TArray<FMotionTrail> MotionTrails;
+
+	double MotionTrailFadeStartTime = 0.f;
+};
diff --git a/Engine/Plugins/VirtualProduction/Avalanche/Source/AvalancheSequencer/Private/EaseCurveTool/Widgets/SAvaEaseCurvePreviewToolTip.cpp b/Engine/Plugins/VirtualProduction/Avalanche/Source/AvalancheSequencer/Private/EaseCurveTool/Widgets/SAvaEaseCurvePreviewToolTip.cpp
new file mode 100644
--- /dev/null
+++ b/Engine/Plugins/VirtualProduction/Avalanche/Source/AvalancheSequencer/Private/EaseCurveTool/Widgets/SAvaEaseCurvePreviewToolTip.cpp
@@ -0,0 +1,147 @@
+// Copyright Epic Games, Inc. All Rights Reserved.
+
+#include "SAvaEaseCurvePreviewToolTip.h"
+#include "EaseCurveTool/Widgets/SAvaEaseCurvePreset.h"
+#include "Styling/AppStyle.h"
+#include "Styling/StyleColors.h"
+#include "Widgets/Colors/SColorBlock.h"
+#include "Widgets/SToolTip.h"
+
+#define LOCTEXT_NAMESPACE "SAvaEaseCurvePreviewToolTip"
+
+SAvaEaseCurvePreviewToolTip::SAvaEaseCurvePreviewToolTip(const SAvaEaseCurvePreview::FArguments& InPreviewArgs, const TSharedPtr<SWidget>& InAdditionalContent)
+	: PreviewArgs(InPreviewArgs), AdditionalContent(InAdditionalContent)
+{
+}
+
+FText SAvaEaseCurvePreviewToolTip::GetToolTipText(const FAvaEaseCurveTangents& InTangents)
+{
+	return FText::Format(LOCTEXT("FullToolTipText", "{0}\n\nRight click to copy to Json values"), InTangents.ToDisplayText());
+}
+
+TSharedRef<SToolTip> SAvaEaseCurvePreviewToolTip::CreateDefaultToolTip(const SAvaEaseCurvePreview::FArguments& InPreviewArgs, const TSharedPtr<SWidget>& InAdditionalContent)
+{
+	const TSharedRef<SVerticalBox> ContentWidget = SNew(SVerticalBox)
+		+ SVerticalBox::Slot()
+		.AutoHeight()
+		.Padding(10.f)
+		[
+			SNew(SOverlay)
+			+ SOverlay::Slot()
+			[
+				SNew(SColorBlock)
+				.Color(FStyleColors::Background.GetSpecifiedColor())
+			]
+			+ SOverlay::Slot()
+			.Padding(15.f, 15.f, 15.f, 20.f)
+			[
+				SNew(SAvaEaseCurvePreview)
+				.Tangents(InPreviewArgs._Tangents)
+				.PreviewSize(256.f)
+				.CanExpandPreview(true)
+				.CurveThickness(InPreviewArgs._CurveThickness)
+				.CurveColor(InPreviewArgs._CurveColor)
+				.CustomToolTip(false)
+				.Animate(true)
+				.DisplayRate(InPreviewArgs._DisplayRate)
+				.DrawMotionTrails(true)
+			]
+		]
+		+ SVerticalBox::Slot()
+		.AutoHeight()
+		.Padding(15.f, 2.f, 15.f, 0.f)
+		[
+			SNew(SHorizontalBox)
+			+ SHorizontalBox::Slot()
+			.HAlign(HAlign_Left)
+			[
+				SNew(STextBlock)
+				.ToolTipText(LOCTEXT("StartTangentTooltipText", "Start tangent and weight"))
+				.Text(InPreviewArgs._Tangents.Get(FAvaEaseCurveTangents()).GetStartTangentText())
+			]
+			+ SHorizontalBox::Slot()
+			.HAlign(HAlign_Right)
+			[
+				SNew(STextBlock)
+				.ToolTipText(LOCTEXT("EndTangentTooltipText", "End tangent and weight"))
+				.Text(InPreviewArgs._Tangents.Get(FAvaEaseCurveTangents()).GetEndTangentText())
+			]
+		]
+		+ SVerticalBox::Slot()
+		.AutoHeight()
+		.Padding(15.f, 5.f, 15.f, 0.f)
+		[
+			SNew(SHorizontalBox)
+			+ SHorizontalBox::Slot()
+			.HAlign(HAlign_Left)
+			[
+				SNew(STextBlock)
+				.Text(LOCTEXT("CubicBezierText", "Cubic Bezier Points:"))
+			]
+			+ SHorizontalBox::Slot()
+			.HAlign(HAlign_Right)
+			[
+				SNew(STextBlock)
+				.ToolTipText(LOCTEXT("CubicBezierTooltipText", "Cubic bezier points"))
+				.Text(InPreviewArgs._Tangents.Get(FAvaEaseCurveTangents()).GetCubicBezierText())
+			]
+		]
+		+ SVerticalBox::Slot()
+		.AutoHeight()
+		.Padding(20.f, 10.f, 0.f, 10.f)
+		[
+			SNew(STextBlock)
+			.TextStyle(FAppStyle::Get(), TEXT("HintText"))
+			.Text(LOCTEXT("ToolTipText", "Right click to copy to tangent values"))
+		];
+
+	if (InAdditionalContent.IsValid())
+	{
+		ContentWidget->AddSlot()
+			.Padding(15.f, 10.f, 15.f, 0.f)
+			[
+				InAdditionalContent.ToSharedRef()
+			];
+	}
+
+	return SNew(SToolTip)
+		.TextMargin(1.0f)
+		[
+			ContentWidget
+		];
+}
+
+void SAvaEaseCurvePreviewToolTip::CreateToolTipWidget()
+{
+	if (ToolTipWidget.IsValid())
+	{
+		return;
+	}
+
+	ToolTipWidget = CreateDefaultToolTip(PreviewArgs, AdditionalContent);
+}
+
+TSharedRef<SWidget> SAvaEaseCurvePreviewToolTip::AsWidget()
+{
+	CreateToolTipWidget();
+	return ToolTipWidget.ToSharedRef();
+}
+
+TSharedRef<SWidget> SAvaEaseCurvePreviewToolTip::GetContentWidget()
+{
+	CreateToolTipWidget();
+	return ToolTipWidget->GetContentWidget();
+}
+
+void SAvaEaseCurvePreviewToolTip::SetContentWidget(const TSharedRef<SWidget>& InContentWidget)
+{
+	CreateToolTipWidget();
+	ToolTipWidget->SetContentWidget(InContentWidget);
+}
+
+void SAvaEaseCurvePreviewToolTip::InvalidateWidget()
+{
+	ToolTipWidget.Reset();
+}
+
+#undef LOCTEXT_NAMESPACE
diff --git a/Engine/Plugins/VirtualProduction/Avalanche/Source/AvalancheSequencer/Private/EaseCurveTool/Widgets/SAvaEaseCurvePreviewToolTip.h b/Engine/Plugins/VirtualProduction/Avalanche/Source/AvalancheSequencer/Private/EaseCurveTool/Widgets/SAvaEaseCurvePreviewToolTip.h
new file mode 100644
--- /dev/null
+++ b/Engine/Plugins/VirtualProduction/Avalanche/Source/AvalancheSequencer/Private/EaseCurveTool/Widgets/SAvaEaseCurvePreviewToolTip.h
@@ -0,0 +1,41 @@
+// Copyright Epic Games, Inc. All Rights Reserved.
+
+#pragma once
+
+#include "EaseCurveTool/Widgets/SAvaEaseCurvePreview.h"
+#include "Templates/SharedPointer.h"
+#include "Widgets/IToolTip.h"
+
+class SToolTip;
+class SWidget;
+
+class SAvaEaseCurvePreviewToolTip : public IToolTip
+{
+public:
+	static FText GetToolTipText(const FAvaEaseCurveTangents& InTangents);
+
+	static TSharedRef<SToolTip> CreateDefaultToolTip(const SAvaEaseCurvePreview::FArguments& InPreviewArgs, const TSharedPtr<SWidget>& InAdditionalContent = nullptr);
+
+	SAvaEaseCurvePreviewToolTip(const SAvaEaseCurvePreview::FArguments& InPreviewArgs, const TSharedPtr<SWidget>& InAdditionalContent = nullptr);
+	virtual ~SAvaEaseCurvePreviewToolTip() override {}
+
+protected:
+	void CreateToolTipWidget();
+
+	void InvalidateWidget();
+
+	//~ Begin IToolTip
+	virtual TSharedRef<SWidget> AsWidget() override;
+	virtual TSharedRef<SWidget> GetContentWidget() override;
+	virtual void SetContentWidget(const TSharedRef<SWidget>& InContentWidget) override;
+	virtual bool IsEmpty() const override { return false; }
+	virtual bool IsInteractive() const override { return false; }
+	virtual void OnOpening() override {}
+	virtual void OnClosed() override {}
+	//~ End IToolTip
+
+	SAvaEaseCurvePreview::FArguments PreviewArgs;
+	TSharedPtr<SWidget> AdditionalContent;
+
+	TSharedPtr<SToolTip> ToolTipWidget;
+};
diff --git a/Engine/Plugins/VirtualProduction/Avalanche/Source/AvalancheSequencer/Private/EaseCurveTool/Widgets/SAvaEaseCurveTool.cpp b/Engine/Plugins/VirtualProduction/Avalanche/Source/AvalancheSequencer/Private/EaseCurveTool/Widgets/SAvaEaseCurveTool.cpp
new file mode 100644
--- /dev/null
+++ b/Engine/Plugins/VirtualProduction/Avalanche/Source/AvalancheSequencer/Private/EaseCurveTool/Widgets/SAvaEaseCurveTool.cpp
@@ -0,0 +1,453 @@
+// Copyright Epic Games, Inc. All Rights Reserved.
+
+#include "SAvaEaseCurveTool.h"
+#include "Curves/KeyHandle.h"
+#include "Curves/RichCurve.h"
+#include "EaseCurveTool/AvaEaseCurvePreset.h"
+#include "EaseCurveTool/AvaEaseCurveTool.h"
+#include "EaseCurveTool/AvaEaseCurveToolSettings.h"
+#include "EaseCurveTool/AvaEaseCurveSubsystem.h"
+#include "EaseCurveTool/Widgets/SAvaEaseCurveEditor.h"
+#include "EaseCurveTool/Widgets/SAvaEaseCurvePreset.h"
+#include "EaseCurveTool/Widgets/AvaEaseCurveToolContextMenu.h"
+#include "Editor.h"
+#include "EngineAnalytics.h"
+#include "SAvaEaseCurveTangents.h"
+#include "Framework/Commands/UICommandList.h"
+#include "Framework/MultiBox/MultiBoxBuilder.h"
+#include "Math/UnrealMathUtility.h"
+#include "Widgets/Images/SImage.h"
+
+#define LOCTEXT_NAMESPACE "SAvaEaseCurveTool"
+
+void SAvaEaseCurveTool::Construct(const FArguments& InArgs, const TSharedRef<FAvaEaseCurveTool>& InEaseCurveTool)
+{
+	ToolMode = InArgs._ToolMode;
+	ToolOperation = InArgs._ToolOperation;
+
+	EaseCurveToolWeak = InEaseCurveTool;
+
+	ChildSlot
+		[
+			SNew(SVerticalBox)
+			+ SVerticalBox::Slot()
+			.AutoHeight()
+			.HAlign(HAlign_Fill)
+			.Padding(0.f, 1.f, 0.f, 0.f)
+			[
+				SAssignNew(CurvePresetWidget, SAvaEaseCurvePreset)
+				.OnPresetChanged(this, &SAvaEaseCurveTool::OnPresetChanged)
+				.OnQuickPresetChanged(this, &SAvaEaseCurveTool::OnQuickPresetChanged)
+				.OnGetNewPresetTangents(this, &SAvaEaseCurveTool::OnGetNewPresetTangents)
+			]
+			+ SVerticalBox::Slot()
+			.AutoHeight()
+			.HAlign(HAlign_Fill)
+			.Padding(0.f, 4.f, 0.f, 0.f)
+			[
+				ConstructCurveEditorPanel()
+			]
+			+ SVerticalBox::Slot()
+			.AutoHeight()
+			.HAlign(HAlign_Fill)
+			.Padding(0.f, 3.f, 0.f, 0.f)
+			[
+				SNew(SAvaEaseCurveTangents)
+				.InitialTangents(GetTangents())
+				.OnStartTangentChanged(this, &SAvaEaseCurveTool::OnStartTangentSpinBoxChanged)
+				.OnStartWeightChanged(this, &SAvaEaseCurveTool::OnStartTangentWeightSpinBoxChanged)
+				.OnEndTangentChanged(this, &SAvaEaseCurveTool::OnEndTangentSpinBoxChanged)
+				.OnEndWeightChanged(this, &SAvaEaseCurveTool::OnEndTangentWeightSpinBoxChanged)
+				.OnBeginSliderMovement(this, &SAvaEaseCurveTool::OnBeginSliderMovement)
+				.OnEndSliderMovement(this, &SAvaEaseCurveTool::OnEndSliderMovement)
+			]
+		];
+
+	if (GEditor)
+	{
+		GEditor->RegisterForUndo(this);
+	}
+
+	CurvePresetWidget->SetSelectedItem(InArgs._InitialTangents);
+}
+
+TSharedRef<SWidget> SAvaEaseCurveTool::ConstructCurveEditorPanel()
+{
+	const TSharedPtr<FAvaEaseCurveTool> EaseCurveTool = EaseCurveToolWeak.Pin();
+	check(EaseCurveTool.IsValid());
+
+	CurrentGraphSize = GetDefault<UAvaEaseCurveToolSettings>()->GetGraphSize();
+	
+	ContextMenu = MakeShared<FAvaEaseCurveToolContextMenu>(EaseCurveTool->GetCommandList()
+		, FAvaEaseCurveToolOnGraphSizeChanged::CreateLambda([this](const int32 InNewSize)
+			{
+				CurrentGraphSize = InNewSize;
+			}));
+
+	return SNew(SBorder)
+		[
+			SNew(SOverlay)
+			+ SOverlay::Slot()
+			[
+				SAssignNew(CurveEaseEditorWidget, SAvaEaseCurveEditor, EaseCurveTool->GetToolCurve())
+				.DisplayRate(this, &SAvaEaseCurveTool::GetDisplayRate)
+				.Operation(this, &SAvaEaseCurveTool::GetToolOperation)
+				.DesiredSize(this, &SAvaEaseCurveTool::GetEditorSize)
+				.ShowEqualValueKeyError(this, &SAvaEaseCurveTool::ShowEqualValueKeyError)
+				.IsEaseCurveSelection(this, &SAvaEaseCurveTool::IsEaseCurveSelection)
+				.OnTangentsChanged(this, &SAvaEaseCurveTool::HandleEditorTangentsChanged)
+				.GridSnap_UObject(GetDefault<UAvaEaseCurveToolSettings>(), &UAvaEaseCurveToolSettings::GetGridSnap)
+				.GridSize_UObject(GetDefault<UAvaEaseCurveToolSettings>(), &UAvaEaseCurveToolSettings::GetGridSize)
+				.GetContextMenuContent(ContextMenu.ToSharedRef(), &FAvaEaseCurveToolContextMenu::GenerateWidget)
+				.StartText(this, &SAvaEaseCurveTool::GetStartText)
+				.StartTooltipText(this, &SAvaEaseCurveTool::GetStartTooltipText)
+				.EndText(this, &SAvaEaseCurveTool::GetEndText)
+				.EndTooltipText(this, &SAvaEaseCurveTool::GetEndTooltipText)
+				.OnKeyDown(this, &SAvaEaseCurveTool::OnKeyDown)
+				.OnDragStart(this, &SAvaEaseCurveTool::OnEditorDragStart)
+				.OnDragEnd(this, &SAvaEaseCurveTool::OnEditorDragEnd)
+			]
+			+ SOverlay::Slot()
+			.HAlign(HAlign_Fill)
+			.VAlign(VAlign_Top)
+			[
+				SNew(SImage)
+				.Visibility(EVisibility::HitTestInvisible)
+				.ColorAndOpacity(FLinearColor(0.f, 0.f, 0.f, 0.5f))
+				.Image(&(FCoreStyle::Get().GetWidgetStyle<FScrollBoxStyle>(TEXT("ScrollBox")).TopShadowBrush))
+			]
+			+ SOverlay::Slot()
+			.HAlign(HAlign_Fill)
+			.VAlign(VAlign_Bottom)
+			[
+				SNew(SImage)
+				.Visibility(EVisibility::HitTestInvisible)
+				.ColorAndOpacity(FLinearColor(0.f, 0.f, 0.f, 0.5f))
+				.Image(&(FCoreStyle::Get().GetWidgetStyle<FScrollBoxStyle>(TEXT("ScrollBox")).BottomShadowBrush))
+			]
+			+ SOverlay::Slot()
+			.HAlign(HAlign_Left)
+			.VAlign(VAlign_Fill)
+			[
+				SNew(SImage)
+				.Visibility(EVisibility::HitTestInvisible)
+				.ColorAndOpacity(FLinearColor(0.f, 0.f, 0.f, 0.5f))
+				.Image(&(FCoreStyle::Get().GetWidgetStyle<FScrollBoxStyle>(TEXT("ScrollBox")).LeftShadowBrush))
+			]
+			+ SOverlay::Slot()
+			.HAlign(HAlign_Right)
+			.VAlign(VAlign_Fill)
+			[
+				SNew(SImage)
+				.Visibility(EVisibility::HitTestInvisible)
+				.ColorAndOpacity(FLinearColor(0.f, 0.f, 0.f, 0.5f))
+				.Image(&(FCoreStyle::Get().GetWidgetStyle<FScrollBoxStyle>(TEXT("ScrollBox")).RightShadowBrush))
+			]
+		];
+}
+
+FFrameRate SAvaEaseCurveTool::GetDisplayRate() const
+{
+	return EaseCurveToolWeak.IsValid() ? EaseCurveToolWeak.Pin()->GetDisplayRate() : FFrameRate();
+}
+
+EAvaEaseCurveToolOperation SAvaEaseCurveTool::GetToolOperation() const
+{
+	return EaseCurveToolWeak.IsValid() ? EaseCurveToolWeak.Pin()->GetToolOperation() : EAvaEaseCurveToolOperation::InOut;
+}
+
+FVector2D SAvaEaseCurveTool::GetEditorSize() const
+{
+	return FVector2D(CurrentGraphSize);
+}
+
+bool SAvaEaseCurveTool::ShowEqualValueKeyError() const
+{
+	return !EaseCurveToolWeak.Pin()->HasCachedKeysToEase();
+}
+
+bool SAvaEaseCurveTool::IsEaseCurveSelection() const
+{
+	return EaseCurveToolWeak.Pin()->AreAllEaseCurves();
+}
+
+void SAvaEaseCurveTool::HandleEditorTangentsChanged(const FAvaEaseCurveTangents& InTangents) const
+{
+	SetTangents(InTangents, ToolOperation.Get(), true, true, true);
+}
+
+void SAvaEaseCurveTool::OnEditorDragStart() const
+{
+	const TSharedPtr<FAvaEaseCurveTool> EaseCurveTool = EaseCurveToolWeak.Pin();
+	if (!EaseCurveTool.IsValid())
+	{
+		return;
+	}
+
+	EaseCurveTool->BeginTransaction(LOCTEXT("EditorDragStartLabel", "Ease Curve Graph Drag"));
+}
+
+void SAvaEaseCurveTool::OnEditorDragEnd() const
+{
+	const TSharedPtr<FAvaEaseCurveTool> EaseCurveTool = EaseCurveToolWeak.Pin();
+	if (!EaseCurveTool.IsValid())
+	{
+		return;
+	}
+
+	EaseCurveTool->EndTransaction();
+
+	if (!EaseCurveTool->HasCachedKeysToEase())
+	{
+		ResetTangentsAndNotify();
+	}
+}
+
+void SAvaEaseCurveTool::SetTangents(const FAvaEaseCurveTangents& InTangents, EAvaEaseCurveToolOperation InOperation
+	, const bool bInSetEaseCurve, const bool bInBroadcastUpdate, const bool bInSetSequencerTangents) const
+{
+	if (CurvePresetWidget.IsValid() && !CurvePresetWidget->SetSelectedItem(InTangents))
+	{
+		CurvePresetWidget->ClearSelection();
+	}
+
+	// To change the graph UI tangents, we need to change the ease curve object tangents and the graph will reflect.
+	if (bInSetEaseCurve)
+	{
+        if (const TSharedPtr<FAvaEaseCurveTool> EaseCurveTool = EaseCurveToolWeak.Pin())
+        {
+        	EaseCurveTool->SetEaseCurveTangents(InTangents, InOperation, bInBroadcastUpdate, bInSetSequencerTangents);
+        }
+	}
+
+	if (GetDefault<UAvaEaseCurveToolSettings>()->GetAutoZoomToFit())
+	{
+		ZoomToFit();
+	}
+}
+
+FAvaEaseCurveTangents SAvaEaseCurveTool::GetTangents() const
+{
+	if (const TSharedPtr<FAvaEaseCurveTool> EaseCurveTool = EaseCurveToolWeak.Pin())
+	{
+		return EaseCurveTool->GetEaseCurveTangents();
+	}
+	return FAvaEaseCurveTangents();
+}
+
+void SAvaEaseCurveTool::OnStartTangentSpinBoxChanged(const double InNewValue) const
+{
+	const TSharedPtr<FAvaEaseCurveTool> EaseCurveTool = EaseCurveToolWeak.Pin();
+	if (!EaseCurveTool.IsValid())
+	{
+		return;
+	}
+
+	FAvaEaseCurveTangents NewTangents = EaseCurveTool->GetEaseCurveTangents();
+	NewTangents.Start = InNewValue;
+	SetTangents(NewTangents, ToolOperation.Get(), true, true, true);
+}
+
+void SAvaEaseCurveTool::OnStartTangentWeightSpinBoxChanged(const double InNewValue) const
+{
+	const TSharedPtr<FAvaEaseCurveTool> EaseCurveTool = EaseCurveToolWeak.Pin();
+	if (!EaseCurveTool.IsValid())
+	{
+		return;
+	}
+
+	FAvaEaseCurveTangents NewTangents = EaseCurveTool->GetEaseCurveTangents();
+	NewTangents.StartWeight = InNewValue;
+	SetTangents(NewTangents, ToolOperation.Get(), true, true, true);
+}
+
+void SAvaEaseCurveTool::OnEndTangentSpinBoxChanged(const double InNewValue) const
+{
+	const TSharedPtr<FAvaEaseCurveTool> EaseCurveTool = EaseCurveToolWeak.Pin();
+	if (!EaseCurveTool.IsValid())
+	{
+		return;
+	}
+
+	FAvaEaseCurveTangents NewTangents = EaseCurveTool->GetEaseCurveTangents();
+	NewTangents.End = InNewValue;
+	SetTangents(NewTangents, ToolOperation.Get(), true, true, true);
+}
+
+void SAvaEaseCurveTool::OnEndTangentWeightSpinBoxChanged(const double InNewValue) const
+{
+	const TSharedPtr<FAvaEaseCurveTool> EaseCurveTool = EaseCurveToolWeak.Pin();
+	if (!EaseCurveTool.IsValid())
+	{
+		return;
+	}
+
+	FAvaEaseCurveTangents NewTangents = EaseCurveTool->GetEaseCurveTangents();
+	NewTangents.EndWeight = InNewValue;
+	SetTangents(NewTangents, ToolOperation.Get(), true, true, true);
+}
+
+void SAvaEaseCurveTool::OnBeginSliderMovement()
+{
+	const TSharedPtr<FAvaEaseCurveTool> EaseCurveTool = EaseCurveToolWeak.Pin();
+    if (!EaseCurveTool.IsValid())
+    {
+    	return;
+    }
+
+	EaseCurveTool->BeginTransaction(LOCTEXT("SliderDragStartLabel", "Ease Curve Slider Drag"));
+}
+
+void SAvaEaseCurveTool::OnEndSliderMovement(const double InNewValue)
+{
+	const TSharedPtr<FAvaEaseCurveTool> EaseCurveTool = EaseCurveToolWeak.Pin();
+	if (!EaseCurveTool.IsValid())
+	{
+		return;
+	}
+
+	EaseCurveTool->EndTransaction();
+}
+
+void SAvaEaseCurveTool::OnPresetChanged(const TSharedPtr<FAvaEaseCurvePreset>& InPreset) const
+{
+	const TSharedPtr<FAvaEaseCurveTool> EaseCurveTool = EaseCurveToolWeak.Pin();
+	if (!EaseCurveTool.IsValid())
+	{
+		return;
+	}
+
+	if (!EaseCurveTool->HasCachedKeysToEase())
+	{
+		ResetTangentsAndNotify();
+		return;
+	}
+
+	if (InPreset.IsValid())
+	{
+		SetTangents(InPreset->Tangents, ToolOperation.Get(), true, true, true);
+	}
+
+	FSlateApplication::Get().SetAllUserFocus(CurveEaseEditorWidget);
+
+	if (FEngineAnalytics::IsAvailable())
+	{
+		// Only send analytics for default presets
+		const TMap<FString, TArray<FString>>& DefaultPresetNames = UAvaEaseCurveSubsystem::GetDefaultCategoryPresetNames();
+		if (DefaultPresetNames.Contains(InPreset->Category)
+			&& DefaultPresetNames[InPreset->Category].Contains(InPreset->Name))
+		{
+			TArray<FAnalyticsEventAttribute> Attributes;
+			Attributes.Emplace(TEXT("Category"), InPreset->Category);
+			Attributes.Emplace(TEXT("Name"), InPreset->Name);
+
+			FEngineAnalytics::GetProvider().RecordEvent(TEXT("Editor.Usage.MotionDesign.EaseCurveTool.SetTangentsPreset"), Attributes);
+		}
+	}
+}
+
+void SAvaEaseCurveTool::OnQuickPresetChanged(const TSharedPtr<FAvaEaseCurvePreset>& InPreset) const
+{
+	FSlateApplication::Get().SetAllUserFocus(CurveEaseEditorWidget);
+}
+
+bool SAvaEaseCurveTool::OnGetNewPresetTangents(FAvaEaseCurveTangents& OutTangents) const
+{
+	const TSharedPtr<FAvaEaseCurveTool> EaseCurveTool = EaseCurveToolWeak.Pin();
+	if (!EaseCurveTool.IsValid())
+	{
+		return false;
+	}
+
+	OutTangents = EaseCurveTool->GetEaseCurveTangents();
+	return true;
+}
+
+void SAvaEaseCurveTool::UndoAction()
+{
+	if (GEditor)
+	{
+		GEditor->UndoTransaction();
+	}
+}
+
+void SAvaEaseCurveTool::RedoAction()
+{
+	if (GEditor)
+	{
+		GEditor->RedoTransaction();
+	}
+}
+
+void SAvaEaseCurveTool::ZoomToFit() const
+{
+	if (CurveEaseEditorWidget.IsValid())
+	{
+		CurveEaseEditorWidget->ZoomToFit();
+	}
+}
+
+FKeyHandle SAvaEaseCurveTool::GetSelectedKeyHandle() const
+{
+	if (CurveEaseEditorWidget.IsValid())
+	{
+		return CurveEaseEditorWidget->GetSelectedKeyHandle();
+	}
+	return FKeyHandle::Invalid();
+}
+
+FReply SAvaEaseCurveTool::OnKeyDown(const FGeometry& MyGeometry, const FKeyEvent& InKeyEvent)
+{
+	const TSharedPtr<FAvaEaseCurveTool> EaseCurveTool = EaseCurveToolWeak.Pin();
+	if (!EaseCurveTool.IsValid())
+	{
+		return FReply::Unhandled();
+	}
+
+	const TSharedPtr<FUICommandList> CommandList = EaseCurveTool->GetCommandList();
+	if (CommandList.IsValid() && CommandList->ProcessCommandBindings(InKeyEvent))
+	{
+		return FReply::Handled();
+	}
+
+	return FReply::Unhandled();
+}
+
+FText SAvaEaseCurveTool::GetStartText() const
+{
+	return (ToolMode.Get(EAvaEaseCurveToolMode::DualKeyEdit) == EAvaEaseCurveToolMode::DualKeyEdit)
+		? LOCTEXT("StartText", "Leave")
+		: LOCTEXT("ArriveText", "Arrive");
+}
+
+FText SAvaEaseCurveTool::GetStartTooltipText() const
+{
+	return (ToolMode.Get(EAvaEaseCurveToolMode::DualKeyEdit) == EAvaEaseCurveToolMode::DualKeyEdit)
+		? LOCTEXT("StartTooltipText", "Start: The selected key's leave tangent")
+		: LOCTEXT("ArriveTooltipText", "Arrive");
+}
+
+FText SAvaEaseCurveTool::GetEndText() const
+{
+	return (ToolMode.Get(EAvaEaseCurveToolMode::DualKeyEdit) == EAvaEaseCurveToolMode::DualKeyEdit)
+		? LOCTEXT("EndText", "Arrive")
+		: LOCTEXT("LeaveText", "Leave");
+}
+
+FText SAvaEaseCurveTool::GetEndTooltipText() const
+{
+	return (ToolMode.Get(EAvaEaseCurveToolMode::DualKeyEdit) == EAvaEaseCurveToolMode::DualKeyEdit)
+		? LOCTEXT("EndTooltipText", "End: The next key's arrive tangent")
+		: LOCTEXT("LeaveTooltipText", "Leave");
+}
+
+void SAvaEaseCurveTool::ResetTangentsAndNotify() const
+{
+	CurvePresetWidget->ClearSelection();
+
+	SetTangents(FAvaEaseCurveTangents(), EAvaEaseCurveToolOperation::InOut, true, true, false);
+
+	FAvaEaseCurveTool::ShowNotificationMessage(LOCTEXT("EqualValueKeys", "No different key values to create ease curve!"));
+}
+
+#undef LOCTEXT_NAMESPACE
diff --git a/Engine/Plugins/VirtualProduction/Avalanche/Source/AvalancheSequencer/Private/EaseCurveTool/Widgets/SAvaEaseCurveTool.h b/Engine/Plugins/VirtualProduction/Avalanche/Source/AvalancheSequencer/Private/EaseCurveTool/Widgets/SAvaEaseCurveTool.h
new file mode 100644
--- /dev/null
+++ b/Engine/Plugins/VirtualProduction/Avalanche/Source/AvalancheSequencer/Private/EaseCurveTool/Widgets/SAvaEaseCurveTool.h
@@ -0,0 +1,94 @@
+// Copyright Epic Games, Inc. All Rights Reserved.
+
+#pragma once
+
+#include "EaseCurveTool/AvaEaseCurveTool.h"
+#include "EditorUndoClient.h"
+#include "Templates/SharedPointer.h"
+#include "Widgets/DeclarativeSyntaxSupport.h"
+#include "Widgets/SCompoundWidget.h"
+
+class FAvaEaseCurveToolContextMenu;
+class FText;
+class SAvaEaseCurveEditor;
+class SAvaEaseCurvePreset;
+struct FKeyHandle;
+
+class SAvaEaseCurveTool
+	: public SCompoundWidget
+	, public FEditorUndoClient
+{
+public:
+	static constexpr int32 DefaultGraphSize = 200;
+	
+	SLATE_BEGIN_ARGS(SAvaEaseCurveTool)
+		: _ToolMode(EAvaEaseCurveToolMode::DualKeyEdit)
+		, _ToolOperation(EAvaEaseCurveToolOperation::InOut)
+	{}
+		SLATE_ATTRIBUTE(EAvaEaseCurveToolMode, ToolMode)
+		SLATE_ATTRIBUTE(EAvaEaseCurveToolOperation, ToolOperation)
+		SLATE_ARGUMENT(FAvaEaseCurveTangents, InitialTangents)
+	SLATE_END_ARGS()
+
+	void Construct(const FArguments& InArgs, const TSharedRef<FAvaEaseCurveTool>& InEaseCurveTool);
+
+	void SetTangents(const FAvaEaseCurveTangents& InTangents, EAvaEaseCurveToolOperation InOperation,
+		const bool bInSetEaseCurve, const bool bInBroadcastUpdate, const bool bInSetSequencerTangents) const;
+
+	FAvaEaseCurveTangents GetTangents() const;
+
+	FKeyHandle GetSelectedKeyHandle() const;
+
+	void ZoomToFit() const;
+
+protected:
+	TSharedRef<SWidget> ConstructCurveEditorPanel();
+
+	FFrameRate GetDisplayRate() const;
+	EAvaEaseCurveToolOperation GetToolOperation() const;
+	FVector2D GetEditorSize() const;
+	bool ShowEqualValueKeyError() const;
+	bool IsEaseCurveSelection() const;
+
+	void HandleEditorTangentsChanged(const FAvaEaseCurveTangents& InTangents) const;
+
+	void OnStartTangentSpinBoxChanged(const double InNewValue) const;
+	void OnStartTangentWeightSpinBoxChanged(const double InNewValue) const;
+	void OnEndTangentSpinBoxChanged(const double InNewValue) const;
+	void OnEndTangentWeightSpinBoxChanged(const double InNewValue) const;
+	void OnBeginSliderMovement();
+	void OnEndSliderMovement(const double InNewValue);
+
+	void OnPresetChanged(const TSharedPtr<FAvaEaseCurvePreset>& InPreset) const;
+	void OnQuickPresetChanged(const TSharedPtr<FAvaEaseCurvePreset>& InPreset) const;
+	bool OnGetNewPresetTangents(FAvaEaseCurveTangents& OutTangents) const;
+
+	void UndoAction();
+	void RedoAction();
+
+	void OnEditorDragStart() const;
+	void OnEditorDragEnd() const;
+
+	FText GetStartText() const;
+	FText GetStartTooltipText() const;
+	FText GetEndText() const;
+	FText GetEndTooltipText() const;
+
+	void ResetTangentsAndNotify() const;
+
+	//~ Begin SWidget
+	virtual FReply OnKeyDown(const FGeometry& MyGeometry, const FKeyEvent& InKeyEvent) override;
+	//~ End SWidget
+
+	TWeakPtr<FAvaEaseCurveTool> EaseCurveToolWeak;
+
+	TAttribute<EAvaEaseCurveToolMode> ToolMode;
+	TAttribute<EAvaEaseCurveToolOperation> ToolOperation;
+
+	TSharedPtr<SAvaEaseCurveEditor> CurveEaseEditorWidget;
+	TSharedPtr<SAvaEaseCurvePreset> CurvePresetWidget;
+
+	int32 CurrentGraphSize = DefaultGraphSize;
+
+	TSharedPtr<FAvaEaseCurveToolContextMenu> ContextMenu;
+};
diff --git a/Engine/Plugins/VirtualProduction/Avalanche/Source/AvalancheSequencer/Private/NavigationTool/AvaNavigationToolProvider.h b/Engine/Plugins/VirtualProduction/Avalanche/Source/AvalancheSequencer/Private/NavigationTool/AvaNavigationToolProvider.h
new file mode 100644
--- /dev/null
+++ b/Engine/Plugins/VirtualProduction/Avalanche/Source/AvalancheSequencer/Private/NavigationTool/AvaNavigationToolProvider.h
@@ -0,0 +1,97 @@
+﻿// Copyright Epic Games, Inc. All Rights Reserved.
+
+#pragma once
+
+#include "Delegates/IDelegateInstance.h"
+#include "NavigationToolDefines.h"
+#include "Providers/NavigationToolProvider.h"
+#include "Templates/SharedPointer.h"
+#include "UObject/WeakObjectPtr.h"
+
+class FAvaSequencer;
+class FUICommandList;
+class IAvaSceneInterface;
+class IAvaSequenceProvider;
+class UAvaSequence;
+class UObject;
+struct FNavigationToolSaveState;
+
+namespace UE::SequenceNavigator
+{
+	class FNavigationToolItem;
+	class FNavigationToolItemDragDropOp;
+	class INavigationTool;
+}
+
+namespace UE::AvaSequencer
+{
+
+class FAvaNavigationToolProvider : public UE::SequenceNavigator::FNavigationToolProvider
+{
+public:
+	static const FName Identifier;
+
+	static const FText MotionDesignColumnViewName;
+
+	FAvaNavigationToolProvider(const TSharedRef<FAvaSequencer>& InAvaSequencer);
+	virtual ~FAvaNavigationToolProvider() override = default;
+
+	//~ Begin INavigationToolProvider
+
+	virtual FName GetIdentifier() const override;
+
+	virtual TSet<TSubclassOf<UMovieSceneSequence>> GetSupportedSequenceClasses() const override;
+
+	virtual FText GetDefaultColumnView() const override;
+
+	virtual FNavigationToolSaveState* GetSaveState(const UE::SequenceNavigator::INavigationTool& InTool) const override;
+	virtual void SetSaveState(const UE::SequenceNavigator::INavigationTool& InTool, const FNavigationToolSaveState& InSaveState) const override;
+
+	virtual void BindCommands(const TSharedRef<FUICommandList>& InCommandList) override;
+
+	virtual void OnActivate() override;
+	virtual void OnDeactivate() override;
+
+	virtual void OnExtendColumns(UE::SequenceNavigator::FNavigationToolColumnExtender& OutExtender) override;
+	virtual void OnExtendColumnViews(TSet<FNavigationToolColumnView>& OutColumnViews) override;
+	virtual void OnExtendItemChildren(UE::SequenceNavigator::INavigationTool& InTool
+		, const UE::SequenceNavigator::FNavigationToolItemPtr& InParentItem
+		, TArray<UE::SequenceNavigator::FNavigationToolItemPtr>& OutChildren
+		, const bool bInRecursive) override;
+	virtual void OnExtendBuiltInFilters(TArray<FNavigationToolBuiltInFilterParams>& OutFilterParams) override;
+
+	//~ End INavigationToolProvider
+
+private:
+	static const FName ToolbarSectionName;
+	static const FName ContextMenuSectionName;
+
+	IAvaSceneInterface* GetSceneInterface(const UE::SequenceNavigator::INavigationTool& InTool) const;
+	IAvaSequenceProvider* GetSequenceProvider(const UE::SequenceNavigator::INavigationTool& InTool) const;
+
+	void ExtendToolToolBar();
+	void ExtendToolItemContextMenu();
+
+	void RemoveToolToolBarExtension();
+	void RemoveToolItemContextMenuExtension();
+
+	void ExtendToolItemDragDropOp(UE::SequenceNavigator::FNavigationToolItemDragDropOp& InDragDropOp);
+
+	void OnSequenceAdded(UAvaSequence* const InAvaSequence);
+
+	void GeneratePresetMenu(UToolMenu* const InToolMenu);
+
+	bool CanRelabelSequence() const;
+	void RelabelSequence();
+
+	bool CanAddSequenceToSelection() const;
+	void AddSequenceToSelection();
+
+	TWeakPtr<FAvaSequencer> WeakAvaSequencer;
+
+	TSharedRef<FUICommandList> ToolCommands;
+
+	FDelegateHandle DragDropInitializedDelegate;
+};
+
+} // UE::AvaSequencer
diff --git a/Engine/Plugins/VirtualProduction/Avalanche/Source/AvalancheSequencer/Private/Settings/AvaSequencerSettings.cpp b/Engine/Plugins/VirtualProduction/Avalanche/Source/AvalancheSequencer/Private/Settings/AvaSequencerSettings.cpp
new file mode 100644
--- /dev/null
+++ b/Engine/Plugins/VirtualProduction/Avalanche/Source/AvalancheSequencer/Private/Settings/AvaSequencerSettings.cpp
@@ -0,0 +1,58 @@
+// Copyright Epic Games, Inc. All Rights Reserved.
+
+#include "AvaSequencerSettings.h"
+#include "AvaSequenceDefaultTags.h"
+#include "Misc/FrameRate.h"
+
+namespace UE::AvaSequencer::Private
+{
+	FAvaSequencePreset CreateDefaultPreset(FName InSequenceType, const FAvaTagSoftHandle& InSequenceTag)
+	{
+		FAvaSequencePreset Preset;
+		Preset.PresetName    = InSequenceType;
+		Preset.SequenceLabel = InSequenceType;
+		Preset.SequenceTag   = InSequenceTag;
+		Preset.bEnableLabel  = true;
+		Preset.bEnableTag    = true;
+		return Preset;
+	}
+
+	TArray<FAvaSequencePreset> MakeDefaultSequencePresets()
+	{
+		const FAvaSequenceDefaultTags& DefaultTags = FAvaSequenceDefaultTags::Get();
+
+		// Guids found via the entries that DefaultSequenceTags had for In/Out/Change
+		TArray<FAvaSequencePreset> DefaultPresets =
+			{
+				CreateDefaultPreset(TEXT("In")    , DefaultTags.In),
+				CreateDefaultPreset(TEXT("Out")   , DefaultTags.Out),
+				CreateDefaultPreset(TEXT("Change"), DefaultTags.Change),
+			};
+
+		// Configure Change Marks
+		{
+			FAvaSequencePreset& ChangePreset = DefaultPresets[2];
+			ChangePreset.bEnableMarks = true;
+
+			FAvaMarkSetting& MarkSetting = ChangePreset.Marks.AddDefaulted_GetRef();
+        	MarkSetting.Label = TEXT("A");
+			MarkSetting.FrameNumber = 60;
+		}
+
+		return DefaultPresets;
+	}
+}
+
+UAvaSequencerSettings::UAvaSequencerSettings()
+{
+	CategoryName = SettingsCategory;
+	SectionName = SettingsSection;
+
+	DisplayRate.FrameRate = FFrameRate(60000, 1001);
+}
+
+TConstArrayView<FAvaSequencePreset> UAvaSequencerSettings::GetDefaultSequencePresets() const
+{
+	static const TArray<FAvaSequencePreset> DefaultPresets = UE::AvaSequencer::Private::MakeDefaultSequencePresets();
+	return DefaultPresets;
+}
diff --git a/Engine/Plugins/VirtualProduction/Avalanche/Source/AvalancheShapes/Private/AvaShapesModule.cpp b/Engine/Plugins/VirtualProduction/Avalanche/Source/AvalancheShapes/Private/AvaShapesModule.cpp
new file mode 100644
--- /dev/null
+++ b/Engine/Plugins/VirtualProduction/Avalanche/Source/AvalancheShapes/Private/AvaShapesModule.cpp
@@ -0,0 +1,5 @@
+// Copyright Epic Games, Inc. All Rights Reserved.
+
+#include "Modules/ModuleManager.h"
+
+IMPLEMENT_MODULE(FDefaultModuleImpl, AvalancheShapes)
diff --git a/Engine/Plugins/VirtualProduction/Avalanche/Source/AvalancheShapesEditor/Private/Tools/AvaShapesEditorShapeTool2DArrow.cpp b/Engine/Plugins/VirtualProduction/Avalanche/Source/AvalancheShapesEditor/Private/Tools/AvaShapesEditorShapeTool2DArrow.cpp
new file mode 100644
--- /dev/null
+++ b/Engine/Plugins/VirtualProduction/Avalanche/Source/AvalancheShapesEditor/Private/Tools/AvaShapesEditorShapeTool2DArrow.cpp
@@ -0,0 +1,33 @@
+// Copyright Epic Games, Inc. All Rights Reserved.
+
+#include "Tools/AvaShapesEditorShapeTool2DArrow.h"
+#include "AvaShapesEditorCommands.h"
+#include "Builders/AvaInteractiveToolsToolBuilder.h"
+#include "DynamicMeshes/AvaShape2DArrowDynMesh.h"
+#include "UObject/ConstructorHelpers.h"
+
+UAvaShapesEditorShapeTool2DArrow::UAvaShapesEditorShapeTool2DArrow()
+{
+	ShapeClass = UAvaShape2DArrowDynamicMesh::StaticClass();
+}
+
+void UAvaShapesEditorShapeTool2DArrow::OnRegisterTool(IAvalancheInteractiveToolsModule* InAITModule)
+{
+	Super::OnRegisterTool(InAITModule);
+
+	FAvaInteractiveToolsToolParameters ToolParameters = 
+	{
+		FAvaShapesEditorCommands::Get().Tool_Shape_2DArrow,
+		TEXT("Parametric 2D Arrow Tool"),
+		7000,
+		FAvalancheInteractiveToolsCreateBuilder::CreateLambda(
+			[](UEdMode* InEdMode)
+			{
+				return UAvaInteractiveToolsToolBuilder::CreateToolBuilder<UAvaShapesEditorShapeTool2DArrow>(InEdMode);
+			}),
+		nullptr,
+		CreateFactory<UAvaShape2DArrowDynamicMesh>()
+	};
+
+	InAITModule->RegisterTool(IAvalancheInteractiveToolsModule::CategoryName2D, MoveTemp(ToolParameters));
+}
diff --git a/Engine/Plugins/VirtualProduction/Avalanche/Source/AvalancheShapesEditor/Private/Tools/AvaShapesEditorShapeToolChevron.cpp b/Engine/Plugins/VirtualProduction/Avalanche/Source/AvalancheShapesEditor/Private/Tools/AvaShapesEditorShapeToolChevron.cpp
new file mode 100644
--- /dev/null
+++ b/Engine/Plugins/VirtualProduction/Avalanche/Source/AvalancheShapesEditor/Private/Tools/AvaShapesEditorShapeToolChevron.cpp
@@ -0,0 +1,33 @@
+// Copyright Epic Games, Inc. All Rights Reserved.
+
+#include "Tools/AvaShapesEditorShapeToolChevron.h"
+#include "AvaShapesEditorCommands.h"
+#include "Builders/AvaInteractiveToolsToolBuilder.h"
+#include "DynamicMeshes/AvaShapeChevronDynMesh.h"
+#include "UObject/ConstructorHelpers.h"
+
+UAvaShapesEditorShapeToolChevron::UAvaShapesEditorShapeToolChevron()
+{
+	ShapeClass = UAvaShapeChevronDynamicMesh::StaticClass();
+}
+
+void UAvaShapesEditorShapeToolChevron::OnRegisterTool(IAvalancheInteractiveToolsModule* InAITModule)
+{
+	Super::OnRegisterTool(InAITModule);
+
+	FAvaInteractiveToolsToolParameters ToolParameters = 
+	{
+		FAvaShapesEditorCommands::Get().Tool_Shape_Chevron,
+		TEXT("Parametric Chevron Tool"),
+		8000,
+		FAvalancheInteractiveToolsCreateBuilder::CreateLambda(
+			[](UEdMode* InEdMode)
+			{
+				return UAvaInteractiveToolsToolBuilder::CreateToolBuilder<UAvaShapesEditorShapeToolChevron>(InEdMode);
+			}),
+		nullptr,
+		CreateFactory<UAvaShapeChevronDynamicMesh>()
+	};
+
+	InAITModule->RegisterTool(IAvalancheInteractiveToolsModule::CategoryName2D, MoveTemp(ToolParameters));
+}
diff --git a/Engine/Plugins/VirtualProduction/Avalanche/Source/AvalancheShapesEditor/Private/Tools/AvaShapesEditorShapeToolCone.cpp b/Engine/Plugins/VirtualProduction/Avalanche/Source/AvalancheShapesEditor/Private/Tools/AvaShapesEditorShapeToolCone.cpp
new file mode 100644
--- /dev/null
+++ b/Engine/Plugins/VirtualProduction/Avalanche/Source/AvalancheShapesEditor/Private/Tools/AvaShapesEditorShapeToolCone.cpp
@@ -0,0 +1,59 @@
+// Copyright Epic Games, Inc. All Rights Reserved.
+
+#include "Tools/AvaShapesEditorShapeToolCone.h"
+#include "AvaShapesEditorCommands.h"
+#include "Builders/AvaInteractiveToolsToolBuilder.h"
+#include "DynamicMeshes/AvaShapeConeDynMesh.h"
+#include "UObject/ConstructorHelpers.h"
+
+UAvaShapesEditorShapeToolCone::UAvaShapesEditorShapeToolCone()
+{
+	ShapeClass = UAvaShapeConeDynamicMesh::StaticClass();
+}
+
+void UAvaShapesEditorShapeToolCone::OnRegisterTool(IAvalancheInteractiveToolsModule* InAITModule)
+{
+	Super::OnRegisterTool(InAITModule);
+
+	FAvaInteractiveToolsToolParameters ConeToolParameters = 
+	{
+		FAvaShapesEditorCommands::Get().Tool_Shape_Cone,
+		TEXT("Parametric Cone Tool"),
+		3000,
+		FAvalancheInteractiveToolsCreateBuilder::CreateLambda(
+			[](UEdMode* InEdMode)
+			{
+				return UAvaInteractiveToolsToolBuilder::CreateToolBuilder<UAvaShapesEditorShapeToolCone>(InEdMode);
+			}),
+		nullptr,
+		CreateFactory<UAvaShapeConeDynamicMesh>()
+	};
+
+	InAITModule->RegisterTool(IAvalancheInteractiveToolsModule::CategoryName3D, MoveTemp(ConeToolParameters));
+
+	const FShapeFactoryParameters CylinderFactoryParameters =
+	{
+		.Functor = [](UAvaShapeDynamicMeshBase* InMesh)
+		{
+			UAvaShapeConeDynamicMesh* Cone = Cast<UAvaShapeConeDynamicMesh>(InMesh);
+			Cone->SetTopRadius(1.f);
+		},
+		.NameOverride = FString("Cylinder"),
+	};
+
+	FAvaInteractiveToolsToolParameters CylinderToolParameters = 
+	{
+		FAvaShapesEditorCommands::Get().Tool_Shape_Cylinder,
+		TEXT("Parametric Cylinder Tool"),
+		3000,
+		FAvalancheInteractiveToolsCreateBuilder::CreateLambda(
+			[](UEdMode* InEdMode)
+			{
+				return UAvaInteractiveToolsToolBuilder::CreateToolBuilder<UAvaShapesEditorShapeToolCone>(InEdMode);
+			}),
+		nullptr,
+		CreateFactory<UAvaShapeConeDynamicMesh>(CylinderFactoryParameters)
+	};
+
+	InAITModule->RegisterTool(IAvalancheInteractiveToolsModule::CategoryName3D, MoveTemp(CylinderToolParameters));
+}
diff --git a/Engine/Plugins/VirtualProduction/Avalanche/Source/AvalancheShapesEditor/Private/Tools/AvaShapesEditorShapeToolCube.cpp b/Engine/Plugins/VirtualProduction/Avalanche/Source/AvalancheShapesEditor/Private/Tools/AvaShapesEditorShapeToolCube.cpp
new file mode 100644
--- /dev/null
+++ b/Engine/Plugins/VirtualProduction/Avalanche/Source/AvalancheShapesEditor/Private/Tools/AvaShapesEditorShapeToolCube.cpp
@@ -0,0 +1,33 @@
+// Copyright Epic Games, Inc. All Rights Reserved.
+
+#include "Tools/AvaShapesEditorShapeToolCube.h"
+#include "AvaShapesEditorCommands.h"
+#include "Builders/AvaInteractiveToolsToolBuilder.h"
+#include "DynamicMeshes/AvaShapeCubeDynMesh.h"
+#include "UObject/ConstructorHelpers.h"
+
+UAvaShapesEditorShapeToolCube::UAvaShapesEditorShapeToolCube()
+{
+	ShapeClass = UAvaShapeCubeDynamicMesh::StaticClass();
+}
+
+void UAvaShapesEditorShapeToolCube::OnRegisterTool(IAvalancheInteractiveToolsModule* InAITModule)
+{
+	Super::OnRegisterTool(InAITModule);
+
+	FAvaInteractiveToolsToolParameters ToolParameters = 
+	{
+		FAvaShapesEditorCommands::Get().Tool_Shape_Cube,
+		TEXT("Parametric Cube Tool"),
+		1000,
+		FAvalancheInteractiveToolsCreateBuilder::CreateLambda(
+			[](UEdMode* InEdMode)
+			{
+				return UAvaInteractiveToolsToolBuilder::CreateToolBuilder<UAvaShapesEditorShapeToolCube>(InEdMode);
+			}),
+		nullptr,
+		CreateFactory<UAvaShapeCubeDynamicMesh>()
+	};
+
+	InAITModule->RegisterTool(IAvalancheInteractiveToolsModule::CategoryName3D, MoveTemp(ToolParameters));
+}
diff --git a/Engine/Plugins/VirtualProduction/Avalanche/Source/AvalancheShapesEditor/Private/Tools/AvaShapesEditorShapeToolEllipse.cpp b/Engine/Plugins/VirtualProduction/Avalanche/Source/AvalancheShapesEditor/Private/Tools/AvaShapesEditorShapeToolEllipse.cpp
new file mode 100644
--- /dev/null
+++ b/Engine/Plugins/VirtualProduction/Avalanche/Source/AvalancheShapesEditor/Private/Tools/AvaShapesEditorShapeToolEllipse.cpp
@@ -0,0 +1,33 @@
+// Copyright Epic Games, Inc. All Rights Reserved.
+
+#include "Tools/AvaShapesEditorShapeToolEllipse.h"
+#include "AvaShapesEditorCommands.h"
+#include "Builders/AvaInteractiveToolsToolBuilder.h"
+#include "DynamicMeshes/AvaShapeEllipseDynMesh.h"
+#include "UObject/ConstructorHelpers.h"
+
+UAvaShapesEditorShapeToolEllipse::UAvaShapesEditorShapeToolEllipse()
+{
+	ShapeClass = UAvaShapeEllipseDynamicMesh::StaticClass();
+}
+
+void UAvaShapesEditorShapeToolEllipse::OnRegisterTool(IAvalancheInteractiveToolsModule* InAITModule)
+{
+	Super::OnRegisterTool(InAITModule);
+
+	FAvaInteractiveToolsToolParameters ToolParameters = 
+	{
+		FAvaShapesEditorCommands::Get().Tool_Shape_Ellipse,
+		TEXT("Parametric Ellipse Tool"),
+		2000,
+		FAvalancheInteractiveToolsCreateBuilder::CreateLambda(
+			[](UEdMode* InEdMode)
+			{
+				return UAvaInteractiveToolsToolBuilder::CreateToolBuilder<UAvaShapesEditorShapeToolEllipse>(InEdMode);
+			}),
+		nullptr,
+		CreateFactory<UAvaShapeEllipseDynamicMesh>()
+	};
+	
+	InAITModule->RegisterTool(IAvalancheInteractiveToolsModule::CategoryName2D, MoveTemp(ToolParameters));
+}
diff --git a/Engine/Plugins/VirtualProduction/Avalanche/Source/AvalancheShapesEditor/Private/Tools/AvaShapesEditorShapeToolNGon.cpp b/Engine/Plugins/VirtualProduction/Avalanche/Source/AvalancheShapesEditor/Private/Tools/AvaShapesEditorShapeToolNGon.cpp
new file mode 100644
--- /dev/null
+++ b/Engine/Plugins/VirtualProduction/Avalanche/Source/AvalancheShapesEditor/Private/Tools/AvaShapesEditorShapeToolNGon.cpp
@@ -0,0 +1,33 @@
+// Copyright Epic Games, Inc. All Rights Reserved.
+
+#include "Tools/AvaShapesEditorShapeToolNGon.h"
+#include "AvaShapesEditorCommands.h"
+#include "Builders/AvaInteractiveToolsToolBuilder.h"
+#include "DynamicMeshes/AvaShapeNGonDynMesh.h"
+#include "UObject/ConstructorHelpers.h"
+
+UAvaShapesEditorShapeToolNGon::UAvaShapesEditorShapeToolNGon()
+{
+	ShapeClass = UAvaShapeNGonDynamicMesh::StaticClass();
+}
+
+void UAvaShapesEditorShapeToolNGon::OnRegisterTool(IAvalancheInteractiveToolsModule* InAITModule)
+{
+	Super::OnRegisterTool(InAITModule);
+
+	FAvaInteractiveToolsToolParameters ToolParameters =
+	{
+		FAvaShapesEditorCommands::Get().Tool_Shape_NGon,
+		TEXT("Parametric Regular Polygon Tool"),
+		3000,
+		FAvalancheInteractiveToolsCreateBuilder::CreateLambda(
+			[](UEdMode* InEdMode)
+			{
+				return UAvaInteractiveToolsToolBuilder::CreateToolBuilder<UAvaShapesEditorShapeToolNGon>(InEdMode);
+			}),
+		nullptr,
+		CreateFactory<UAvaShapeNGonDynamicMesh>()
+	};
+	
+	InAITModule->RegisterTool(IAvalancheInteractiveToolsModule::CategoryName2D, MoveTemp(ToolParameters));
+}
diff --git a/Engine/Plugins/VirtualProduction/Avalanche/Source/AvalancheShapesEditor/Private/Tools/AvaShapesEditorShapeToolRectangle.cpp b/Engine/Plugins/VirtualProduction/Avalanche/Source/AvalancheShapesEditor/Private/Tools/AvaShapesEditorShapeToolRectangle.cpp
new file mode 100644
--- /dev/null
+++ b/Engine/Plugins/VirtualProduction/Avalanche/Source/AvalancheShapesEditor/Private/Tools/AvaShapesEditorShapeToolRectangle.cpp
@@ -0,0 +1,67 @@
+// Copyright Epic Games, Inc. All Rights Reserved.
+
+#include "Tools/AvaShapesEditorShapeToolRectangle.h"
+#include "AvaShapesEditorCommands.h"
+#include "Builders/AvaInteractiveToolsToolBuilder.h"
+#include "DynamicMeshes/AvaShapeRectangleDynMesh.h"
+#include "UObject/ConstructorHelpers.h"
+
+#define LOCTEXT_NAMESPACE "AvaShapesEditorShapeToolRectangle"
+
+UAvaShapesEditorShapeToolRectangle::UAvaShapesEditorShapeToolRectangle()
+{
+	ShapeClass = UAvaShapeRectangleDynamicMesh::StaticClass();
+}
+
+void UAvaShapesEditorShapeToolRectangle::OnRegisterTool(IAvalancheInteractiveToolsModule* InAITModule)
+{
+	Super::OnRegisterTool(InAITModule);
+
+	FShapeFactoryParameters RectangleFactoryParameters =
+	{
+		.Size = FVector(0, 160, 90),
+		.Functor = [](UAvaShapeDynamicMeshBase* InMesh)
+		{
+			InMesh->SetMaterialUVMode(UAvaShapeDynamicMeshBase::MESH_INDEX_PRIMARY, EAvaShapeUVMode::Stretch);
+		}
+	};
+
+	FAvaInteractiveToolsToolParameters RectangleToolParameters = 
+	{
+		FAvaShapesEditorCommands::Get().Tool_Shape_Rectangle,
+		TEXT("Parametric Rectangle Tool"),
+		1000,
+		FAvalancheInteractiveToolsCreateBuilder::CreateLambda(
+			[](UEdMode* InEdMode)
+			{
+				return UAvaInteractiveToolsToolBuilder::CreateToolBuilder<UAvaShapesEditorShapeToolRectangle>(InEdMode);
+			}),
+		nullptr,
+		CreateFactory<UAvaShapeRectangleDynamicMesh>(RectangleFactoryParameters)
+	};
+
+	InAITModule->RegisterTool(IAvalancheInteractiveToolsModule::CategoryName2D, MoveTemp(RectangleToolParameters));
+
+	FShapeFactoryParameters SquareFactoryParameters =
+	{
+		.NameOverride = FString("Square")
+	};
+
+	FAvaInteractiveToolsToolParameters SquareToolParameters = 
+	{
+		FAvaShapesEditorCommands::Get().Tool_Shape_Square,
+		TEXT("Parametric Square Tool"),
+		1001,
+		FAvalancheInteractiveToolsCreateBuilder::CreateLambda(
+			[](UEdMode* InEdMode)
+			{
+				return UAvaInteractiveToolsToolBuilder::CreateToolBuilder<UAvaShapesEditorShapeToolRectangle>(InEdMode);
+			}),
+		nullptr,
+		CreateFactory<UAvaShapeRectangleDynamicMesh>(SquareFactoryParameters)
+	};
+
+	InAITModule->RegisterTool(IAvalancheInteractiveToolsModule::CategoryName2D, MoveTemp(SquareToolParameters));
+}
+
+#undef LOCTEXT_NAMESPACE
diff --git a/Engine/Plugins/VirtualProduction/Avalanche/Source/AvalancheShapesEditor/Private/Tools/AvaShapesEditorShapeToolRing.cpp b/Engine/Plugins/VirtualProduction/Avalanche/Source/AvalancheShapesEditor/Private/Tools/AvaShapesEditorShapeToolRing.cpp
new file mode 100644
--- /dev/null
+++ b/Engine/Plugins/VirtualProduction/Avalanche/Source/AvalancheShapesEditor/Private/Tools/AvaShapesEditorShapeToolRing.cpp
@@ -0,0 +1,33 @@
+// Copyright Epic Games, Inc. All Rights Reserved.
+
+#include "Tools/AvaShapesEditorShapeToolRing.h"
+#include "AvaShapesEditorCommands.h"
+#include "Builders/AvaInteractiveToolsToolBuilder.h"
+#include "DynamicMeshes/AvaShapeRingDynMesh.h"
+#include "UObject/ConstructorHelpers.h"
+
+UAvaShapesEditorShapeToolRing::UAvaShapesEditorShapeToolRing()
+{
+	ShapeClass = UAvaShapeRingDynamicMesh::StaticClass();
+}
+
+void UAvaShapesEditorShapeToolRing::OnRegisterTool(IAvalancheInteractiveToolsModule* InAITModule)
+{
+	Super::OnRegisterTool(InAITModule);
+
+	FAvaInteractiveToolsToolParameters ToolParameters = 
+	{
+		FAvaShapesEditorCommands::Get().Tool_Shape_Ring,
+		TEXT("Parametric Ring Tool"),
+		5000,
+		FAvalancheInteractiveToolsCreateBuilder::CreateLambda(
+			[](UEdMode* InEdMode)
+			{
+				return UAvaInteractiveToolsToolBuilder::CreateToolBuilder<UAvaShapesEditorShapeToolRing>(InEdMode);
+			}),
+		nullptr,
+		CreateFactory<UAvaShapeRingDynamicMesh>()
+	};
+
+	InAITModule->RegisterTool(IAvalancheInteractiveToolsModule::CategoryName2D, MoveTemp(ToolParameters));
+}
diff --git a/Engine/Plugins/VirtualProduction/Avalanche/Source/AvalancheShapesEditor/Private/Tools/AvaShapesEditorShapeToolSphere.cpp b/Engine/Plugins/VirtualProduction/Avalanche/Source/AvalancheShapesEditor/Private/Tools/AvaShapesEditorShapeToolSphere.cpp
new file mode 100644
--- /dev/null
+++ b/Engine/Plugins/VirtualProduction/Avalanche/Source/AvalancheShapesEditor/Private/Tools/AvaShapesEditorShapeToolSphere.cpp
@@ -0,0 +1,33 @@
+// Copyright Epic Games, Inc. All Rights Reserved.
+
+#include "Tools/AvaShapesEditorShapeToolSphere.h"
+#include "AvaShapesEditorCommands.h"
+#include "Builders/AvaInteractiveToolsToolBuilder.h"
+#include "DynamicMeshes/AvaShapeSphereDynMesh.h"
+#include "UObject/ConstructorHelpers.h"
+
+UAvaShapesEditorShapeToolSphere::UAvaShapesEditorShapeToolSphere()
+{
+	ShapeClass = UAvaShapeSphereDynamicMesh::StaticClass();
+}
+
+void UAvaShapesEditorShapeToolSphere::OnRegisterTool(IAvalancheInteractiveToolsModule* InAITModule)
+{
+	Super::OnRegisterTool(InAITModule);
+
+	FAvaInteractiveToolsToolParameters ToolParameters = 
+	{
+		FAvaShapesEditorCommands::Get().Tool_Shape_Sphere,
+		TEXT("Parametric Sphere Tool"),
+		2000,
+		FAvalancheInteractiveToolsCreateBuilder::CreateLambda(
+			[](UEdMode* InEdMode)
+			{
+				return UAvaInteractiveToolsToolBuilder::CreateToolBuilder<UAvaShapesEditorShapeToolSphere>(InEdMode);
+			}),
+		nullptr,
+		CreateFactory<UAvaShapeSphereDynamicMesh>()
+	};
+	
+	InAITModule->RegisterTool(IAvalancheInteractiveToolsModule::CategoryName3D, MoveTemp(ToolParameters));
+}
diff --git a/Engine/Plugins/VirtualProduction/Avalanche/Source/AvalancheShapesEditor/Private/Tools/AvaShapesEditorShapeToolStar.cpp b/Engine/Plugins/VirtualProduction/Avalanche/Source/AvalancheShapesEditor/Private/Tools/AvaShapesEditorShapeToolStar.cpp
new file mode 100644
--- /dev/null
+++ b/Engine/Plugins/VirtualProduction/Avalanche/Source/AvalancheShapesEditor/Private/Tools/AvaShapesEditorShapeToolStar.cpp
@@ -0,0 +1,33 @@
+// Copyright Epic Games, Inc. All Rights Reserved.
+
+#include "Tools/AvaShapesEditorShapeToolStar.h"
+#include "AvaShapesEditorCommands.h"
+#include "Builders/AvaInteractiveToolsToolBuilder.h"
+#include "DynamicMeshes/AvaShapeStarDynMesh.h"
+#include "UObject/ConstructorHelpers.h"
+
+UAvaShapesEditorShapeToolStar::UAvaShapesEditorShapeToolStar()
+{
+	ShapeClass = UAvaShapeStarDynamicMesh::StaticClass();
+}
+
+void UAvaShapesEditorShapeToolStar::OnRegisterTool(IAvalancheInteractiveToolsModule* InAITModule)
+{
+	Super::OnRegisterTool(InAITModule);
+
+	FAvaInteractiveToolsToolParameters ToolParameters = 
+	{
+		FAvaShapesEditorCommands::Get().Tool_Shape_Star,
+		TEXT("Parametric Star Tool"),
+		6000,
+		FAvalancheInteractiveToolsCreateBuilder::CreateLambda(
+			[](UEdMode* InEdMode)
+			{
+				return UAvaInteractiveToolsToolBuilder::CreateToolBuilder<UAvaShapesEditorShapeToolStar>(InEdMode);
+			}),
+		nullptr,
+		CreateFactory<UAvaShapeStarDynamicMesh>()
+	};
+	
+	InAITModule->RegisterTool(IAvalancheInteractiveToolsModule::CategoryName2D, MoveTemp(ToolParameters));
+}
diff --git a/Engine/Plugins/VirtualProduction/Avalanche/Source/AvalancheShapesEditor/Private/Tools/AvaShapesEditorShapeToolTorus.cpp b/Engine/Plugins/VirtualProduction/Avalanche/Source/AvalancheShapesEditor/Private/Tools/AvaShapesEditorShapeToolTorus.cpp
new file mode 100644
--- /dev/null
+++ b/Engine/Plugins/VirtualProduction/Avalanche/Source/AvalancheShapesEditor/Private/Tools/AvaShapesEditorShapeToolTorus.cpp
@@ -0,0 +1,33 @@
+// Copyright Epic Games, Inc. All Rights Reserved.
+
+#include "Tools/AvaShapesEditorShapeToolTorus.h"
+#include "AvaShapesEditorCommands.h"
+#include "Builders/AvaInteractiveToolsToolBuilder.h"
+#include "DynamicMeshes/AvaShapeTorusDynMesh.h"
+#include "UObject/ConstructorHelpers.h"
+
+UAvaShapesEditorShapeToolTorus::UAvaShapesEditorShapeToolTorus()
+{
+	ShapeClass = UAvaShapeTorusDynamicMesh::StaticClass();
+}
+
+void UAvaShapesEditorShapeToolTorus::OnRegisterTool(IAvalancheInteractiveToolsModule* InAITModule)
+{
+	Super::OnRegisterTool(InAITModule);
+
+	FAvaInteractiveToolsToolParameters ToolParameters = 
+	{
+		FAvaShapesEditorCommands::Get().Tool_Shape_Torus,
+		TEXT("Parametric Torus Tool"),
+		4000,
+		FAvalancheInteractiveToolsCreateBuilder::CreateLambda(
+			[](UEdMode* InEdMode)
+			{
+				return UAvaInteractiveToolsToolBuilder::CreateToolBuilder<UAvaShapesEditorShapeToolTorus>(InEdMode);
+			}),
+		nullptr,
+		CreateFactory<UAvaShapeTorusDynamicMesh>()
+	};
+	
+	InAITModule->RegisterTool(IAvalancheInteractiveToolsModule::CategoryName3D, MoveTemp(ToolParameters));
+}
diff --git a/Engine/Plugins/VirtualProduction/Avalanche/Source/AvalancheText/Private/AvaTextModule.cpp b/Engine/Plugins/VirtualProduction/Avalanche/Source/AvalancheText/Private/AvaTextModule.cpp
new file mode 100644
--- /dev/null
+++ b/Engine/Plugins/VirtualProduction/Avalanche/Source/AvalancheText/Private/AvaTextModule.cpp
@@ -0,0 +1,5 @@
+// Copyright Epic Games, Inc. All Rights Reserved.
+
+#include "Modules/ModuleManager.h"
+
+IMPLEMENT_MODULE(FDefaultModuleImpl, AvalancheText)
diff --git a/Engine/Plugins/VirtualProduction/Avalanche/Source/AvalancheTextEditor/Private/Tool/AvaTextActorTool.cpp b/Engine/Plugins/VirtualProduction/Avalanche/Source/AvalancheTextEditor/Private/Tool/AvaTextActorTool.cpp
new file mode 100644
--- /dev/null
+++ b/Engine/Plugins/VirtualProduction/Avalanche/Source/AvalancheTextEditor/Private/Tool/AvaTextActorTool.cpp
@@ -0,0 +1,35 @@
+// Copyright Epic Games, Inc. All Rights Reserved.
+
+#include "Tool/AvaTextActorTool.h"
+#include "AvaInteractiveToolsSettings.h"
+#include "AvaTextActorFactory.h"
+#include "AvaTextEditorCommands.h"
+#include "Builders/AvaInteractiveToolsToolBuilder.h"
+#include "IAvalancheComponentVisualizersModule.h"
+#include "Text3DActor.h"
+
+UAvaTextActorTool::UAvaTextActorTool()
+{
+	ActorClass = AText3DActor::StaticClass();
+}
+
+void UAvaTextActorTool::OnRegisterTool(IAvalancheInteractiveToolsModule* InAITModule)
+{
+	Super::OnRegisterTool(InAITModule);
+
+	FAvaInteractiveToolsToolParameters ToolParameters =
+	{
+		FAvaTextEditorCommands::Get().Tool_Actor_Text3D,
+		TEXT("Text Actor Tool"),
+		1000,
+		FAvalancheInteractiveToolsCreateBuilder::CreateLambda(
+			[](UEdMode* InEdMode)
+			{
+				return UAvaInteractiveToolsToolBuilder::CreateToolBuilder<UAvaTextActorTool>(InEdMode);
+			}),
+		ActorClass,
+		CreateActorFactory<UAvaTextActorFactory>()
+	};
+
+	InAITModule->RegisterTool(IAvalancheInteractiveToolsModule::CategoryNameActor, MoveTemp(ToolParameters));
+}
diff --git a/Engine/Plugins/VirtualProduction/Avalanche/Source/AvalancheTransition/Private/Behavior/AvaTransitionBehaviorActor.cpp b/Engine/Plugins/VirtualProduction/Avalanche/Source/AvalancheTransition/Private/Behavior/AvaTransitionBehaviorActor.cpp
new file mode 100644
--- /dev/null
+++ b/Engine/Plugins/VirtualProduction/Avalanche/Source/AvalancheTransition/Private/Behavior/AvaTransitionBehaviorActor.cpp
@@ -0,0 +1,106 @@
+// Copyright Epic Games, Inc. All Rights Reserved.
+
+#include "AvaTransitionBehaviorActor.h"
+#include "AvaTransitionContext.h"
+#include "AvaTransitionEnums.h"
+#include "AvaTransitionScene.h"
+#include "AvaTransitionSubsystem.h"
+#include "AvaTransitionTree.h"
+#include "AvaTransitionTreeSchema.h"
+#include "Engine/World.h"
+#include "IAvaTransitionModule.h"
+#include "StateTreeDelegates.h"
+#include "StateTreeExecutionContext.h"
+#include "StateTreeTaskBase.h"
+
+#if WITH_EDITOR
+#include "Editor.h"
+#endif
+
+AAvaTransitionBehaviorActor::AAvaTransitionBehaviorActor()
+{
+	// TickBehavior is called by the Subsystem separately  
+	PrimaryActorTick.bCanEverTick = false;
+
+	TransitionTree = CreateDefaultSubobject<UAvaTransitionTree>(TEXT("Transition Logic"));
+	StateTreeReference.SetStateTree(TransitionTree);
+}
+
+void AAvaTransitionBehaviorActor::PostActorCreated()
+{
+	Super::PostActorCreated();
+
+	ValidateTransitionTree();
+
+	if (UAvaTransitionSubsystem* const TransitionSubsystem = GetTransitionSubsystem())
+	{
+		TransitionSubsystem->RegisterTransitionBehavior(GetLevel(), this);
+	}
+}
+
+void AAvaTransitionBehaviorActor::PostLoad()
+{
+	Super::PostLoad();
+	ValidateTransitionTree();
+}
+
+#if WITH_EDITOR
+bool AAvaTransitionBehaviorActor::CanEditChange(const FProperty* InProperty) const
+{
+	return Super::CanEditChange(InProperty);
+}
+
+bool AAvaTransitionBehaviorActor::CanEditChange(const FEditPropertyChain& InPropertyChain) const
+{
+	if (!UObject::CanEditChange(InPropertyChain))
+	{
+		return false;
+	}
+
+	auto GetPropertyName =
+		[](TDoubleLinkedList<FProperty*>::TDoubleLinkedListNode* InChainNode)->FName
+		{
+			if (InChainNode && InChainNode->GetValue())
+			{
+				return InChainNode->GetValue()->GetFName();
+			}
+			return NAME_None;	
+		};
+
+	const FName MemberPropertyName = GetPropertyName(InPropertyChain.GetActiveMemberNode());
+	const FName ActivePropertyName = GetPropertyName(InPropertyChain.GetActiveNode());
+
+	// Disable change for StateTree in FStateTreeReference
+	return MemberPropertyName != GET_MEMBER_NAME_CHECKED(AAvaTransitionBehaviorActor, StateTreeReference)
+		|| ActivePropertyName != TEXT("StateTree");
+}
+#endif
+
+UAvaTransitionSubsystem* AAvaTransitionBehaviorActor::GetTransitionSubsystem() const
+{
+	UWorld* const World = GetWorld();
+	return World ? World->GetSubsystem<UAvaTransitionSubsystem>() : nullptr;
+}
+
+void AAvaTransitionBehaviorActor::ValidateTransitionTree()
+{
+	if (!ensureAlwaysMsgf(TransitionTree, TEXT("Transition Tree is null. Cannot not validate tree")))
+	{
+		return;
+	}
+
+	IAvaTransitionModule::FOnValidateTransitionTree& OnValidateTransitionTree = IAvaTransitionModule::Get().GetOnValidateTransitionTree();
+
+	if (!OnValidateTransitionTree.IsBound())
+	{
+#if WITH_EDITOR
+		if (GEditor)
+		{
+			ensureAlwaysMsgf(false, TEXT("OnValidateTransitionTree expected to be bound by FAvaTransitionEditorModule"));
+		}
+#endif
+		return;
+	}
+
+	OnValidateTransitionTree.Execute(TransitionTree);
+}
diff --git a/Engine/Plugins/VirtualProduction/Avalanche/Source/AvalancheTransition/Public/AvaTransitionTree.h b/Engine/Plugins/VirtualProduction/Avalanche/Source/AvalancheTransition/Public/AvaTransitionTree.h
new file mode 100644
--- /dev/null
+++ b/Engine/Plugins/VirtualProduction/Avalanche/Source/AvalancheTransition/Public/AvaTransitionTree.h
@@ -0,0 +1,70 @@
+// Copyright Epic Games, Inc. All Rights Reserved.
+
+#pragma once
+
+#include "AvaTagHandle.h"
+#include "AvaTransitionEnums.h"
+#include "StateTree.h"
+#include "AvaTransitionTree.generated.h"
+
+struct FAvaTransitionTask;
+
+/**
+ * Motion Design Transition Tree is a State Tree with the purpose of executing user-defined logic
+ * when there's a Transition between multiple scenes in multiple layers.
+ */
+UCLASS(MinimalAPI, DisplayName = "Motion Design Transition Tree")
+class UAvaTransitionTree : public UStateTree
+{
+	GENERATED_BODY()
+
+public:
+	AVALANCHETRANSITION_API FAvaTagHandle GetTransitionLayer() const;
+
+	AVALANCHETRANSITION_API void SetTransitionLayer(FAvaTagHandle InTransitionLayer);
+
+	AVALANCHETRANSITION_API bool IsEnabled() const;
+
+	AVALANCHETRANSITION_API void SetEnabled(bool bInEnabled);
+
+	/** Returns whether a Task of a given type exists and is enabled within an enabled state in the Transition Tree */
+	AVALANCHETRANSITION_API bool ContainsTask(const UScriptStruct* InTaskStruct) const;
+
+	/** Returns whether a Task of a given type exists and is enabled within an enabled state in the Transition Tree */
+	template <class InTaskType UE_REQUIRES(std::is_base_of_v<FAvaTransitionTask, InTaskType>)>
+	bool ContainsTask() const
+	{
+		return ContainsTask(InTaskType::StaticStruct());
+	}
+
+	void SetInstancingMode(EAvaTransitionInstancingMode InInstancingMode)
+	{
+		InstancingMode = InInstancingMode;
+	}
+
+	UFUNCTION(BlueprintPure, Category = "Transition Logic")
+	EAvaTransitionInstancingMode GetInstancingMode() const
+	{
+		return InstancingMode;
+	}
+
+	static FName GetEnabledPropertyName()
+	{
+		return GET_MEMBER_NAME_CHECKED(UAvaTransitionTree, bEnabled);
+	}
+
+private:
+	/** The Layer this Transition Logic Tree deals with */
+	UPROPERTY()
+	FAvaTagHandle TransitionLayer;
+
+	/**
+	 * Determines whether this Transition Logic is enabled, by default.
+	 * Can be overriden by a Transition Instance to force the logic to run regardless
+	 */
+	UPROPERTY()
+	bool bEnabled = true;
+
+	UPROPERTY()
+	EAvaTransitionInstancingMode InstancingMode = EAvaTransitionInstancingMode::New;
+};
diff --git a/Engine/Plugins/VirtualProduction/Avalanche/Source/AvalancheTransition/Public/Conditions/AvaTransitionLayerCondition.h b/Engine/Plugins/VirtualProduction/Avalanche/Source/AvalancheTransition/Public/Conditions/AvaTransitionLayerCondition.h
new file mode 100644
--- /dev/null
+++ b/Engine/Plugins/VirtualProduction/Avalanche/Source/AvalancheTransition/Public/Conditions/AvaTransitionLayerCondition.h
@@ -0,0 +1,58 @@
+// Copyright Epic Games, Inc. All Rights Reserved.
+
+#pragma once
+
+#include "AvaTagHandle.h"
+#include "AvaTransitionCondition.h"
+#include "AvaTransitionLayer.h"
+#include "AvaTransitionLayerCondition.generated.h"
+
+struct FAvaTransitionContext;
+
+USTRUCT()
+struct FAvaTransitionLayerConditionInstanceData
+{
+	GENERATED_BODY()
+
+	UPROPERTY(EditAnywhere, Category="Parameter")
+	EAvaTransitionLayerCompareType LayerType = EAvaTransitionLayerCompareType::Same;
+
+	UPROPERTY(EditAnywhere, Category="Parameter", meta=(EditCondition="LayerType==EAvaTransitionLayerCompareType::MatchingTag", EditConditionHides))
+	FAvaTagHandle SpecificLayer;
+};
+
+USTRUCT(Category="Transition Logic", meta=(Hidden))
+struct AVALANCHETRANSITION_API FAvaTransitionLayerCondition : public FAvaTransitionCondition
+{
+	GENERATED_BODY()
+
+	using FInstanceDataType = FAvaTransitionLayerConditionInstanceData;
+
+	PRAGMA_DISABLE_DEPRECATION_WARNINGS
+	FAvaTransitionLayerCondition() = default;
+	virtual ~FAvaTransitionLayerCondition() override = default;
+	FAvaTransitionLayerCondition(const FAvaTransitionLayerCondition&) = default;
+	FAvaTransitionLayerCondition(FAvaTransitionLayerCondition&&) = default;
+	FAvaTransitionLayerCondition& operator=(const FAvaTransitionLayerCondition&) = default;
+	FAvaTransitionLayerCondition& operator=(FAvaTransitionLayerCondition&&) = default;
+	PRAGMA_ENABLE_DEPRECATION_WARNINGS
+
+	//~ Begin FStateTreeNodeBase
+#if WITH_EDITOR
+	virtual FText GetDescription(const FGuid& InId, FStateTreeDataView InInstanceDataView, const IStateTreeBindingLookup& InBindingLookup, EStateTreeNodeFormatting InFormatting) const override;
+#endif
+	virtual const UStruct* GetInstanceDataType() const override { return FInstanceDataType::StaticStruct(); }
+	virtual void PostLoad(FStateTreeDataView InInstanceDataView) override;
+	//~ End FStateTreeNodeBase
+
+	/** Gets all the Behavior Instances that match the Layer Query. Always excludes the Instance belonging to this Transition */
+	TArray<const FAvaTransitionBehaviorInstance*> QueryBehaviorInstances(const FStateTreeExecutionContext& InContext) const;
+
+	UE_DEPRECATED(5.5, "LayerType has been moved to Instance Data")
+	UPROPERTY(meta = (DeprecatedProperty, DeprecationMessage = "Use the Instance Data LayerType instead"))
+	EAvaTransitionLayerCompareType LayerType_DEPRECATED = EAvaTransitionLayerCompareType::None;
+
+	UE_DEPRECATED(5.5, "SpecificLayer has been moved to Instance Data")
+	UPROPERTY(meta = (DeprecatedProperty, DeprecationMessage = "Use the Instance Data SpecificLayer instead"))
+	FAvaTagHandle SpecificLayer_DEPRECATED;
+};
diff --git a/Engine/Plugins/VirtualProduction/Avalanche/Source/AvalancheTransition/Public/Conditions/AvaTransitionSceneMatchCondition.h b/Engine/Plugins/VirtualProduction/Avalanche/Source/AvalancheTransition/Public/Conditions/AvaTransitionSceneMatchCondition.h
new file mode 100644
--- /dev/null
+++ b/Engine/Plugins/VirtualProduction/Avalanche/Source/AvalancheTransition/Public/Conditions/AvaTransitionSceneMatchCondition.h
@@ -0,0 +1,49 @@
+// Copyright Epic Games, Inc. All Rights Reserved.
+
+#pragma once
+
+#include "AvaTransitionEnums.h"
+#include "AvaTransitionLayerCondition.h"
+#include "AvaTransitionSceneMatchCondition.generated.h"
+
+USTRUCT()
+struct FAvaTransitionSceneMatchConditionInstanceData : public FAvaTransitionLayerConditionInstanceData
+{
+	GENERATED_BODY()
+
+	UPROPERTY(EditAnywhere, Category="Parameter")
+    EAvaTransitionComparisonResult SceneComparisonType = EAvaTransitionComparisonResult::None;
+};
+
+USTRUCT(DisplayName="Compare other Scene in Layer", Category="Transition Logic")
+struct AVALANCHETRANSITION_API FAvaTransitionSceneMatchCondition : public FAvaTransitionLayerCondition
+{
+	GENERATED_BODY()
+
+	using FInstanceDataType = FAvaTransitionSceneMatchConditionInstanceData;
+
+	PRAGMA_DISABLE_DEPRECATION_WARNINGS
+	FAvaTransitionSceneMatchCondition() = default;
+	virtual ~FAvaTransitionSceneMatchCondition() override = default;
+	FAvaTransitionSceneMatchCondition(const FAvaTransitionSceneMatchCondition&) = default;
+	FAvaTransitionSceneMatchCondition(FAvaTransitionSceneMatchCondition&&) = default;
+	FAvaTransitionSceneMatchCondition& operator=(const FAvaTransitionSceneMatchCondition&) = default;
+	FAvaTransitionSceneMatchCondition& operator=(FAvaTransitionSceneMatchCondition&&) = default;
+	PRAGMA_ENABLE_DEPRECATION_WARNINGS
+
+	//~ Begin FStateTreeNodeBase
+#if WITH_EDITOR
+	virtual FText GetDescription(const FGuid& InId, FStateTreeDataView InInstanceDataView, const IStateTreeBindingLookup& InBindingLookup, EStateTreeNodeFormatting InFormatting) const override;
+#endif
+	virtual const UStruct* GetInstanceDataType() const override { return FInstanceDataType::StaticStruct(); }
+	virtual void PostLoad(FStateTreeDataView InInstanceDataView) override;
+	//~ End FStateTreeNodeBase
+
+	//~ Begin FStateTreeConditionBase
+	virtual bool TestCondition(FStateTreeExecutionContext& InContext) const override;
+	//~ End FStateTreeConditionBase
+
+	UE_DEPRECATED(5.5, "SceneComparisonType has been moved to Instance Data")
+	UPROPERTY(meta = (DeprecatedProperty, DeprecationMessage = "Use the Instance Data SceneComparisonType instead"))
+	EAvaTransitionComparisonResult SceneComparisonType_DEPRECATED = EAvaTransitionComparisonResult::None;
+};
diff --git a/Engine/Plugins/VirtualProduction/Avalanche/Source/AvalancheTransition/Public/Conditions/AvaTransitionStateMatchCondition.h b/Engine/Plugins/VirtualProduction/Avalanche/Source/AvalancheTransition/Public/Conditions/AvaTransitionStateMatchCondition.h
new file mode 100644
--- /dev/null
+++ b/Engine/Plugins/VirtualProduction/Avalanche/Source/AvalancheTransition/Public/Conditions/AvaTransitionStateMatchCondition.h
@@ -0,0 +1,49 @@
+// Copyright Epic Games, Inc. All Rights Reserved.
+
+#pragma once
+
+#include "AvaTransitionEnums.h"
+#include "AvaTransitionLayerCondition.h"
+#include "AvaTransitionStateMatchCondition.generated.h"
+
+USTRUCT()
+struct FAvaTransitionStateMatchConditionInstanceData : public FAvaTransitionLayerConditionInstanceData
+{
+	GENERATED_BODY()
+
+	UPROPERTY(EditAnywhere, Category="Parameter")
+	EAvaTransitionRunState TransitionState = EAvaTransitionRunState::Running;
+};
+
+USTRUCT(DisplayName="Check State of other Scene in Layer", Category="Transition Logic")
+struct AVALANCHETRANSITION_API FAvaTransitionStateMatchCondition : public FAvaTransitionLayerCondition
+{
+	GENERATED_BODY()
+
+	using FInstanceDataType = FAvaTransitionStateMatchConditionInstanceData;
+
+	PRAGMA_DISABLE_DEPRECATION_WARNINGS
+	FAvaTransitionStateMatchCondition() = default;
+	virtual ~FAvaTransitionStateMatchCondition() override = default;
+	FAvaTransitionStateMatchCondition(const FAvaTransitionStateMatchCondition&) = default;
+	FAvaTransitionStateMatchCondition(FAvaTransitionStateMatchCondition&&) = default;
+	FAvaTransitionStateMatchCondition& operator=(const FAvaTransitionStateMatchCondition&) = default;
+	FAvaTransitionStateMatchCondition& operator=(FAvaTransitionStateMatchCondition&&) = default;
+	PRAGMA_ENABLE_DEPRECATION_WARNINGS
+
+	//~ Begin FStateTreeNodeBase
+#if WITH_EDITOR
+	virtual FText GetDescription(const FGuid& InId, FStateTreeDataView InInstanceDataView, const IStateTreeBindingLookup& InBindingLookup, EStateTreeNodeFormatting InFormatting) const override;
+#endif
+	virtual const UStruct* GetInstanceDataType() const override { return FAvaTransitionStateMatchConditionInstanceData::StaticStruct(); }
+	virtual void PostLoad(FStateTreeDataView InInstanceDataView) override;
+	//~ End FStateTreeNodeBase
+
+	//~ Begin FStateTreeConditionBase
+	virtual bool TestCondition(FStateTreeExecutionContext& InContext) const override;
+	//~ End FStateTreeConditionBase
+
+	UE_DEPRECATED(5.5, "TransitionState has been moved to Instance Data")
+	UPROPERTY(meta = (DeprecatedProperty, DeprecationMessage = "Use the Instance Data TransitionState instead"))
+	EAvaTransitionRunState TransitionState_DEPRECATED = EAvaTransitionRunState::Unknown;
+};
diff --git a/Engine/Plugins/VirtualProduction/Avalanche/Source/AvalancheTransition/Public/Conditions/AvaTransitionTypeMatchCondition.h b/Engine/Plugins/VirtualProduction/Avalanche/Source/AvalancheTransition/Public/Conditions/AvaTransitionTypeMatchCondition.h
new file mode 100644
--- /dev/null
+++ b/Engine/Plugins/VirtualProduction/Avalanche/Source/AvalancheTransition/Public/Conditions/AvaTransitionTypeMatchCondition.h
@@ -0,0 +1,49 @@
+// Copyright Epic Games, Inc. All Rights Reserved.
+
+#pragma once
+
+#include "AvaTransitionEnums.h"
+#include "Conditions/AvaTransitionCondition.h"
+#include "AvaTransitionTypeMatchCondition.generated.h"
+
+USTRUCT()
+struct FAvaTransitionTypeMatchConditionInstanceData
+{
+	GENERATED_BODY()
+
+	UPROPERTY(EditAnywhere, Category="Parameter")
+	EAvaTransitionType TransitionType = EAvaTransitionType::In;
+};
+
+USTRUCT(DisplayName="My Transition Type is", Category="Transition Logic")
+struct AVALANCHETRANSITION_API FAvaTransitionTypeMatchCondition : public FAvaTransitionCondition
+{
+	GENERATED_BODY()
+
+	using FInstanceDataType = FAvaTransitionTypeMatchConditionInstanceData;
+
+	PRAGMA_DISABLE_DEPRECATION_WARNINGS
+	FAvaTransitionTypeMatchCondition() = default;
+	virtual ~FAvaTransitionTypeMatchCondition() override = default;
+	FAvaTransitionTypeMatchCondition(const FAvaTransitionTypeMatchCondition&) = default;
+	FAvaTransitionTypeMatchCondition(FAvaTransitionTypeMatchCondition&&) = default;
+	FAvaTransitionTypeMatchCondition& operator=(const FAvaTransitionTypeMatchCondition&) = default;
+	FAvaTransitionTypeMatchCondition& operator=(FAvaTransitionTypeMatchCondition&&) = default;
+	PRAGMA_ENABLE_DEPRECATION_WARNINGS
+
+	//~ Begin FStateTreeNodeBase
+#if WITH_EDITOR
+	virtual FText GetDescription(const FGuid& InId, FStateTreeDataView InInstanceDataView, const IStateTreeBindingLookup& InBindingLookup, EStateTreeNodeFormatting InFormatting) const override;
+#endif
+	virtual const UStruct* GetInstanceDataType() const override { return FInstanceDataType::StaticStruct(); }
+	virtual void PostLoad(FStateTreeDataView InInstanceDataView) override;
+	//~ End FStateTreeNodeBase
+
+	//~ Begin FStateTreeConditionBase
+	virtual bool TestCondition(FStateTreeExecutionContext& InContext) const override;
+	//~ End FStateTreeConditionBase
+
+	UE_DEPRECATED(5.5, "TransitionType has been moved to Instance Data")
+	UPROPERTY(meta = (DeprecatedProperty, DeprecationMessage = "Use the Instance Data TransitionType instead"))
+	EAvaTransitionType TransitionType_DEPRECATED = EAvaTransitionType::None;
+};
diff --git a/Engine/Plugins/VirtualProduction/Avalanche/Source/AvalancheTransition/Public/Tasks/AvaTransitionLayerTask.h b/Engine/Plugins/VirtualProduction/Avalanche/Source/AvalancheTransition/Public/Tasks/AvaTransitionLayerTask.h
new file mode 100644
--- /dev/null
+++ b/Engine/Plugins/VirtualProduction/Avalanche/Source/AvalancheTransition/Public/Tasks/AvaTransitionLayerTask.h
@@ -0,0 +1,56 @@
+// Copyright Epic Games, Inc. All Rights Reserved.
+
+#pragma once
+
+#include "AvaTagHandle.h"
+#include "AvaTransitionLayer.h"
+#include "AvaTransitionTask.h"
+#include "AvaTransitionLayerTask.generated.h"
+
+USTRUCT()
+struct FAvaTransitionLayerTaskInstanceData
+{
+	GENERATED_BODY()
+
+	UPROPERTY(EditAnywhere, Category="Parameter")
+	EAvaTransitionLayerCompareType LayerType = EAvaTransitionLayerCompareType::Same;
+
+	UPROPERTY(EditAnywhere, Category="Parameter", meta=(EditCondition="LayerType==EAvaTransitionLayerCompareType::MatchingTag", EditConditionHides))
+	FAvaTagHandle SpecificLayer;
+};
+
+USTRUCT(meta=(Hidden))
+struct AVALANCHETRANSITION_API FAvaTransitionLayerTask : public FAvaTransitionTask
+{
+	GENERATED_BODY()
+
+	using FInstanceDataType = FAvaTransitionLayerTaskInstanceData;
+
+	PRAGMA_DISABLE_DEPRECATION_WARNINGS
+	FAvaTransitionLayerTask() = default;
+	virtual ~FAvaTransitionLayerTask() override = default;
+	FAvaTransitionLayerTask(const FAvaTransitionLayerTask&) = default;
+	FAvaTransitionLayerTask(FAvaTransitionLayerTask&&) = default;
+	FAvaTransitionLayerTask& operator=(const FAvaTransitionLayerTask&) = default;
+	FAvaTransitionLayerTask& operator=(FAvaTransitionLayerTask&&) = default;
+	PRAGMA_ENABLE_DEPRECATION_WARNINGS
+
+	//~ Begin FStateTreeNodeBase
+#if WITH_EDITOR
+	virtual FText GetDescription(const FGuid& InId, FStateTreeDataView InInstanceDataView, const IStateTreeBindingLookup& InBindingLookup, EStateTreeNodeFormatting InFormatting) const override;
+#endif
+	virtual const UStruct* GetInstanceDataType() const override { return FInstanceDataType::StaticStruct(); }
+	virtual void PostLoad(FStateTreeDataView InInstanceDataView) override;
+	//~ End FStateTreeNodeBase
+
+	/** Gets all the Behavior Instances that match the Layer Query. Always excludes the Instance belonging to this Transition */
+	TArray<const FAvaTransitionBehaviorInstance*> QueryBehaviorInstances(const FStateTreeExecutionContext& InContext) const;
+
+	UE_DEPRECATED(5.5, "LayerType has been moved to Instance Data")
+	UPROPERTY(meta = (DeprecatedProperty, DeprecationMessage = "Use the Instance Data LayerType instead"))
+	EAvaTransitionLayerCompareType LayerType_DEPRECATED = EAvaTransitionLayerCompareType::None;
+
+	UE_DEPRECATED(5.5, "SpecificLayer has been moved to Instance Data")
+	UPROPERTY(meta = (DeprecatedProperty, DeprecationMessage = "Use the Instance Data SpecificLayer instead"))
+	FAvaTagHandle SpecificLayer_DEPRECATED;
+};
diff --git a/Engine/Plugins/VirtualProduction/Avalanche/Source/AvalancheTransitionEditor/AvalancheTransitionEditor.Build.cs b/Engine/Plugins/VirtualProduction/Avalanche/Source/AvalancheTransitionEditor/AvalancheTransitionEditor.Build.cs
new file mode 100644
--- /dev/null
+++ b/Engine/Plugins/VirtualProduction/Avalanche/Source/AvalancheTransitionEditor/AvalancheTransitionEditor.Build.cs
@@ -0,0 +1,64 @@
+// Copyright Epic Games, Inc. All Rights Reserved.
+
+using UnrealBuildTool;
+
+public class AvalancheTransitionEditor : ModuleRules
+{
+    public AvalancheTransitionEditor(ReadOnlyTargetRules Target) : base(Target)
+    {
+        PCHUsage = ModuleRules.PCHUsageMode.UseExplicitOrSharedPCHs;
+
+        PublicDependencyModuleNames.AddRange(
+            new string[]
+            {
+                "AvalancheTag",
+                "Core",
+                "CoreUObject",
+                "StateTreeEditorModule",
+            }
+        );
+
+        PrivateDependencyModuleNames.AddRange(
+            new string[]
+            {
+                "ApplicationCore",
+                "AssetDefinition",
+                "AssetTools",
+                "AvalancheCore",
+                "AvalancheTransition",
+                "DeveloperSettings",
+                "EditorStyle",
+                "Engine",
+                "InputCore",
+                "MessageLog",
+                "Projects",
+                "PropertyBindingUtils",
+                "PropertyEditor",
+                "Slate",
+                "SlateCore",
+                "StateTreeModule",
+                "StructUtilsEditor",
+                "ToolMenus",
+                "UnrealEd",
+            }
+        );
+
+        // Until we split use of TraceServices and TraceAnalysis in the debugger we enable only on platforms supporting both at the moment
+        if (Target.Platform.IsInGroup(UnrealPlatformGroup.Desktop) && (Target.Configuration != UnrealTargetConfiguration.Shipping || Target.bBuildEditor))
+        {
+            PrivateDefinitions.Add("WITH_STATETREE_DEBUGGER=1");
+            PrivateDependencyModuleNames.AddRange(
+                new string[]
+                {
+                    "TraceAnalysis",
+                    "TraceLog",
+                    "TraceServices",
+                }
+            );
+        }
+        else
+        {
+            PrivateDefinitions.Add("WITH_STATETREE_DEBUGGER=0");
+        }
+    }
+}
