Motion Design: marked the motion design plugins moved out of experimental to beta.

diff --git a/Engine/Plugins/VirtualProduction/ActorModifier/ActorModifier.uplugin b/Engine/Plugins/VirtualProduction/ActorModifier/ActorModifier.uplugin
--- a/Engine/Plugins/VirtualProduction/ActorModifier/ActorModifier.uplugin
+++ b/Engine/Plugins/VirtualProduction/ActorModifier/ActorModifier.uplugin
@@ -4,15 +4,15 @@
 	"VersionName": "1.0",
 	"FriendlyName": "Actor Modifier",
 	"Description": "Actual implementation of modifiers for actors based on ActorModifierCore plugin",
-	"Category": "Experimental",
+	"Category": "Virtual Production",
 	"CreatedBy": "Epic Games, Inc.",
 	"CreatedByURL": "https://epicgames.com",
 	"DocsURL": "",
 	"MarketplaceURL": "",
 	"SupportURL": "",
 	"CanContainContent": true,
-	"IsBetaVersion": false,
-	"IsExperimentalVersion": true,
+	"IsBetaVersion": true,
+	"IsExperimentalVersion": false,
 	"Installed": false,
 	"PlatformAllowList": [
 		"Win64",
diff --git a/Engine/Plugins/VirtualProduction/ActorModifierCore/ActorModifierCore.uplugin b/Engine/Plugins/VirtualProduction/ActorModifierCore/ActorModifierCore.uplugin
--- a/Engine/Plugins/VirtualProduction/ActorModifierCore/ActorModifierCore.uplugin
+++ b/Engine/Plugins/VirtualProduction/ActorModifierCore/ActorModifierCore.uplugin
@@ -4,15 +4,15 @@
 	"VersionName": "1.0",
 	"FriendlyName": "Actor Modifier Core",
 	"Description": "Use modifier objects on actors to apply a custom behavior",
-	"Category": "Experimental",
+	"Category": "Virtual Production",
 	"CreatedBy": "Epic Games, Inc.",
 	"CreatedByURL": "https://epicgames.com",
 	"DocsURL": "",
 	"MarketplaceURL": "",
 	"SupportURL": "",
 	"CanContainContent": true,
-	"IsBetaVersion": false,
-	"IsExperimentalVersion": true,
+	"IsBetaVersion": true,
+	"IsExperimentalVersion": false,
 	"Installed": false,
 	"PlatformAllowList": [
 		"Win64",
diff --git a/Engine/Plugins/VirtualProduction/Avalanche/Avalanche.uplugin b/Engine/Plugins/VirtualProduction/Avalanche/Avalanche.uplugin
--- a/Engine/Plugins/VirtualProduction/Avalanche/Avalanche.uplugin
+++ b/Engine/Plugins/VirtualProduction/Avalanche/Avalanche.uplugin
@@ -4,15 +4,15 @@
 	"VersionName": "1.0",
 	"FriendlyName": "Motion Design",
 	"Description": "Compositing, designer and broadcasting tool.\n\nPlugin Dependencies: Advanced Renamer, Custom Details View, Dynamic Material, Geometry Cache, Geometry Scripting, Media Compositing, Media IO Framework, Mesh Modeling Toolset Exp, Remote Control, SVG Importer, Text3D and ActorModifierCore.",
-	"Category": "Experimental",
+	"Category": "Virtual Production",
 	"CreatedBy": "Epic Games, Inc.",
 	"CreatedByURL": "https://epicgames.com",
 	"DocsURL": "",
 	"MarketplaceURL": "",
 	"SupportURL": "",
 	"CanContainContent": true,
-	"IsBetaVersion": false,
-	"IsExperimentalVersion": true,
+	"IsBetaVersion": true,
+	"IsExperimentalVersion": false,
 	"Installed": false,
 	"EnabledByDefault": false,
 	"SupportedTargetPlatforms": [
diff --git a/Engine/Plugins/VirtualProduction/ClonerEffector/ClonerEffector.uplugin b/Engine/Plugins/VirtualProduction/ClonerEffector/ClonerEffector.uplugin
--- a/Engine/Plugins/VirtualProduction/ClonerEffector/ClonerEffector.uplugin
+++ b/Engine/Plugins/VirtualProduction/ClonerEffector/ClonerEffector.uplugin
@@ -4,14 +4,14 @@
 	"VersionName": "1.0",
 	"FriendlyName": "Cloners and Effectors",
 	"Description": "Niagara based cloner system with various layouts and effector affecting each clone instances",
-	"Category": "Experimental",
+	"Category": "Virtual Production",
 	"CreatedBy": "Epic Games, Inc.",
 	"CreatedByURL": "https://epicgames.com",
 	"DocsURL": "",
 	"MarketplaceURL": "",
 	"CanContainContent": true,
-	"IsBetaVersion": false,
-	"IsExperimentalVersion": true,
+	"IsBetaVersion": true,
+	"IsExperimentalVersion": false,
 	"Installed": false,
 	"Modules": [
 		{
diff --git a/Engine/Plugins/VirtualProduction/CustomDetailsView/CustomDetailsView.uplugin b/Engine/Plugins/VirtualProduction/CustomDetailsView/CustomDetailsView.uplugin
--- a/Engine/Plugins/VirtualProduction/CustomDetailsView/CustomDetailsView.uplugin
+++ b/Engine/Plugins/VirtualProduction/CustomDetailsView/CustomDetailsView.uplugin
@@ -4,15 +4,15 @@
 	"VersionName": "1.0",
 	"FriendlyName": "Custom Details View",
 	"Description": "",
-	"Category": "Experimental",
+	"Category": "Virtual Production",
 	"CreatedBy": "Epic Games, Inc.",
 	"CreatedByURL": "https://epicgames.com",
 	"DocsURL": "",
 	"MarketplaceURL": "",
 	"SupportURL": "",
 	"CanContainContent": false,
-	"IsBetaVersion": false,
-	"IsExperimentalVersion": true,
+	"IsBetaVersion": true,
+	"IsExperimentalVersion": false,
 	"Installed": false,
 	"Modules": [
 		{
diff --git a/Engine/Plugins/VirtualProduction/DynamicMaterial/DynamicMaterial.uplugin b/Engine/Plugins/VirtualProduction/DynamicMaterial/DynamicMaterial.uplugin
--- a/Engine/Plugins/VirtualProduction/DynamicMaterial/DynamicMaterial.uplugin
+++ b/Engine/Plugins/VirtualProduction/DynamicMaterial/DynamicMaterial.uplugin
@@ -4,15 +4,15 @@
 	"VersionName": "1.0",
 	"FriendlyName": "Material Designer",
 	"Description": "Compact dynamic material creator and editor, similar in style to other DDCs.",
-	"Category": "Experimental",
+	"Category": "Virtual Production",
 	"CreatedBy": "Epic Games, Inc.",
 	"CreatedByURL": "https://epicgames.com",
 	"DocsURL": "",
 	"MarketplaceURL": "",
 	"SupportURL": "",
 	"CanContainContent": true,
-	"IsBetaVersion": false,
-	"IsExperimentalVersion": true,
+	"IsBetaVersion": true,
+	"IsExperimentalVersion": false,
 	"Installed": false,
 	"Modules": [
 		{
diff --git a/Engine/Plugins/VirtualProduction/GeometryMask/GeometryMask.uplugin b/Engine/Plugins/VirtualProduction/GeometryMask/GeometryMask.uplugin
--- a/Engine/Plugins/VirtualProduction/GeometryMask/GeometryMask.uplugin
+++ b/Engine/Plugins/VirtualProduction/GeometryMask/GeometryMask.uplugin
@@ -4,14 +4,14 @@
 	"VersionName": "1.0",
 	"FriendlyName": "Geometry Mask",
 	"Description": "",
-	"Category": "Experimental",
+	"Category": "Virtual Production",
 	"CreatedBy": "Epic Games, Inc.",
 	"CreatedByURL": "https://epicgames.com",
 	"DocsURL": "",
 	"MarketplaceURL": "",
 	"CanContainContent": true,
-	"IsBetaVersion": false,
-	"IsExperimentalVersion": true,
+	"IsBetaVersion": true,
+	"IsExperimentalVersion": false,
 	"Installed": false,
 	"Modules": [
 		{
diff --git a/Engine/Plugins/VirtualProduction/OperatorStack/OperatorStack.uplugin b/Engine/Plugins/VirtualProduction/OperatorStack/OperatorStack.uplugin
--- a/Engine/Plugins/VirtualProduction/OperatorStack/OperatorStack.uplugin
+++ b/Engine/Plugins/VirtualProduction/OperatorStack/OperatorStack.uplugin
@@ -4,14 +4,14 @@
 	"VersionName": "1.0",
 	"FriendlyName": "Operator Stack",
 	"Description": "Allows you to display linked objects in a container with an intuitive UI",
-	"Category": "Experimental",
+	"Category": "Virtual Production",
 	"CreatedBy": "Epic Games, Inc.",
 	"CreatedByURL": "https://epicgames.com",
 	"DocsURL": "",
 	"MarketplaceURL": "",
 	"CanContainContent": false,
-	"IsBetaVersion": false,
-	"IsExperimentalVersion": true,
+	"IsBetaVersion": true,
+	"IsExperimentalVersion": false,
 	"Installed": false,
 	"Modules": [
 		{
diff --git a/Engine/Plugins/VirtualProduction/PropertyAnimator/PropertyAnimator.uplugin b/Engine/Plugins/VirtualProduction/PropertyAnimator/PropertyAnimator.uplugin
--- a/Engine/Plugins/VirtualProduction/PropertyAnimator/PropertyAnimator.uplugin
+++ b/Engine/Plugins/VirtualProduction/PropertyAnimator/PropertyAnimator.uplugin
@@ -4,14 +4,14 @@
 	"VersionName": "1.0",
 	"FriendlyName": "Property Animator",
 	"Description": "Re-usable behaviors to animate the value of one or more properties",
-	"Category": "Experimental",
+	"Category": "Virtual Production",
 	"CreatedBy": "Epic Games, Inc.",
 	"CreatedByURL": "https://epicgames.com",
 	"DocsURL": "",
 	"MarketplaceURL": "",
 	"CanContainContent": true,
-	"IsBetaVersion": false,
-	"IsExperimentalVersion": true,
+	"IsBetaVersion": true,
+	"IsExperimentalVersion": false,
 	"Installed": false,
 	"Modules": [
 		{
diff --git a/Engine/Plugins/VirtualProduction/PropertyAnimatorCore/PropertyAnimatorCore.uplugin b/Engine/Plugins/VirtualProduction/PropertyAnimatorCore/PropertyAnimatorCore.uplugin
--- a/Engine/Plugins/VirtualProduction/PropertyAnimatorCore/PropertyAnimatorCore.uplugin
+++ b/Engine/Plugins/VirtualProduction/PropertyAnimatorCore/PropertyAnimatorCore.uplugin
@@ -4,14 +4,14 @@
 	"VersionName": "1.0",
 	"FriendlyName": "Property Animator Core",
 	"Description": "Re-usable behaviors to control properties at runtime and in editor",
-	"Category": "Experimental",
+	"Category": "Virtual Production",
 	"CreatedBy": "Epic Games, Inc.",
 	"CreatedByURL": "https://epicgames.com",
 	"DocsURL": "",
 	"MarketplaceURL": "",
 	"CanContainContent": true,
-	"IsBetaVersion": false,
-	"IsExperimentalVersion": true,
+	"IsBetaVersion": true,
+	"IsExperimentalVersion": false,
 	"Installed": false,
 	"Modules": [
 		{
diff --git a/Engine/Plugins/VirtualProduction/StormSync/StormSync.uplugin b/Engine/Plugins/VirtualProduction/StormSync/StormSync.uplugin
--- a/Engine/Plugins/VirtualProduction/StormSync/StormSync.uplugin
+++ b/Engine/Plugins/VirtualProduction/StormSync/StormSync.uplugin
@@ -4,15 +4,15 @@
 	"VersionName": "1.0",
 	"FriendlyName": "Storm Sync",
 	"Description": "Sync, Pull, Push, asset dependencies.\n\nThis plugin is a recommended part of the Motion Design work flow.",
-	"Category": "Experimental",
+	"Category": "Virtual Production",
 	"CreatedBy": "Epic Games, Inc.",
 	"CreatedByURL": "https://epicgames.com",
 	"DocsURL": "",
 	"MarketplaceURL": "",
 	"SupportURL": "",
 	"CanContainContent": true,
-	"IsBetaVersion": false,
-	"IsExperimentalVersion": true,
+	"IsBetaVersion": true,
+	"IsExperimentalVersion": false,
 	"Installed": false,
 	"Modules": [
 		{
diff --git a/Engine/Plugins/VirtualProduction/StormSyncAvalancheBridge/StormSyncAvalancheBridge.uplugin b/Engine/Plugins/VirtualProduction/StormSyncAvalancheBridge/StormSyncAvalancheBridge.uplugin
--- a/Engine/Plugins/VirtualProduction/StormSyncAvalancheBridge/StormSyncAvalancheBridge.uplugin
+++ b/Engine/Plugins/VirtualProduction/StormSyncAvalancheBridge/StormSyncAvalancheBridge.uplugin
@@ -4,15 +4,15 @@
 	"VersionName": "1.0",
 	"FriendlyName": "Storm Sync Motion Design Bridge",
 	"Description": "Plugin bridge between Motion Design Plugin and Storm Sync to provide in-editor integration to synchronize assets",
-	"Category": "Experimental",
+	"Category": "Virtual Production",
 	"CreatedBy": "Epic Games, Inc.",
 	"CreatedByURL": "https://epicgames.com",
 	"DocsURL": "",
 	"MarketplaceURL": "",
 	"SupportURL": "",
 	"CanContainContent": true,
-	"IsBetaVersion": false,
-	"IsExperimentalVersion": true,
+	"IsBetaVersion": true,
+	"IsExperimentalVersion": false,
 	"Installed": false,
 	"Modules": [
 		{
