[Motion Design] Rundown Server: implementing API for new commands for transitions and layers.

- Stop Transition command
- Stop/ForceStop Layer command

diff --git a/Engine/Plugins/VirtualProduction/Avalanche/Source/AvalancheMedia/Internal/Rundown/AvaRundownMessages.h b/Engine/Plugins/VirtualProduction/Avalanche/Source/AvalancheMedia/Internal/Rundown/AvaRundownMessages.h
--- a/Engine/Plugins/VirtualProduction/Avalanche/Source/AvalancheMedia/Internal/Rundown/AvaRundownMessages.h
+++ b/Engine/Plugins/VirtualProduction/Avalanche/Source/AvalancheMedia/Internal/Rundown/AvaRundownMessages.h
@@ -1069,6 +1069,68 @@ public:
 	FString PreviewChannelName;
 };
 
+/** Supported Transition actions for playback. */
+UENUM()
+enum class EAvaRundownTransitionActions
+{
+	None UMETA(Hidden),
+	/** This action will forcefully stop specified transitions. */
+	ForceStop
+};
+
+/**
+ * Command to override transition logic directly.
+ * As it currently stands, we can only have 1 transition per channel.
+ * If there is an issue with it, it may block further playback.
+ */
+USTRUCT()
+struct FAvaRundownTransitionAction : public FAvaRundownMsgBase
+{
+	GENERATED_BODY()
+
+public:
+	/** Specifies the channel that is the target of this action command. */
+	UPROPERTY()
+	FString ChannelName;
+
+	/** Specifies the page transition action to execute. */
+	UPROPERTY()
+	EAvaRundownTransitionActions Action = EAvaRundownTransitionActions::None;
+};
+
+/** Supported Page Logic Layer actions for playback. */
+UENUM()
+enum class EAvaRundownTransitionLayerActions
+{
+	None UMETA(Hidden),
+	/** Trigger the out transition for the specified layer. */
+	Stop,
+	/** Forcefully stop, without transition, pages on the specified layer. */
+	ForceStop
+};
+
+/**
+ * Command to override transition logic.
+ */
+USTRUCT()
+struct FAvaRundownTransitionLayerAction : public FAvaRundownMsgBase
+{
+	GENERATED_BODY()
+
+public:
+	/** Specifies the channel that is the target of this action command. */
+	UPROPERTY()
+	FString ChannelName;
+
+	/** Specifies the transition logic layers for this action command. */
+	UPROPERTY()
+	TArray<FString> LayerNames;
+	
+	/** Specifies the page layer action to execute. */
+	UPROPERTY()
+	EAvaRundownTransitionLayerActions Action = EAvaRundownTransitionLayerActions::None;
+};
+
 /**
  * This message is sent by the server when a page sequence event occurs.
  */
diff --git a/Engine/Plugins/VirtualProduction/Avalanche/Source/AvalancheMedia/Private/Rundown/AvaRundownServer.cpp b/Engine/Plugins/VirtualProduction/Avalanche/Source/AvalancheMedia/Private/Rundown/AvaRundownServer.cpp
--- a/Engine/Plugins/VirtualProduction/Avalanche/Source/AvalancheMedia/Private/Rundown/AvaRundownServer.cpp
+++ b/Engine/Plugins/VirtualProduction/Avalanche/Source/AvalancheMedia/Private/Rundown/AvaRundownServer.cpp
@@ -423,6 +423,8 @@ void FAvaRundownServer::Init(const FString& InAssignedHostName)
 	.Handling<FAvaRundownPagePreviewAction>(this, &FAvaRundownServer::HandlePagePreviewAction)
 	.Handling<FAvaRundownPageActions>(this, &FAvaRundownServer::HandlePageActions)
 	.Handling<FAvaRundownPagePreviewActions>(this, &FAvaRundownServer::HandlePagePreviewActions)
+	.Handling<FAvaRundownTransitionAction>(this, &FAvaRundownServer::HandleTransitionAction)
+	.Handling<FAvaRundownTransitionLayerAction>(this, &FAvaRundownServer::HandleTransitionLayerAction)
 	.Handling<FAvaRundownGetProfiles>(this, &FAvaRundownServer::HandleGetProfiles)
 	.Handling<FAvaRundownDuplicateProfile>(this, &FAvaRundownServer::HandleDuplicateProfile)
 	.Handling<FAvaRundownCreateProfile>(this, &FAvaRundownServer::HandleCreateProfile)
@@ -1520,6 +1522,196 @@ void FAvaRundownServer::HandlePagePreviewActions(const FAvaRundownPagePreviewAct
 	HandlePageActions(RequestInfo, InMessage.PageIds, true, FName(InMessage.PreviewChannelName), InMessage.Action);
 }
 
+void FAvaRundownServer::HandleTransitionAction(const FAvaRundownTransitionAction& InMessage, const TSharedRef<IMessageContext, ESPMode::ThreadSafe>& InContext)
+{
+	const FRequestInfo RequestInfo = {InMessage.RequestId, InContext->GetSender()};
+	const FName ActionName = StaticEnum<EAvaRundownTransitionActions>()->GetNameByValue((int64)InMessage.Action);
+
+	bool bSuccess = false;
+	FString FailureReason;
+
+	// This is the response handler for this command.
+	auto HandleCommandResponse = [this, &RequestInfo, &bSuccess, &FailureReason, ActionName]()
+	{
+		if (bSuccess)
+		{
+			SendMessage(RequestInfo.Sender, RequestInfo.RequestId, ELogVerbosity::Log,
+				TEXT("Transition Action \"%s\": Ok."), *ActionName.ToString());
+		}
+		else if (!FailureReason.IsEmpty())
+		{
+			LogAndSendMessage(RequestInfo.Sender, RequestInfo.RequestId, ELogVerbosity::Error,
+				TEXT("Transition Action \"%s\" Failed. Reason: %s"), *ActionName.ToString(), *FailureReason);
+		}
+		else
+		{
+			LogAndSendMessage(RequestInfo.Sender, RequestInfo.RequestId, ELogVerbosity::Error,
+				TEXT("Transition Action \"%s\" Failed."), *ActionName.ToString());
+		}
+	};
+	
+	UAvaRundown* Rundown = PlaybackCommandContext.GetCurrentRundown();
+	
+	if (!Rundown)
+	{
+		FailureReason = TEXT("no rundown currently loaded for playback.");
+		HandleCommandResponse();
+		return;
+	}
+
+	// Validate that the channel is specified and exists.
+	UAvaBroadcast& Broadcast = UAvaBroadcast::Get();
+	if (InMessage.ChannelName.IsEmpty())
+	{
+		FailureReason = TEXT("No Channel Name Provided");
+		HandleCommandResponse();
+		return;
+	}
+
+	const FName ChannelName(InMessage.ChannelName);
+	const FAvaBroadcastOutputChannel& Channel = Broadcast.GetCurrentProfile().GetChannel(ChannelName);
+	if (!Channel.IsValidChannel())
+	{
+		FailureReason = FString::Printf(TEXT("\"%s\" is not a valid channel"), *ChannelName.ToString());
+		HandleCommandResponse();
+		return;
+	}
+
+	if (InMessage.Action == EAvaRundownTransitionActions::ForceStop)
+	{
+		const int32 NumTransitions = Rundown->StopPageTransitionsForChannel(ChannelName);
+		if (NumTransitions > 0)
+		{
+			bSuccess = true;
+		}
+		else
+		{
+			FailureReason = TEXT("No Transitions were stopped");
+		}
+	}
+	else
+	{
+		FailureReason = TEXT("Invalid action");
+	}
+
+	HandleCommandResponse();
+}
+
+void FAvaRundownServer::HandleTransitionLayerAction(const FAvaRundownTransitionLayerAction& InMessage, const TSharedRef<IMessageContext, ESPMode::ThreadSafe>& InContext)
+{
+	const FRequestInfo RequestInfo = {InMessage.RequestId, InContext->GetSender()};
+
+	const FName ActionName = StaticEnum<EAvaRundownTransitionLayerActions>()->GetNameByValue((int64)InMessage.Action);
+
+	bool bSuccess = false;
+	FString FailureReason;
+
+	// This is the response handler for this command.
+	auto HandleCommandResponse = [this, &RequestInfo, &bSuccess, &FailureReason, ActionName]()
+	{
+		if (bSuccess)
+		{
+			SendMessage(RequestInfo.Sender, RequestInfo.RequestId, ELogVerbosity::Log,
+				TEXT("Transition Layer Action \"%s\": Ok."), *ActionName.ToString());
+		}
+		else if (!FailureReason.IsEmpty())
+		{
+			LogAndSendMessage(RequestInfo.Sender, RequestInfo.RequestId, ELogVerbosity::Error,
+				TEXT("Transition Layer Action \"%s\" Failed. Reason: %s"), *ActionName.ToString(), *FailureReason);
+		}
+		else
+		{
+			LogAndSendMessage(RequestInfo.Sender, RequestInfo.RequestId, ELogVerbosity::Error,
+				TEXT("Transition Layer Action \"%s\" Failed."), *ActionName.ToString());
+		}
+	};
+
+	UAvaRundown* Rundown = PlaybackCommandContext.GetCurrentRundown();
+	
+	if (!Rundown)
+	{
+		FailureReason = TEXT("no rundown currently loaded for playback.");
+		HandleCommandResponse();
+		return;
+	}
+	
+	// Validate that the channel is specified and exists.
+	UAvaBroadcast& Broadcast = UAvaBroadcast::Get();
+	if (InMessage.ChannelName.IsEmpty())
+	{
+		FailureReason = TEXT("No Channel Name Provided");
+		HandleCommandResponse();
+		return;
+	}
+
+	const FName ChannelName(InMessage.ChannelName);
+	const FAvaBroadcastOutputChannel& Channel = Broadcast.GetCurrentProfile().GetChannel(ChannelName);
+	if (!Channel.IsValidChannel())
+	{
+		FailureReason = FString::Printf(TEXT("\"%s\" is not a valid channel"), *ChannelName.ToString());
+		HandleCommandResponse();
+		return;
+	}
+
+	if (InMessage.Action == EAvaRundownTransitionLayerActions::Stop
+		|| InMessage.Action == EAvaRundownTransitionLayerActions::ForceStop)
+	{
+		// We need to gather playing layer handles that correspond to the layer names from the command.
+		TArray<FAvaTagHandle> Layers;
+		Layers.Reserve(InMessage.LayerNames.Num());
+	
+		for (UAvaRundownPagePlayer* PagePlayer : Rundown->GetPagePlayers())
+		{
+			if (PagePlayer->ChannelFName != ChannelName)
+			{
+				continue;
+			}
+
+			PagePlayer->ForEachInstancePlayer([&Layers, &LayerNames = InMessage.LayerNames](UAvaRundownPlaybackInstancePlayer* InInstancePlayer)
+			{
+				const FString LayerName = InInstancePlayer->TransitionLayer.ToString();
+				if (LayerNames.Contains(LayerName))
+				{
+					const bool bAlreadyAdded = Layers.ContainsByPredicate([&OtherTagHandle = InInstancePlayer->TransitionLayer](const FAvaTagHandle& InTagHandle)
+					{
+						return InTagHandle.MatchesExact(OtherTagHandle);
+					});
+					if (!bAlreadyAdded)
+					{
+						Layers.Add(InInstancePlayer->TransitionLayer);
+					}
+				}
+			});
+		}
+
+		if (Layers.IsEmpty())
+		{
+			FailureReason = TEXT("No playing layers corresponding to given layer names were found.");
+		}
+		else
+		{
+			EAvaRundownPageStopOptions StopOptions = InMessage.Action == EAvaRundownTransitionLayerActions::ForceStop ?
+				EAvaRundownPageStopOptions::ForceNoTransition : EAvaRundownPageStopOptions::Default;
+		
+			const TArray<int32> StoppedPages = Rundown->StopLayers(ChannelName, Layers, StopOptions);
+			if (StoppedPages.Num() > 0)
+			{
+				bSuccess = true;
+			}
+			else
+			{
+				FailureReason = TEXT("No Pages were stopped");
+			}
+		}
+	}
+	else
+	{
+		FailureReason = TEXT("Invalid action");
+	}
+
+	HandleCommandResponse();
+}
+
 void FAvaRundownServer::HandleGetProfiles(const FAvaRundownGetProfiles& InMessage, const TSharedRef<IMessageContext, ESPMode::ThreadSafe>& InContext)
 {
 	const UAvaBroadcast& Broadcast = UAvaBroadcast::Get();
diff --git a/Engine/Plugins/VirtualProduction/Avalanche/Source/AvalancheMedia/Private/Rundown/AvaRundownServer.h b/Engine/Plugins/VirtualProduction/Avalanche/Source/AvalancheMedia/Private/Rundown/AvaRundownServer.h
--- a/Engine/Plugins/VirtualProduction/Avalanche/Source/AvalancheMedia/Private/Rundown/AvaRundownServer.h
+++ b/Engine/Plugins/VirtualProduction/Avalanche/Source/AvalancheMedia/Private/Rundown/AvaRundownServer.h
@@ -90,6 +90,8 @@ public:
 	void HandlePagePreviewAction(const FAvaRundownPagePreviewAction& InMessage, const TSharedRef<IMessageContext, ESPMode::ThreadSafe>& InContext);
 	void HandlePageActions(const FAvaRundownPageActions& InMessage, const TSharedRef<IMessageContext, ESPMode::ThreadSafe>& InContext);
 	void HandlePagePreviewActions(const FAvaRundownPagePreviewActions& InMessage, const TSharedRef<IMessageContext, ESPMode::ThreadSafe>& InContext);
+	void HandleTransitionAction(const FAvaRundownTransitionAction& InMessage, const TSharedRef<IMessageContext, ESPMode::ThreadSafe>& InContext);
+	void HandleTransitionLayerAction(const FAvaRundownTransitionLayerAction& InMessage, const TSharedRef<IMessageContext, ESPMode::ThreadSafe>& InContext);
 
 	// Broadcast message handlers
 	void HandleGetProfiles(const FAvaRundownGetProfiles& InMessage, const TSharedRef<IMessageContext, ESPMode::ThreadSafe>& InContext);
