Remote Control: fixed crashes with behaviors being assumed valid when these behaviors could've originated from plugins that become disabled (or when moving assets to projects with these plugins disabled)

diff --git a/Engine/Plugins/VirtualProduction/Avalanche/Source/AvalancheMedia/Private/Playable/AvaPlayableRemoteControl.cpp b/Engine/Plugins/VirtualProduction/Avalanche/Source/AvalancheMedia/Private/Playable/AvaPlayableRemoteControl.cpp
--- a/Engine/Plugins/VirtualProduction/Avalanche/Source/AvalancheMedia/Private/Playable/AvaPlayableRemoteControl.cpp
+++ b/Engine/Plugins/VirtualProduction/Avalanche/Source/AvalancheMedia/Private/Playable/AvaPlayableRemoteControl.cpp
@@ -336,6 +336,11 @@ bool UE::AvaPlayableRemoteControl::GetEntitiesControlledByController(const URemo
 
 	for (const URCBehaviour* Behavior : Controller->Behaviours)
 	{
+		if (!Behavior)
+		{
+			continue;
+		}
+
 		// Special case: SetAssetByPath behaviors have an extra 'TargetEntity' that these behaviors control outside of RC Actions.
 		if (const URCSetAssetByPathBehaviour* SetAssetByPathBehavior = Cast<URCSetAssetByPathBehaviour>(Behavior))
 		{
@@ -384,8 +389,11 @@ namespace UE::AvaPlayableRemoteControl
 			PreviousBehavioursEnabled.Reserve(Controller->Behaviours.Num());			
 			for (URCBehaviour* const Behavior : Controller->Behaviours)
 			{
-				PreviousBehavioursEnabled.Add(Behavior->bIsEnabled);
-				Behavior->bIsEnabled = bInBehavioursEnabled;
+				if (Behavior)
+				{
+					PreviousBehavioursEnabled.Add(Behavior->bIsEnabled);
+					Behavior->bIsEnabled = bInBehavioursEnabled;
+				}
 			}
 		}
 	}
@@ -397,7 +405,10 @@ namespace UE::AvaPlayableRemoteControl
 			int32 BehaviourIndex = 0;
 			for (URCBehaviour* const Behavior : Controller->Behaviours)
 			{
-				Behavior->bIsEnabled = PreviousBehavioursEnabled[BehaviourIndex++];
+				if (Behavior)
+				{
+					Behavior->bIsEnabled = PreviousBehavioursEnabled[BehaviourIndex++];
+				}
 			}
 		}
 	}
diff --git a/Engine/Plugins/VirtualProduction/RemoteControl/Source/RemoteControlLogic/Private/Controller/RCController.cpp b/Engine/Plugins/VirtualProduction/RemoteControl/Source/RemoteControlLogic/Private/Controller/RCController.cpp
--- a/Engine/Plugins/VirtualProduction/RemoteControl/Source/RemoteControlLogic/Private/Controller/RCController.cpp
+++ b/Engine/Plugins/VirtualProduction/RemoteControl/Source/RemoteControlLogic/Private/Controller/RCController.cpp
@@ -86,10 +86,11 @@ int32 URCController::RemoveBehaviour(URCBehaviour* InBehaviour)
 int32 URCController::RemoveBehaviour(const FGuid InBehaviourId)
 {
 	int32 RemovedCount = 0;
-	
-	for (auto BehaviourIt = Behaviours.CreateIterator(); BehaviourIt; ++BehaviourIt)
+
+	for (TSet<TObjectPtr<URCBehaviour>>::TIterator BehaviourIt(Behaviours); BehaviourIt; ++BehaviourIt)
 	{
-		if (const URCBehaviour* Behaviour = *BehaviourIt; Behaviour->Id == InBehaviourId)
+		const URCBehaviour* Behaviour = *BehaviourIt;
+		if (Behaviour && Behaviour->Id == InBehaviourId)
 		{
 			BehaviourIt.RemoveCurrent();
 			RemovedCount++;
@@ -108,7 +109,7 @@ void URCController::ExecuteBehaviours(const bool bIsPreChange/* = false*/)
 {
 	for (URCBehaviour* Behaviour : Behaviours)
 	{
-		if (!Behaviour->bIsEnabled)
+		if (!Behaviour || !Behaviour->bIsEnabled)
 		{
 			continue;
 		}
diff --git a/Engine/Plugins/VirtualProduction/RemoteControl/Source/RemoteControlUI/Private/UI/Behaviour/SRCBehaviourPanelList.cpp b/Engine/Plugins/VirtualProduction/RemoteControl/Source/RemoteControlUI/Private/UI/Behaviour/SRCBehaviourPanelList.cpp
--- a/Engine/Plugins/VirtualProduction/RemoteControl/Source/RemoteControlUI/Private/UI/Behaviour/SRCBehaviourPanelList.cpp
+++ b/Engine/Plugins/VirtualProduction/RemoteControl/Source/RemoteControlUI/Private/UI/Behaviour/SRCBehaviourPanelList.cpp
@@ -256,7 +256,7 @@ void SRCBehaviourPanelList::SetIsBehaviourEnabled(const TSharedPtr<FRCBehaviourM
 
 void SRCBehaviourPanelList::AddBehaviourToList(URCBehaviour* InBehaviour)
 {
-	if (!BehaviourPanelWeakPtr.IsValid())
+	if (!InBehaviour || !BehaviourPanelWeakPtr.IsValid())
 	{
 		return;
 	}
