Motion Design: moved the Apply Mask material function outside of Motion Design plugin into GeometryMask plugin. This is so that plugins that depend on this material function don't require Motion Design as a dependency (which has a lot of other plugin dependencies like Remote Control), and instead rely on a GeometryMask that has no extra plugin dependencies.
Updated these plugins (Text3D, SVG Importer and PSD Importer) that depend on such material function in their own material assets to  depend on Geometry Mask so that the Material Function is valid and usable.

diff --git a/Engine/Plugins/Experimental/PSDImporter/PSDImporter.uplugin b/Engine/Plugins/Experimental/PSDImporter/PSDImporter.uplugin
--- a/Engine/Plugins/Experimental/PSDImporter/PSDImporter.uplugin
+++ b/Engine/Plugins/Experimental/PSDImporter/PSDImporter.uplugin
@@ -35,5 +35,11 @@
 			"LoadingPhase": "Default",
 			"PlatformAllowList": [ "Win64" ]
 		}
+	],
+	"Plugins": [
+		{
+			"Name": "GeometryMask",
+			"Enabled": true
+		}
 	]
 }
diff --git a/Engine/Plugins/Experimental/SVGImporter/SVGImporter.uplugin b/Engine/Plugins/Experimental/SVGImporter/SVGImporter.uplugin
--- a/Engine/Plugins/Experimental/SVGImporter/SVGImporter.uplugin
+++ b/Engine/Plugins/Experimental/SVGImporter/SVGImporter.uplugin
@@ -30,6 +30,10 @@
 		{
 			"Name": "GeometryScripting",
 			"Enabled":  true
+		},
+		{
+			"Name": "GeometryMask",
+			"Enabled": true
 		}
 	]
 }
diff --git a/Engine/Plugins/Experimental/Text3D/Text3D.uplugin b/Engine/Plugins/Experimental/Text3D/Text3D.uplugin
--- a/Engine/Plugins/Experimental/Text3D/Text3D.uplugin
+++ b/Engine/Plugins/Experimental/Text3D/Text3D.uplugin
@@ -30,6 +30,10 @@
       "Name": "GeometryProcessing",
       "Enabled": true,
 	  "TargetDenyList": [ "Server" ]
+    },
+    {
+      "Name": "GeometryMask",
+      "Enabled": true
     }
   ]
 }
diff --git a/Engine/Plugins/VirtualProduction/Avalanche/Source/AvalancheMask/Private/AvaMaskSettings.cpp b/Engine/Plugins/VirtualProduction/Avalanche/Source/AvalancheMask/Private/AvaMaskSettings.cpp
--- a/Engine/Plugins/VirtualProduction/Avalanche/Source/AvalancheMask/Private/AvaMaskSettings.cpp
+++ b/Engine/Plugins/VirtualProduction/Avalanche/Source/AvalancheMask/Private/AvaMaskSettings.cpp
@@ -33,10 +33,7 @@ UMaterialFunctionInterface* UAvaMaskSettings::GetMaterialFunction()
 
 UMaterialFunctionInterface* UAvaMaskSettings::GetDefaultMaterialFunction()
 {
-	// /Script/Engine.MaterialFunction'/GeometryMask/GeometryMask/MF_ApplyGeometryMask.MF_ApplyGeometryMask'
-	// const FString AssetPath = FString(TEXT("/GeometryMask")) / TEXT("GeometryMask") / TEXT("MF_ApplyGeometryMask");
-	// const FString AssetPath = TEXT("/Script/Engine.MaterialFunction'/GeometryMask/GeometryMask/MF_ApplyGeometryMask.MF_ApplyGeometryMask'");
-	const FString AssetPath = TEXT("/Script/Engine.MaterialFunction'/Avalanche/MaskResources/MF_ApplyMask2D_Single.MF_ApplyMask2D_Single'");
+	const FString AssetPath = TEXT("/Script/Engine.MaterialFunction'/GeometryMask/GeometryMask/MF_ApplyMask2D_Single.MF_ApplyMask2D_Single'");
 	const FAssetRegistryModule& AssetRegistryModule = FModuleManager::LoadModuleChecked<FAssetRegistryModule>("AssetRegistry");
 	if (const FAssetData AssetData = AssetRegistryModule.Get().GetAssetByObjectPath(FSoftObjectPath(AssetPath));
 		!AssetData.IsValid())
