Fix a build warning caused by recent move of plugin.

#rnx

diff --git a/Engine/Plugins/VirtualProduction/Avalanche/Avalanche.uplugin b/Engine/Plugins/VirtualProduction/Avalanche/Avalanche.uplugin
--- a/Engine/Plugins/VirtualProduction/Avalanche/Avalanche.uplugin
+++ b/Engine/Plugins/VirtualProduction/Avalanche/Avalanche.uplugin
@@ -311,6 +311,10 @@
 			"Name": "RemoteControlComponents",
 			"Enabled": true
 		},
+		{
+			"Name": "SequenceNavigator",
+			"Enabled": true
+		},
 		{
 			"Name": "SerializationUtils",
 			"Enabled": true
