Motion Design: Fixed issues with existing materials using Apply Mask material functions originally in Motion Design becoming invalid after the change moving these material functions to Geometry Mask as there was no redirector for backwards compatibility

diff --git a/Engine/Plugins/VirtualProduction/Avalanche/Config/DefaultAvalanche.ini b/Engine/Plugins/VirtualProduction/Avalanche/Config/DefaultAvalanche.ini
--- a/Engine/Plugins/VirtualProduction/Avalanche/Config/DefaultAvalanche.ini
+++ b/Engine/Plugins/VirtualProduction/Avalanche/Config/DefaultAvalanche.ini
@@ -274,6 +274,10 @@
 +PackageRedirects=(OldName="/Avalanche/EditorResources/M_AvalancheChannelVisualizer_Blue.M_AvalancheChannelVisualizer_Blue",NewName="/Avalanche/EditorResources/M_ChannelVisualizer_Blue.M_ChannelVisualizer_Blue")
 +PackageRedirects=(OldName="/Avalanche/EditorResources/M_AvalancheChannelVisualizer_Alpha.M_AvalancheChannelVisualizer_Alpha",NewName="/Avalanche/EditorResources/M_ChannelVisualizer_Alpha.M_ChannelVisualizer_Alpha")
 
+; Redirects for Mask Material functions moved to Geometry Mask plugin
++PackageRedirects=(OldName="/Avalanche/MaskResources/MF_ApplyMask2D", NewName="/GeometryMask/GeometryMask/MF_ApplyMask2D")
++PackageRedirects=(OldName="/Avalanche/MaskResources/MF_ApplyMask2D_Single", NewName="/GeometryMask/GeometryMask/MF_ApplyMask2D_Single")
+
 ; Redirects for rundown server requests for legacy clients.
 +StructRedirects=(OldName="/Script/AvalancheMediaEditor.AvaRundownPing",NewName="/Script/AvalancheMedia.AvaRundownPing")
 +StructRedirects=(OldName="/Script/AvalancheMediaEditor.AvaRundownGetRundowns",NewName="/Script/AvalancheMedia.AvaRundownGetRundowns")
