Motion Design: Updated tooltip for grid button

diff --git a/Engine/Plugins/VirtualProduction/Avalanche/Source/AvalancheLevelViewport/Private/AvaLevelViewportCommands.cpp b/Engine/Plugins/VirtualProduction/Avalanche/Source/AvalancheLevelViewport/Private/AvaLevelViewportCommands.cpp
--- a/Engine/Plugins/VirtualProduction/Avalanche/Source/AvalancheLevelViewport/Private/AvaLevelViewportCommands.cpp
+++ b/Engine/Plugins/VirtualProduction/Avalanche/Source/AvalancheLevelViewport/Private/AvaLevelViewportCommands.cpp
@@ -202,7 +202,7 @@ void FAvaLevelViewportCommands::RegisterGridCommands()
 {
 	UI_COMMAND(ToggleGrid
 		, "Toggle Grid"
-		, "Toggles the grid on and off."
+		, "Toggles the grid on and off.\n\nNote: Grid will only show if Always On is enabled or a snap operation is ongoing."
 		, EUserInterfaceActionType::Check
 		, FInputChord(EKeys::Apostrophe, EModifierKey::Control));
 
