[Motion Design] Rundown: Fix editor trying to save rundown created transient packages (playback instances or rundown managed instances).

diff --git a/Engine/Plugins/VirtualProduction/Avalanche/Source/AvalancheMedia/Private/Playable/Playables/AvaPlayableLevelStreaming.cpp b/Engine/Plugins/VirtualProduction/Avalanche/Source/AvalancheMedia/Private/Playable/Playables/AvaPlayableLevelStreaming.cpp
--- a/Engine/Plugins/VirtualProduction/Avalanche/Source/AvalancheMedia/Private/Playable/Playables/AvaPlayableLevelStreaming.cpp
+++ b/Engine/Plugins/VirtualProduction/Avalanche/Source/AvalancheMedia/Private/Playable/Playables/AvaPlayableLevelStreaming.cpp
@@ -551,6 +551,12 @@ void UAvaPlayableLevelStreaming::OnLevelStreamingStateChanged_Synchronized(ELeve
 			// Change the sub-level's world to be a "persistent" world type.
 			OuterWorld->WorldType = EWorldType::GamePreview;
 
+			if (UPackage* Package = OuterWorld->GetPackage())
+			{
+				// Mark package as transient. Prevents SaveDirtyPackages from trying to save this package.
+				Package->SetFlags(RF_Transient);
+			}
+
 			if (bLoadSubPlayables)
 			{
 				LoadSubPlayables(OuterWorld);
diff --git a/Engine/Plugins/VirtualProduction/Avalanche/Source/AvalancheMedia/Private/Rundown/AvaRundownManagedInstanceUtils.cpp b/Engine/Plugins/VirtualProduction/Avalanche/Source/AvalancheMedia/Private/Rundown/AvaRundownManagedInstanceUtils.cpp
--- a/Engine/Plugins/VirtualProduction/Avalanche/Source/AvalancheMedia/Private/Rundown/AvaRundownManagedInstanceUtils.cpp
+++ b/Engine/Plugins/VirtualProduction/Avalanche/Source/AvalancheMedia/Private/Rundown/AvaRundownManagedInstanceUtils.cpp
@@ -60,6 +60,9 @@ void FAvaRundownManagedInstanceUtils::PreventWorldFromBeingSeenAsLeakingByLevelE
 
 	if (UPackage* Package = InWorld->GetPackage())
 	{
+		// Mark package as transient. Prevents SaveDirtyPackages from trying to save this package.
+		Package->SetFlags(RF_Transient);
+
 #if WITH_EDITOR
 		// Prevents "unloading" of map package by marking it as new.
 		Package->MarkAsNewlyCreated();
