Material Designer / Avalanche Text: Style fix
- New SetDelegate that can be used to manually handle setting object property.
- Text now uses this delegate to enable custom styling.

diff --git a/Engine/Plugins/VirtualProduction/Avalanche/Source/AvalancheTextEditor/Private/AvaTextEditorModule.cpp b/Engine/Plugins/VirtualProduction/Avalanche/Source/AvalancheTextEditor/Private/AvaTextEditorModule.cpp
--- a/Engine/Plugins/VirtualProduction/Avalanche/Source/AvalancheTextEditor/Private/AvaTextEditorModule.cpp
+++ b/Engine/Plugins/VirtualProduction/Avalanche/Source/AvalancheTextEditor/Private/AvaTextEditorModule.cpp
@@ -12,14 +12,93 @@
 #include "DMObjectMaterialProperty.h"
 #include "Extensions/Text3DDefaultMaterialExtension.h"
 #include "Framework/Application/SlateApplication.h"
+#include "GameFramework/Actor.h"
 #include "IAvalancheComponentVisualizersModule.h"
 #include "IDynamicMaterialEditorModule.h"
+#include "Material/DynamicMaterialInstance.h"
 #include "Misc/CoreDelegates.h"
 #include "Modules/ModuleManager.h"
 #include "PropertyEditorModule.h"
 #include "Text3DActor.h"
 #include "Visualizer/AvaTextVisualizer.h"
 
+namespace UE::AvalancheTextEditor::Private
+{
+	TConstArrayView<FProperty*> GetText3DProperties()
+	{
+		static const UClass* MaterialExtensionClass = UText3DDefaultMaterialExtension::StaticClass();
+		static TArray<FProperty*> Text3DProperties =
+		{
+			MaterialExtensionClass->FindPropertyByName(TEXT("FrontMaterial")),
+			MaterialExtensionClass->FindPropertyByName(TEXT("BevelMaterial")),
+			MaterialExtensionClass->FindPropertyByName(TEXT("ExtrudeMaterial")),
+			MaterialExtensionClass->FindPropertyByName(TEXT("BackMaterial"))
+		};
+
+		return Text3DProperties;
+	}
+
+	/**
+	 * Make sure that the Style is set to Custom when a material is set.
+	 */
+	bool SetTextMaterialProperty(const FDMObjectMaterialProperty& InProperty, UDynamicMaterialInstance* InMaterial)
+	{
+		if (UText3DDefaultMaterialExtension* MaterialExtension = Cast<UText3DDefaultMaterialExtension>(InProperty.GetOuter()))
+		{
+			TConstArrayView<FProperty*> Text3DProperties = GetText3DProperties();
+			FProperty* PropertyToSet = InProperty.GetProperty();
+
+			if (Text3DProperties.Contains(PropertyToSet))
+			{
+				// Only set the style if we're actaully seting a material.
+				if (InMaterial)
+				{
+					MaterialExtension->SetStyle(EText3DMaterialStyle::Custom);
+				}
+			}
+		}
+
+		// Return to the original method and continue setting normally
+		return false;
+	}
+
+	/**
+	 * Customize Text actor to only show material extension materials instead of all character slots
+	 */
+	TArray<FDMObjectMaterialProperty> GetText3DMaterialProperties(UObject* InObject)
+	{
+		TArray<FDMObjectMaterialProperty> Properties;
+
+		if (const AActor* TextActor = Cast<AActor>(InObject))
+		{
+			if (const UText3DComponent* TextComponent = TextActor->FindComponentByClass<UText3DComponent>())
+			{
+				if (UText3DDefaultMaterialExtension* MaterialExtension = TextComponent->GetCastedMaterialExtension<UText3DDefaultMaterialExtension>())
+				{
+					TConstArrayView<FProperty*> Text3DProperties = GetText3DProperties();
+					Properties.Reserve(Text3DProperties.Num());
+
+					for (FProperty* Property : Text3DProperties)
+					{
+						if (Property)
+						{
+							FDMObjectMaterialProperty& ObjectMaterialProperty = Properties.Emplace_GetRef(MaterialExtension, Property);
+							ObjectMaterialProperty.SetMaterialSetterDelegate(FDMSetMaterialObjectProperty::CreateStatic(&SetTextMaterialProperty));
+						}
+
+						if (MaterialExtension->GetUseSingleMaterial())
+						{
+							break;
+						}
+					}
+				}
+			}
+		}
+
+		return Properties;
+	}
+}
+
 void FAvaTextEditorModule::StartupModule()
 {
 	FCoreDelegates::OnPostEngineInit.AddRaw(this, &FAvaTextEditorModule::RegisterComponentVisualizers);
@@ -56,54 +135,11 @@ void FAvaTextEditorModule::RegisterDynamicMaterialPropertyGenerator()
 		AAvaTextActor::StaticClass(),
 	};
 
-	/**
-	 * Customize Text actor to only show material extension materials instead of all character slots
-	 */
 	for (UClass* RegisterClass : RegisterClasses)
 	{
 		IDynamicMaterialEditorModule::Get().RegisterCustomMaterialPropertyGenerator(
 			RegisterClass,
-			FDMGetObjectMaterialPropertiesDelegate::CreateLambda(
-				[](UObject* InObject)
-				{
-					const UClass* MaterialExtensionClass = UText3DDefaultMaterialExtension::StaticClass();
-					static TArray<FProperty*> Text3DProperties =
-					{
-						MaterialExtensionClass->FindPropertyByName("FrontMaterial"),
-						MaterialExtensionClass->FindPropertyByName("BevelMaterial"),
-						MaterialExtensionClass->FindPropertyByName("ExtrudeMaterial"),
-						MaterialExtensionClass->FindPropertyByName("BackMaterial")
-					};
-
-					TArray<FDMObjectMaterialProperty> Properties;
-
-					if (const AActor* TextActor = Cast<AActor>(InObject))
-					{
-						if (const UText3DComponent* TextComponent = TextActor->FindComponentByClass<UText3DComponent>())
-						{
-							if (UText3DDefaultMaterialExtension* MaterialExtension = TextComponent->GetCastedMaterialExtension<UText3DDefaultMaterialExtension>())
-							{
-								Properties.Reserve(Text3DProperties.Num());
-
-								for (FProperty* Property : Text3DProperties)
-								{
-									if (Property)
-									{
-										Properties.Emplace(MaterialExtension, Property);
-									}
-
-									if (MaterialExtension->GetUseSingleMaterial())
-									{
-										break;
-									}
-								}
-							}
-						}
-					}
-
-					return Properties;
-				}
-			)
+			FDMGetObjectMaterialPropertiesDelegate::CreateStatic(UE::AvalancheTextEditor::Private::GetText3DMaterialProperties)
 		);
 	}
 }
diff --git a/Engine/Plugins/VirtualProduction/DynamicMaterial/Source/DynamicMaterialEditor/Private/DMObjectMaterialProperty.cpp b/Engine/Plugins/VirtualProduction/DynamicMaterial/Source/DynamicMaterialEditor/Private/DMObjectMaterialProperty.cpp
--- a/Engine/Plugins/VirtualProduction/DynamicMaterial/Source/DynamicMaterialEditor/Private/DMObjectMaterialProperty.cpp
+++ b/Engine/Plugins/VirtualProduction/DynamicMaterial/Source/DynamicMaterialEditor/Private/DMObjectMaterialProperty.cpp
@@ -115,8 +115,21 @@ UMaterialInterface* FDMObjectMaterialProperty::GetMaterialInterface() const
 	return nullptr;
 }
 
+DYNAMICMATERIALEDITOR_API void FDMObjectMaterialProperty::SetMaterialSetterDelegate(const FDMSetMaterialObjectProperty& InDelegate)
+{
+	MaterialSetterDelegate = InDelegate;
+}
+
 void FDMObjectMaterialProperty::SetMaterial(UDynamicMaterialInstance* InDynamicMaterial)
 {
+	if (MaterialSetterDelegate.IsBound())
+	{
+		if (MaterialSetterDelegate.Execute(*this, InDynamicMaterial))
+		{
+			return;
+		}
+	}
+
 	UObject* Outer = OuterWeak.Get();
 
 	if (!Outer)
diff --git a/Engine/Plugins/VirtualProduction/DynamicMaterial/Source/DynamicMaterialEditor/Public/DMObjectMaterialProperty.h b/Engine/Plugins/VirtualProduction/DynamicMaterial/Source/DynamicMaterialEditor/Public/DMObjectMaterialProperty.h
--- a/Engine/Plugins/VirtualProduction/DynamicMaterial/Source/DynamicMaterialEditor/Public/DMObjectMaterialProperty.h
+++ b/Engine/Plugins/VirtualProduction/DynamicMaterial/Source/DynamicMaterialEditor/Public/DMObjectMaterialProperty.h
@@ -16,6 +16,12 @@ class UMaterialInterface;
 class UObject;
 class UPrimitiveComponent;
 
+/**
+ * Overrides the material setter algorithm.
+ * @returns true means that the material was set. False if it was not. If false is returned, the normal material setting algorithm is used.
+ */
+DECLARE_DELEGATE_RetVal_TwoParams(bool /* Set Material */, FDMSetMaterialObjectProperty, const FDMObjectMaterialProperty& /* Object Property */, UDynamicMaterialInstance* /* Material to set */)
+
 /**
  * Defines a material property slot that can be a Material Designer Material.
  */
@@ -74,6 +80,9 @@ struct FDMObjectMaterialProperty
 		return nullptr;
 	}
 
+	/* The delegate can be used to override (or preprocess) the material setter function */
+	DYNAMICMATERIALEDITOR_API void SetMaterialSetterDelegate(const FDMSetMaterialObjectProperty& InDelegate);
+
 protected:
 	UPROPERTY()
 	TWeakObjectPtr<UObject> OuterWeak = nullptr;
@@ -84,5 +93,7 @@ protected:
 	/** Component or array property index. */
 	int32 Index = INDEX_NONE;
 
+	FDMSetMaterialObjectProperty MaterialSetterDelegate;
+
 	void SetMaterial(UDynamicMaterialInstance* InDynamicMaterial);
 };
