MotionDesign : Modifiers
- Fixed mask modifier not restoring original material due to generated material being transient and original material was not saved
- Fixed crash occuring when applying mask modifier on mac due to canvas texture copy

diff --git a/Engine/Plugins/Experimental/Text3D/Source/Text3D/Private/Extensions/Text3DDefaultLayoutExtension.cpp b/Engine/Plugins/Experimental/Text3D/Source/Text3D/Private/Extensions/Text3DDefaultLayoutExtension.cpp
--- a/Engine/Plugins/Experimental/Text3D/Source/Text3D/Private/Extensions/Text3DDefaultLayoutExtension.cpp
+++ b/Engine/Plugins/Experimental/Text3D/Source/Text3D/Private/Extensions/Text3DDefaultLayoutExtension.cpp
@@ -466,7 +466,7 @@ void UText3DDefaultLayoutExtension::OnLayoutOptionsChanged()
 	EText3DRendererFlags Flags = EText3DRendererFlags::Material;
 	EnumAddFlags(Flags, EText3DRendererFlags::Layout);
 
-	if (MaxWidthBehavior == EText3DMaxWidthHandling::WrapAndScale)
+	if (MaxWidthBehavior == EText3DMaxWidthHandling::WrapAndScale && ShapedText.IsValid())
 	{
 		for (const TConstEnumerateRef<FShapedGlyphLine> ShapedLine : EnumerateRange(ShapedText->Lines))
 		{
diff --git a/Engine/Plugins/VirtualProduction/Avalanche/Source/AvalancheMask/Private/Handling/AvaMaskAvaShapeMaterialCollectionHandle.cpp b/Engine/Plugins/VirtualProduction/Avalanche/Source/AvalancheMask/Private/Handling/AvaMaskAvaShapeMaterialCollectionHandle.cpp
--- a/Engine/Plugins/VirtualProduction/Avalanche/Source/AvalancheMask/Private/Handling/AvaMaskAvaShapeMaterialCollectionHandle.cpp
+++ b/Engine/Plugins/VirtualProduction/Avalanche/Source/AvalancheMask/Private/Handling/AvaMaskAvaShapeMaterialCollectionHandle.cpp
@@ -311,7 +311,7 @@ UAvaShapeDynamicMeshBase* FAvaMaskAvaShapeMaterialCollectionHandle::GetComponent
 void FAvaMaskAvaShapeMaterialCollectionHandle::OnMaterialAssignmentChanged(UPrimitiveComponent* InPrimitiveComponent, const int32 InPrevSlotNum, const int32 InNewSlotNum, const TArray<int32>& InChangedMaterialSlots)
 {
 	check(InPrimitiveComponent);
-	
+
 	RefreshMaterials();
 
 	if (UAvaShapeDynamicMeshBase* ShapeComponent = GetComponent())
@@ -323,7 +323,7 @@ void FAvaMaskAvaShapeMaterialCollectionHandle::OnMaterialAssignmentChanged(UPrim
 		{
 			if (!MaterialHandles.Contains(ChangedMaterialIdx))
 			{
-				UE_LOG(LogAvaMask, Error, TEXT("No MaterialHandle found slot: %i"), ChangedMaterialIdx);
+				UE_LOG(LogAvaMask, Verbose, TEXT("No MaterialHandle found slot: %i"), ChangedMaterialIdx);
 				ChangedHandles.Emplace(nullptr); // Still need something at this index
 			}
 			else
diff --git a/Engine/Plugins/VirtualProduction/Avalanche/Source/AvalancheMask/Private/Handling/AvaMaskText3DActorMaterialCollectionHandle.cpp b/Engine/Plugins/VirtualProduction/Avalanche/Source/AvalancheMask/Private/Handling/AvaMaskText3DActorMaterialCollectionHandle.cpp
--- a/Engine/Plugins/VirtualProduction/Avalanche/Source/AvalancheMask/Private/Handling/AvaMaskText3DActorMaterialCollectionHandle.cpp
+++ b/Engine/Plugins/VirtualProduction/Avalanche/Source/AvalancheMask/Private/Handling/AvaMaskText3DActorMaterialCollectionHandle.cpp
@@ -113,22 +113,42 @@ void FAvaMaskText3DActorMaterialCollectionHandle::SetMaterial(
 
 	if (UText3DComponent* TextComponent = GetComponent())
 	{
+		bool bRefresh = false;
+
 		switch (InSlotIdx)
 		{
 		case 0:
-			TextComponent->SetFrontMaterial(InMaterial);
-
+			{
+				TextComponent->SetFrontMaterial(InMaterial);
+				bRefresh = true;
+			}
+		break;
 		case 1:
-			TextComponent->SetBevelMaterial(InMaterial);
-
+			{
+				TextComponent->SetBevelMaterial(InMaterial);
+				bRefresh = true;
+			}
+		break;
 		case 2:
-			TextComponent->SetExtrudeMaterial(InMaterial);
-
+			{
+				TextComponent->SetExtrudeMaterial(InMaterial);
+				bRefresh = true;
+			}
+		break;
 		case 3:
-			TextComponent->SetBackMaterial(InMaterial);
-
+			{
+				TextComponent->SetBackMaterial(InMaterial);
+				bRefresh = true;
+			}
+		break;
 		default:
-			return;
+		break;
+		}
+
+		if (bRefresh)
+		{
+			constexpr bool bImmediate = true;
+			TextComponent->RequestUpdate(EText3DRendererFlags::Material, bImmediate);
 		}
 	}
 }
@@ -151,24 +171,36 @@ void FAvaMaskText3DActorMaterialCollectionHandle::SetMaterials(
 
 	if (UText3DComponent* TextComponent = GetComponent())
 	{
+		bool bRefresh = false;
+
 		if (InSetToggle[0])
 		{
-			TextComponent->SetFrontMaterial(InMaterials[0]);	
+			TextComponent->SetFrontMaterial(InMaterials[0]);
+			bRefresh = true;
 		}
 
 		if (InSetToggle[1])
 		{
-			TextComponent->SetBevelMaterial(InMaterials[1]);	
+			TextComponent->SetBevelMaterial(InMaterials[1]);
+			bRefresh = true;
 		}
 
 		if (InSetToggle[2])
 		{
-			TextComponent->SetExtrudeMaterial(InMaterials[2]);	
+			TextComponent->SetExtrudeMaterial(InMaterials[2]);
+			bRefresh = true;
 		}
 
 		if (InSetToggle[3])
 		{
-			TextComponent->SetBackMaterial(InMaterials[3]);	
+			TextComponent->SetBackMaterial(InMaterials[3]);
+			bRefresh = true;
+		}
+
+		if (bRefresh)
+		{
+			constexpr bool bImmediate = true;
+			TextComponent->RequestUpdate(EText3DRendererFlags::Material, bImmediate);
 		}
 	}
 }
diff --git a/Engine/Plugins/VirtualProduction/Avalanche/Source/AvalancheMask/Private/Mask2D/AvaMask2DBaseModifier.cpp b/Engine/Plugins/VirtualProduction/Avalanche/Source/AvalancheMask/Private/Mask2D/AvaMask2DBaseModifier.cpp
--- a/Engine/Plugins/VirtualProduction/Avalanche/Source/AvalancheMask/Private/Mask2D/AvaMask2DBaseModifier.cpp
+++ b/Engine/Plugins/VirtualProduction/Avalanche/Source/AvalancheMask/Private/Mask2D/AvaMask2DBaseModifier.cpp
@@ -337,7 +337,7 @@ void UAvaMask2DBaseModifier::RestorePreState()
 	Super::RestorePreState();
 	
 	TMap<TWeakObjectPtr<AActor>, FAvaMask2DActorData> ActorDataCopy = ActorData;
-	for (const TPair<TWeakObjectPtr<AActor>, FAvaMask2DActorData>& ActorDataPair : ActorDataCopy)
+	for (TPair<TWeakObjectPtr<AActor>, FAvaMask2DActorData>& ActorDataPair : ActorDataCopy)
 	{
 		if (AActor* Actor = ActorDataPair.Key.Get())
 		{
@@ -375,21 +375,37 @@ void UAvaMask2DBaseModifier::SaveActorPreState(AActor* InActor, FAvaMask2DActorD
 		UE::Ava::Internal::FindOrAddHandleByLambda(
 			MaterialCollectionHandles
 			, InActor
-			, [this, InActor]()
+			, [this, InActor, &InActorData]()
 			{
 				TSharedPtr<IAvaMaskMaterialCollectionHandle> Handle = GetObjectHandleSubsystem()->MakeHandle<IAvaMaskMaterialCollectionHandle>(InActor, UE::AvaMask::Internal::HandleTag);
+
+				for (const TPair<FAvaMask2DComponentMaterialPath, TObjectPtr<UMaterialInterface>>& OriginalMaterial : InActorData.OriginalMaterials)
+				{
+					Handle->SetMaterial(OriginalMaterial.Key.Component, OriginalMaterial.Key.SlotIdx, OriginalMaterial.Value);
+				}
+
+				InActorData.OriginalMaterials.Empty();
+
+				Handle->ForEachMaterial([&InActorData](const FSoftComponentReference& InComponent, const int32 InIdx, UMaterialInterface* InMaterial)
+				{
+					InActorData.OriginalMaterials.Add(FAvaMask2DComponentMaterialPath(InComponent, InIdx), InMaterial);
+					return true;
+				});
+
 				if (!Handle->OnSourceMaterialsChanged().IsBoundToObject(this))
 				{
 					Handle->OnSourceMaterialsChanged().BindUObject(this, &UAvaMask2DBaseModifier::OnMaterialsChanged);
 				}
+
 				return Handle;
-			});
+			}
+		);
 
 	FInstancedStruct& HandleData = MaterialCollectionHandleData.FindOrAdd(InActor, MaterialCollectionHandle->MakeDataStruct());
 	MaterialCollectionHandle->SaveOriginalState(HandleData);
 }
 
-void UAvaMask2DBaseModifier::RestoreActorPreState(AActor* InActor, const FAvaMask2DActorData& InActorData)
+void UAvaMask2DBaseModifier::RestoreActorPreState(AActor* InActor, FAvaMask2DActorData& InActorData)
 {
 	// Get MaterialCollectionHandle
     const TSharedPtr<IAvaMaskMaterialCollectionHandle>& MaterialCollectionHandle =
@@ -398,11 +414,14 @@ void UAvaMask2DBaseModifier::RestoreActorPreState(AActor* InActor, const FAvaMas
     		, InActor
     		, [this, InActor]()
     		{
-    			return GetObjectHandleSubsystem()->MakeHandle<IAvaMaskMaterialCollectionHandle>(InActor, UE::AvaMask::Internal::HandleTag);
+    			TSharedPtr<IAvaMaskMaterialCollectionHandle> Handle = GetObjectHandleSubsystem()->MakeHandle<IAvaMaskMaterialCollectionHandle>(InActor, UE::AvaMask::Internal::HandleTag);
+    			return Handle;
     		});
 
     FInstancedStruct& MaterialCollectionData = MaterialCollectionHandleData.FindOrAdd(InActor, MaterialCollectionHandle->MakeDataStruct());
     MaterialCollectionHandle->ApplyOriginalState(MaterialCollectionData);
+
+	InActorData.OriginalMaterials.Empty();
 }
 
 void UAvaMask2DBaseModifier::OnSceneTreeTrackedActorParentChanged(
diff --git a/Engine/Plugins/VirtualProduction/Avalanche/Source/AvalancheMask/Public/Mask2D/AvaMask2DBaseModifier.h b/Engine/Plugins/VirtualProduction/Avalanche/Source/AvalancheMask/Public/Mask2D/AvaMask2DBaseModifier.h
--- a/Engine/Plugins/VirtualProduction/Avalanche/Source/AvalancheMask/Public/Mask2D/AvaMask2DBaseModifier.h
+++ b/Engine/Plugins/VirtualProduction/Avalanche/Source/AvalancheMask/Public/Mask2D/AvaMask2DBaseModifier.h
@@ -50,10 +50,6 @@ struct FAvaMask2DActorData
 	/** Original assigned materials, used to create instances from, and restore to when modifier removed. */
 	UPROPERTY()
 	TMap<FAvaMask2DComponentMaterialPath, TObjectPtr<UMaterialInterface>> OriginalMaterials;
-
-	/** Instantiated materials corresponding to the above. */
-	UPROPERTY()
-	TMap<FAvaMask2DComponentMaterialPath, TObjectPtr<UMaterialInstanceDynamic>> GeneratedMaterialInstances;
 };
 
 /** Uses scene actors to create a mask texture and applies it to attached actors */
@@ -123,7 +119,7 @@ protected:
 	//~ End IAvaRenderStateUpdateExtension
 
 	virtual void SaveActorPreState(AActor* InActor, FAvaMask2DActorData& InActorData);
-	virtual void RestoreActorPreState(AActor* InActor, const FAvaMask2DActorData& InActorData);
+	virtual void RestoreActorPreState(AActor* InActor, FAvaMask2DActorData& InActorData);
 
 	void OnUseParentChannelChanged();
 	void OnChannelChanged();
diff --git a/Engine/Plugins/VirtualProduction/GeometryMask/Source/GeometryMask/Private/GeometryMaskCanvasResource.cpp b/Engine/Plugins/VirtualProduction/GeometryMask/Source/GeometryMask/Private/GeometryMaskCanvasResource.cpp
--- a/Engine/Plugins/VirtualProduction/GeometryMask/Source/GeometryMask/Private/GeometryMaskCanvasResource.cpp
+++ b/Engine/Plugins/VirtualProduction/GeometryMask/Source/GeometryMask/Private/GeometryMaskCanvasResource.cpp
@@ -20,8 +20,6 @@
 
 namespace UE::GeometryMask::Private
 {
-	static constexpr ETextureRenderTargetSampleCount RenderTargetSampleCount = ETextureRenderTargetSampleCount::RTSC_4;
-	
 	void OverscanProjectionMatrix(FMatrix& InOutMatrix, const FIntPoint& InSize, const int32 InPadding)
 	{
 		const FVector2f Multiplier(
@@ -353,10 +351,8 @@ void UGeometryMaskCanvasResource::UpdateRenderParameters(
 	bApplyBlur = BlurParameters.bPerChannelApplyBlur.CountSetBits(0) > 0; // if any channels have blur
 	bApplyDF = DFParameters.bPerChannelCalculateDF.CountSetBits(0) > 0; // if any channels have DF
 
-	if (UCanvasRenderTarget2D* Texture = GetRenderTargetTexture())
-	{
-		Texture->SetSampleCount(GetSampleCount());
-	}
+	// create render target
+	GetRenderTargetTexture();
 
 	// Effects may have changed viewport padding
     UpdateViewportSize();
@@ -442,7 +438,7 @@ void UGeometryMaskCanvasResource::Update(
 ETextureRenderTargetSampleCount UGeometryMaskCanvasResource::GetSampleCount() const
 {
 	// MSAA not supported for these effects (they need UAV access)
-	return bApplyBlur || bApplyDF ? ETextureRenderTargetSampleCount::RTSC_1 : UE::GeometryMask::Private::RenderTargetSampleCount;
+	return ETextureRenderTargetSampleCount::RTSC_1;
 }
 
 void UGeometryMaskCanvasResource::Draw(const ULevel* InLevel, FSceneView& InView, int32 InViewIndex)
