MotionDesign : Modifiers
- Fixed size to texture modifier changes size but UVs are not updated leading to stretched results and affecting final result

diff --git a/Engine/Plugins/VirtualProduction/Avalanche/Source/AvalancheModifiers/Private/Modifiers/AvaSizeToTextureModifier.cpp b/Engine/Plugins/VirtualProduction/Avalanche/Source/AvalancheModifiers/Private/Modifiers/AvaSizeToTextureModifier.cpp
--- a/Engine/Plugins/VirtualProduction/Avalanche/Source/AvalancheModifiers/Private/Modifiers/AvaSizeToTextureModifier.cpp
+++ b/Engine/Plugins/VirtualProduction/Avalanche/Source/AvalancheModifiers/Private/Modifiers/AvaSizeToTextureModifier.cpp
@@ -245,6 +245,9 @@ void UAvaSizeToTextureModifier::Apply()
 		if (!Shape2D->GetSize2D().Equals(NewUnscaledShapeSize, 0.01))
 		{
 			Shape2D->SetSize2D(NewUnscaledShapeSize);
+
+			// Refresh UV after changing size, these calls are async so size and UV update will be batched together
+			Shape2D->SetMaterialUVParams(UAvaShape2DDynMeshBase::MESH_INDEX_PRIMARY, Shape2D->GetMaterialUVParams(UAvaShape2DDynMeshBase::MESH_INDEX_PRIMARY));
 		}
 	}
 	else
diff --git a/Engine/Plugins/VirtualProduction/Avalanche/Source/AvalancheShapes/Private/DynamicMeshes/AvaShapeDynMeshBase.cpp b/Engine/Plugins/VirtualProduction/Avalanche/Source/AvalancheShapes/Private/DynamicMeshes/AvaShapeDynMeshBase.cpp
--- a/Engine/Plugins/VirtualProduction/Avalanche/Source/AvalancheShapes/Private/DynamicMeshes/AvaShapeDynMeshBase.cpp
+++ b/Engine/Plugins/VirtualProduction/Avalanche/Source/AvalancheShapes/Private/DynamicMeshes/AvaShapeDynMeshBase.cpp
@@ -1852,8 +1852,6 @@ void UAvaShapeDynamicMeshBase::SetShapeComponentMaterial(int32 InMaterialIndex,
 	}
 
 	ShapeComponent->SetMaterial(InMaterialIndex, InNewMaterial);
-
-	MarkMeshRenderStateDirty();
 }
 
 void UAvaShapeDynamicMeshBase::ScaleVertices(const FVector& InScale)
@@ -1874,10 +1872,7 @@ void UAvaShapeDynamicMeshBase::ScaleVertices(const FVector& InScale)
 			Vertex.Z *= InScale.Z;
 		}
 
-		if (MeshData->MaterialUVParams.Mode != EAvaShapeUVMode::Stretch)
-		{
-			MeshData->MaterialUVParams.bUVsDirty = true;
-		}
+		MeshData->MaterialUVParams.bUVsDirty = true;
 	}
 
 	for (FAvaSnapPoint& SnapPoint : LocalSnapPoints)
diff --git a/Engine/Plugins/VirtualProduction/Avalanche/Source/AvalancheShapes/Public/DynamicMeshes/AvaShapeDynMeshBase.h b/Engine/Plugins/VirtualProduction/Avalanche/Source/AvalancheShapes/Public/DynamicMeshes/AvaShapeDynMeshBase.h
--- a/Engine/Plugins/VirtualProduction/Avalanche/Source/AvalancheShapes/Public/DynamicMeshes/AvaShapeDynMeshBase.h
+++ b/Engine/Plugins/VirtualProduction/Avalanche/Source/AvalancheShapes/Public/DynamicMeshes/AvaShapeDynMeshBase.h
@@ -155,10 +155,10 @@ public:
 	bool GetOverridePrimaryUVParams(int32 InMeshIndex) const;
 
 	UFUNCTION(BlueprintCallable, Category="Shape")
-	void SetMaterialUVParams(int32 InMeshIndex, const FAvaShapeMaterialUVParameters& InParams);
+	AVALANCHESHAPES_API void SetMaterialUVParams(int32 InMeshIndex, const FAvaShapeMaterialUVParameters& InParams);
 
 	UFUNCTION(BlueprintPure, Category="Shape")
-	const FAvaShapeMaterialUVParameters& GetMaterialUVParams(int32 InMeshIndex) const;
+	AVALANCHESHAPES_API const FAvaShapeMaterialUVParameters& GetMaterialUVParams(int32 InMeshIndex) const;
 
 	FAvaShapeMaterialUVParameters* GetMaterialUVParamsPtr(int32 MeshIndex);
 
