MotionDesign : Modifiers
- Fixed size to texture modifier would break on load due to shape pre edit size value

diff --git a/Engine/Plugins/VirtualProduction/Avalanche/Source/AvalancheModifiers/Private/Modifiers/AvaSizeToTextureModifier.cpp b/Engine/Plugins/VirtualProduction/Avalanche/Source/AvalancheModifiers/Private/Modifiers/AvaSizeToTextureModifier.cpp
--- a/Engine/Plugins/VirtualProduction/Avalanche/Source/AvalancheModifiers/Private/Modifiers/AvaSizeToTextureModifier.cpp
+++ b/Engine/Plugins/VirtualProduction/Avalanche/Source/AvalancheModifiers/Private/Modifiers/AvaSizeToTextureModifier.cpp
@@ -80,6 +80,7 @@ void UAvaSizeToTextureModifier::OnModifierCDOSetup(FActorModifierCoreMetadata& I
 	InMetadata.SetName(TEXT("SizeToTexture"));
 	InMetadata.SetCategory(TEXT("Geometry"));
 #if WITH_EDITOR
+	InMetadata.SetDisplayName(LOCTEXT("ModifierDisplayName", "Size To Texture"));
 	InMetadata.SetDescription(LOCTEXT("ModifierDescription", "The modified actor will be resized to match a texture size based on the provided rule"));
 #endif
 
@@ -209,7 +210,7 @@ void UAvaSizeToTextureModifier::Apply()
 		
 		NewSize = FVector2D(FixedWidth, FixedWidth * HeightWidthRatio);
 	}
-	else if (Rule == EAvaSizeToTextureRule::AdaptiveHeight)
+	else if (Rule == EAvaSizeToTextureRule::AdaptiveWidth)
 	{
 		if (FMath::IsNearlyEqual(ShapeSize.X, ShapeSize.Y * WidthHeightRatio))
 		{
@@ -219,7 +220,7 @@ void UAvaSizeToTextureModifier::Apply()
 
 		NewSize = FVector2D(ShapeSize.Y * WidthHeightRatio, ShapeSize.Y);
 	}
-	else if (Rule == EAvaSizeToTextureRule::AdaptiveWidth)
+	else if (Rule == EAvaSizeToTextureRule::AdaptiveHeight)
 	{
 		if (FMath::IsNearlyEqual(ShapeSize.Y, ShapeSize.X * HeightWidthRatio))
 		{
diff --git a/Engine/Plugins/VirtualProduction/Avalanche/Source/AvalancheModifiers/Public/Modifiers/AvaSizeToTextureModifier.h b/Engine/Plugins/VirtualProduction/Avalanche/Source/AvalancheModifiers/Public/Modifiers/AvaSizeToTextureModifier.h
--- a/Engine/Plugins/VirtualProduction/Avalanche/Source/AvalancheModifiers/Public/Modifiers/AvaSizeToTextureModifier.h
+++ b/Engine/Plugins/VirtualProduction/Avalanche/Source/AvalancheModifiers/Public/Modifiers/AvaSizeToTextureModifier.h
@@ -14,14 +14,18 @@ class UTexture;
 UENUM()
 enum class EAvaSizeToTextureRule : uint8
 {
-	/** Changes to width are allowed, height will react accordingly */
+	/**
+	 * Automatically adapts the width based on the height, to maintain texture ratio
+	 */
 	AdaptiveWidth,
-	/** Changes to height are allowed, width will react accordingly */
+	/**
+	 * Automatically adapts the height based on the width, to maintain texture ratio
+	 */
 	AdaptiveHeight,
 	/** Lock height and adapt width */
-	FixedHeight,
+	FixedHeight    UMETA(Hidden),
 	/** Lock width and adapt height */
-	FixedWidth,
+	FixedWidth     UMETA(Hidden),
 };
 
 /**
@@ -102,7 +106,7 @@ protected:
 
 	/** Rule for texture resize */
 	UPROPERTY(EditInstanceOnly, Setter, Getter, Category="TextureSize", meta=(AllowPrivateAccess="true"))
-	EAvaSizeToTextureRule Rule = EAvaSizeToTextureRule::AdaptiveHeight;
+	EAvaSizeToTextureRule Rule = EAvaSizeToTextureRule::AdaptiveWidth;
 
 	/** The fixed height size */
 	UPROPERTY(EditInstanceOnly, Setter, Getter, Category="TextureSize", meta=(ClampMin="0", EditCondition="Rule == EAvaSizeToTextureRule::FixedHeight", EditConditionHides, AllowPrivateAccess="true"))
diff --git a/Engine/Plugins/VirtualProduction/Avalanche/Source/AvalancheShapes/Private/DynamicMeshes/AvaShape2DDynMeshBase.cpp b/Engine/Plugins/VirtualProduction/Avalanche/Source/AvalancheShapes/Private/DynamicMeshes/AvaShape2DDynMeshBase.cpp
--- a/Engine/Plugins/VirtualProduction/Avalanche/Source/AvalancheShapes/Private/DynamicMeshes/AvaShape2DDynMeshBase.cpp
+++ b/Engine/Plugins/VirtualProduction/Avalanche/Source/AvalancheShapes/Private/DynamicMeshes/AvaShape2DDynMeshBase.cpp
@@ -20,6 +20,8 @@ void UAvaShape2DDynMeshBase::PostLoad()
 
 	Size3D.Y = Size2D.X;
 	Size3D.Z = Size2D.Y;
+
+	PreEditSize2D = Size2D;
 }
 
 int32 UAvaShape2DDynMeshBase::AddVertexRaw(FAvaShapeMesh& InMesh, const FVector2D& Location, bool bForceNew)
@@ -171,6 +173,7 @@ void UAvaShape2DDynMeshBase::OnRegisteredMeshes()
 	Super::OnRegisteredMeshes();
 
 	PixelSize2D = Size2D;
+	PreEditSize2D = Size2D;
 }
 
 void UAvaShape2DDynMeshBase::OnPixelSizeChanged()
diff --git a/Engine/Plugins/VirtualProduction/Avalanche/Source/AvalancheShapes/Private/DynamicMeshes/AvaShape3DDynMeshBase.cpp b/Engine/Plugins/VirtualProduction/Avalanche/Source/AvalancheShapes/Private/DynamicMeshes/AvaShape3DDynMeshBase.cpp
--- a/Engine/Plugins/VirtualProduction/Avalanche/Source/AvalancheShapes/Private/DynamicMeshes/AvaShape3DDynMeshBase.cpp
+++ b/Engine/Plugins/VirtualProduction/Avalanche/Source/AvalancheShapes/Private/DynamicMeshes/AvaShape3DDynMeshBase.cpp
@@ -86,6 +86,13 @@ void UAvaShape3DDynMeshBase::GetBounds(FVector& Origin, FVector& BoxExtent, FVec
 	Pivot = FVector(0, 0, 0);
 }
 
+void UAvaShape3DDynMeshBase::PostLoad()
+{
+	Super::PostLoad();
+
+	PreEditSize3D = Size3D;
+}
+
 void UAvaShape3DDynMeshBase::SetPixelSize3D(const FVector& InPixelSize)
 {
 	if (!bAllowEditSize)
@@ -129,6 +136,7 @@ void UAvaShape3DDynMeshBase::OnRegisteredMeshes()
 	Super::OnRegisteredMeshes();
 
 	PixelSize3D = Size3D;
+	PreEditSize3D = Size3D;
 }
 
 void UAvaShape3DDynMeshBase::OnPixelSizeChanged()
diff --git a/Engine/Plugins/VirtualProduction/Avalanche/Source/AvalancheShapes/Public/DynamicMeshes/AvaShape3DDynMeshBase.h b/Engine/Plugins/VirtualProduction/Avalanche/Source/AvalancheShapes/Public/DynamicMeshes/AvaShape3DDynMeshBase.h
--- a/Engine/Plugins/VirtualProduction/Avalanche/Source/AvalancheShapes/Public/DynamicMeshes/AvaShape3DDynMeshBase.h
+++ b/Engine/Plugins/VirtualProduction/Avalanche/Source/AvalancheShapes/Public/DynamicMeshes/AvaShape3DDynMeshBase.h
@@ -41,6 +41,7 @@ public:
 
 protected:
 	//~ Begin UObject
+	virtual void PostLoad() override;
 #if WITH_EDITOR
 	virtual void PostEditChangeProperty(FPropertyChangedEvent& PropertyChangedEvent) override;
 #endif
