Motion Design: fixed crash with mask modifiers caused by holding onto element references after a map reallocation.

diff --git a/Engine/Plugins/VirtualProduction/Avalanche/Source/AvalancheMask/Private/Mask2D/AvaMask2DBaseModifier.cpp b/Engine/Plugins/VirtualProduction/Avalanche/Source/AvalancheMask/Private/Mask2D/AvaMask2DBaseModifier.cpp
--- a/Engine/Plugins/VirtualProduction/Avalanche/Source/AvalancheMask/Private/Mask2D/AvaMask2DBaseModifier.cpp
+++ b/Engine/Plugins/VirtualProduction/Avalanche/Source/AvalancheMask/Private/Mask2D/AvaMask2DBaseModifier.cpp
@@ -371,7 +371,7 @@ void UAvaMask2DBaseModifier::Apply()
 void UAvaMask2DBaseModifier::SaveActorPreState(AActor* InActor, FAvaMask2DActorData& InActorData)
 {
 	// Get MaterialCollectionHandle
-	const TSharedPtr<IAvaMaskMaterialCollectionHandle>& MaterialCollectionHandle =
+	const TSharedPtr<IAvaMaskMaterialCollectionHandle> MaterialCollectionHandle =
 		UE::Ava::Internal::FindOrAddHandleByLambda(
 			MaterialCollectionHandles
 			, InActor
@@ -408,7 +408,7 @@ void UAvaMask2DBaseModifier::SaveActorPreState(AActor* InActor, FAvaMask2DActorD
 void UAvaMask2DBaseModifier::RestoreActorPreState(AActor* InActor, FAvaMask2DActorData& InActorData)
 {
 	// Get MaterialCollectionHandle
-    const TSharedPtr<IAvaMaskMaterialCollectionHandle>& MaterialCollectionHandle =
+    const TSharedPtr<IAvaMaskMaterialCollectionHandle> MaterialCollectionHandle =
     	UE::Ava::Internal::FindOrAddHandleByLambda(
     		MaterialCollectionHandles
     		, InActor
diff --git a/Engine/Plugins/VirtualProduction/Avalanche/Source/AvalancheMask/Private/Mask2D/AvaMask2DReadModifier.cpp b/Engine/Plugins/VirtualProduction/Avalanche/Source/AvalancheMask/Private/Mask2D/AvaMask2DReadModifier.cpp
--- a/Engine/Plugins/VirtualProduction/Avalanche/Source/AvalancheMask/Private/Mask2D/AvaMask2DReadModifier.cpp
+++ b/Engine/Plugins/VirtualProduction/Avalanche/Source/AvalancheMask/Private/Mask2D/AvaMask2DReadModifier.cpp
@@ -91,7 +91,7 @@ bool UAvaMask2DReadModifier::ApplyRead(AActor* InActor, FAvaMask2DActorData& InA
 	if (const UGeometryMaskCanvas* Canvas = GetCurrentCanvas())
 	{
 		// Get MaterialCollectionHandle
-		const TSharedPtr<IAvaMaskMaterialCollectionHandle>& MaterialCollectionHandle =
+		const TSharedPtr<IAvaMaskMaterialCollectionHandle> MaterialCollectionHandle =
 			UE::Ava::Internal::FindOrAddHandleByLambda(
 				MaterialCollectionHandles
 				, InActor
