MotionDesign : Modifiers
- Fixed modifier restoring and reapplying itself when stack is removed instead of restoring only thus leaving actor in a modified state
- Fixed disabling modifier stack should restore all modifier but instead does nothing since modifier stack state is not checked during execution
- Removing MarkModifierDirty call in OnModifierAdded -> OnChannelChanged() since modifier will be executed later and does not need to be called twice

diff --git a/Engine/Plugins/VirtualProduction/ActorModifierCore/Source/ActorModifierCore/Private/Modifiers/ActorModifierCoreBase.cpp b/Engine/Plugins/VirtualProduction/ActorModifierCore/Source/ActorModifierCore/Private/Modifiers/ActorModifierCoreBase.cpp
--- a/Engine/Plugins/VirtualProduction/ActorModifierCore/Source/ActorModifierCore/Private/Modifiers/ActorModifierCoreBase.cpp
+++ b/Engine/Plugins/VirtualProduction/ActorModifierCore/Source/ActorModifierCore/Private/Modifiers/ActorModifierCoreBase.cpp
@@ -316,7 +316,7 @@ void UActorModifierCoreBase::MarkModifierDirty(bool bExecute)
 	const UActorModifierCoreStack* Stack = GetRootModifierStack();
 
 	// Do not mark dirty if stack is in execution
-	if (!Stack || !Stack->IsModifierIdle() || !IsModifierInitialized())
+	if (!Stack || !Stack->IsModifierIdle() || !Stack->IsModifierInitialized() || !IsModifierInitialized())
 	{
 		return;
 	}
@@ -752,11 +752,11 @@ void UActorModifierCoreBase::DisableModifier(EActorModifierCoreDisableReason InR
 
 void UActorModifierCoreBase::InitializeModifier(EActorModifierCoreEnableReason InReason)
 {
-	const bool bModifierValid = !bModifierInitialized;
-
 	// is the modifier correctly setup
-	if (bModifierValid)
+	if (!bModifierInitialized)
 	{
+		bModifierInitialized = true;
+
 #if WITH_EDITOR
 		// to be able to track property changes and stack updates
 		if (!HasAnyFlags(RF_Transactional))
@@ -773,6 +773,9 @@ void UActorModifierCoreBase::InitializeModifier(EActorModifierCoreEnableReason I
 		{
 			PostModifierCreation(GetModifierStack());
 
+			// Set it to true again since PostModifierCreation sets bModifierInitialized to false
+			bModifierInitialized = true;
+
 			// Cannot proceed with invalid metadata
 			if (!Metadata.IsValid())
 			{
@@ -797,9 +800,8 @@ void UActorModifierCoreBase::InitializeModifier(EActorModifierCoreEnableReason I
 			EnableModifier(InReason);
 		}
 
-		bModifierInitialized = true;
-
-		MarkModifierDirty(false);
+		const bool bExecuteModifiers = IsModifierStack();
+		MarkModifierDirty(bExecuteModifiers);
 
 		UActorModifierCoreStack::OnModifierAddedDelegate.Broadcast(this, InReason);
 	}
diff --git a/Engine/Plugins/VirtualProduction/ActorModifierCore/Source/ActorModifierCore/Private/Modifiers/ActorModifierCoreDefs.cpp b/Engine/Plugins/VirtualProduction/ActorModifierCore/Source/ActorModifierCore/Private/Modifiers/ActorModifierCoreDefs.cpp
--- a/Engine/Plugins/VirtualProduction/ActorModifierCore/Source/ActorModifierCore/Private/Modifiers/ActorModifierCoreDefs.cpp
+++ b/Engine/Plugins/VirtualProduction/ActorModifierCore/Source/ActorModifierCore/Private/Modifiers/ActorModifierCoreDefs.cpp
@@ -14,8 +14,10 @@ FActorModifierCoreMetadata::FActorModifierCoreMetadata()
 	: bIsStack(false)
 #if WITH_EDITOR
 	, bHidden(false)
+	, Color(FLinearColor::White)
 #endif
-{}
+{
+}
 
 FActorModifierCoreMetadata::FActorModifierCoreMetadata(const UActorModifierCoreBase* InModifier)
 {
@@ -458,6 +460,11 @@ void FActorModifierCoreExecutionTask::BuildExecutionChain(TConstArrayView<UActor
 
 	if (ModifierStack)
 	{
+		if (!ModifierStack->IsModifierEnabled())
+		{
+			return;
+		}
+
 		bAllModifiersDirty = ModifierStack->bAllModifiersDirty;
 	}
 
diff --git a/Engine/Plugins/VirtualProduction/ActorModifierCore/Source/ActorModifierCore/Private/Modifiers/ActorModifierCoreStack.cpp b/Engine/Plugins/VirtualProduction/ActorModifierCore/Source/ActorModifierCore/Private/Modifiers/ActorModifierCoreStack.cpp
--- a/Engine/Plugins/VirtualProduction/ActorModifierCore/Source/ActorModifierCore/Private/Modifiers/ActorModifierCoreStack.cpp
+++ b/Engine/Plugins/VirtualProduction/ActorModifierCore/Source/ActorModifierCore/Private/Modifiers/ActorModifierCoreStack.cpp
@@ -1622,7 +1622,10 @@ void UActorModifierCoreStack::Apply()
 	// restore original state before dirty modifiers
 	RestorePreState();
 
-	ExecutionTask.Apply(GetModifiers(), this);
+	if (bModifierInitialized)
+	{
+		ExecutionTask.Apply(GetModifiers(), this);
+	}
 }
 
 void UActorModifierCoreStack::OnModifierRemoved(EActorModifierCoreDisableReason InReason)
diff --git a/Engine/Plugins/VirtualProduction/Avalanche/Source/AvalancheMask/Private/Mask2D/AvaMask2DBaseModifier.cpp b/Engine/Plugins/VirtualProduction/Avalanche/Source/AvalancheMask/Private/Mask2D/AvaMask2DBaseModifier.cpp
--- a/Engine/Plugins/VirtualProduction/Avalanche/Source/AvalancheMask/Private/Mask2D/AvaMask2DBaseModifier.cpp
+++ b/Engine/Plugins/VirtualProduction/Avalanche/Source/AvalancheMask/Private/Mask2D/AvaMask2DBaseModifier.cpp
@@ -185,6 +185,32 @@ void UAvaMask2DBaseModifier::OnCanvasChanged()
 	OnFeatherChanged();
 }
 
+void UAvaMask2DBaseModifier::UpdateCanvas()
+{
+	// Reset/invalidate cached
+	{
+		for (TPair<TWeakObjectPtr<AActor>, FAvaMask2DActorData>& ActorDataPair : ActorData)
+		{
+			ActorDataPair.Value.CanvasTextureWeak.Reset();
+		}
+		
+		LastResolvedCanvasName = NAME_None;
+		CanvasWeak.Reset();
+	}
+
+#if WITH_EDITOR
+	if (!Channel.IsNone())
+	{
+		if (UAvaMaskSubsystem* Subsystem = GEngine->GetEngineSubsystem<UAvaMaskSubsystem>())
+		{
+			Subsystem->SetLastSpecifiedChannelName(Channel);
+		}
+	}
+#endif
+
+	TryResolveCanvas();
+}
+
 void UAvaMask2DBaseModifier::CanvasParamsToLocal()
 {
 	if (UGeometryMaskCanvas* Canvas = GetCurrentCanvas())
@@ -262,7 +288,7 @@ void UAvaMask2DBaseModifier::OnModifierAdded(EActorModifierCoreEnableReason InRe
 
 	TRACE_BOOKMARK(TEXT("UAvaMask2DModifier::OnModifierAdded"));
 
-	OnChannelChanged();
+	UpdateCanvas();
 
 	if (InReason == EActorModifierCoreEnableReason::User
 		|| InReason == EActorModifierCoreEnableReason::Duplicate)
@@ -643,29 +669,7 @@ void UAvaMask2DBaseModifier::RemoveFromActor(AActor* InActor)
 
 void UAvaMask2DBaseModifier::OnChannelChanged()
 {
-	// Reset/invalidate cached
-	{
-		for (TPair<TWeakObjectPtr<AActor>, FAvaMask2DActorData>& ActorDataPair : ActorData)
-		{
-			ActorDataPair.Value.CanvasTextureWeak.Reset();
-		}
-		
-		LastResolvedCanvasName = NAME_None;
-		CanvasWeak.Reset();
-	}
-
-#if WITH_EDITOR
-	if (!Channel.IsNone())
-	{
-		if (UAvaMaskSubsystem* Subsystem = GEngine->GetEngineSubsystem<UAvaMaskSubsystem>())
-		{
-			Subsystem->SetLastSpecifiedChannelName(Channel);
-		}
-	}
-#endif
-
-	TryResolveCanvas();
-	
+	UpdateCanvas();
 	MarkModifierDirty();
 }
 
diff --git a/Engine/Plugins/VirtualProduction/Avalanche/Source/AvalancheMask/Public/Mask2D/AvaMask2DBaseModifier.h b/Engine/Plugins/VirtualProduction/Avalanche/Source/AvalancheMask/Public/Mask2D/AvaMask2DBaseModifier.h
--- a/Engine/Plugins/VirtualProduction/Avalanche/Source/AvalancheMask/Public/Mask2D/AvaMask2DBaseModifier.h
+++ b/Engine/Plugins/VirtualProduction/Avalanche/Source/AvalancheMask/Public/Mask2D/AvaMask2DBaseModifier.h
@@ -128,6 +128,8 @@ protected:
 	void OnFeatherChanged();
 	void OnCanvasChanged();
 
+	void UpdateCanvas();
+
 	/** Stores specific properties on the canvas locally. */	
 	void CanvasParamsToLocal();
 
