[MotionDesign] Stagger Tool:
- Change auto update icon to be more clear
- Move auto update button to the far left

diff --git a/Engine/Plugins/VirtualProduction/Avalanche/Source/AvalancheSequencer/Private/StaggerTool/AvaStaggerTool.cpp b/Engine/Plugins/VirtualProduction/Avalanche/Source/AvalancheSequencer/Private/StaggerTool/AvaStaggerTool.cpp
--- a/Engine/Plugins/VirtualProduction/Avalanche/Source/AvalancheSequencer/Private/StaggerTool/AvaStaggerTool.cpp
+++ b/Engine/Plugins/VirtualProduction/Avalanche/Source/AvalancheSequencer/Private/StaggerTool/AvaStaggerTool.cpp
@@ -616,7 +616,7 @@ void FAvaStaggerTool::Stagger()
 	}
 }
 
-bool FAvaStaggerTool::IsAutoAppying() const
+bool FAvaStaggerTool::IsAutoApplying() const
 {
 	return Settings->bAutoApply;
 }
diff --git a/Engine/Plugins/VirtualProduction/Avalanche/Source/AvalancheSequencer/Private/StaggerTool/AvaStaggerTool.h b/Engine/Plugins/VirtualProduction/Avalanche/Source/AvalancheSequencer/Private/StaggerTool/AvaStaggerTool.h
--- a/Engine/Plugins/VirtualProduction/Avalanche/Source/AvalancheSequencer/Private/StaggerTool/AvaStaggerTool.h
+++ b/Engine/Plugins/VirtualProduction/Avalanche/Source/AvalancheSequencer/Private/StaggerTool/AvaStaggerTool.h
@@ -40,7 +40,7 @@ public:
 
 	void Stagger();
 
-	bool IsAutoAppying() const;
+	bool IsAutoApplying() const;
 
 	void OnSequencerSelectionChanged();
 
diff --git a/Engine/Plugins/VirtualProduction/Avalanche/Source/AvalancheSequencer/Private/StaggerTool/Widgets/SAvaStaggerTool.cpp b/Engine/Plugins/VirtualProduction/Avalanche/Source/AvalancheSequencer/Private/StaggerTool/Widgets/SAvaStaggerTool.cpp
--- a/Engine/Plugins/VirtualProduction/Avalanche/Source/AvalancheSequencer/Private/StaggerTool/Widgets/SAvaStaggerTool.cpp
+++ b/Engine/Plugins/VirtualProduction/Avalanche/Source/AvalancheSequencer/Private/StaggerTool/Widgets/SAvaStaggerTool.cpp
@@ -78,109 +78,105 @@ TSharedRef<SWidget> SAvaStaggerTool::ConstructApplyRow()
 	constexpr float ImageSize = 14.f;
 	constexpr float Spacing = 5.f;
 
-	return 
-		SNew(SHorizontalBox)
-		+ SHorizontalBox::Slot()
-		.VAlign(VAlign_Center)
-		.Padding(Spacing, 0.f, 0.f, 0.f)
+	return SNew(SBox)
+		.HeightOverride(ButtonSize)
 		[
-			SNew(STextBlock)
-			.TextStyle(FAppStyle::Get(), TEXT("HintText"))
-			.Text(this, &SAvaStaggerTool::GetSelectionText)
-			.ToolTipText(this, &SAvaStaggerTool::GetSelectionToolTipText)
-			.ColorAndOpacity(this, &SAvaStaggerTool::GetSelectionTextColor)
-		]
-		+ SHorizontalBox::Slot()
-		.AutoWidth()
-		.VAlign(VAlign_Bottom)
-		.Padding(Spacing, 0.f, 0.f, 0.f)
-		[
-			SNew(SButton)
-			.ContentPadding(2.f)
-			.ToolTipText(FSequencerCommands::Get().AlignSelectionToPlayhead->GetDescription())
-			.IsEnabled(this, &SAvaStaggerTool::OnCanAlignToPlayhead)
-			.OnClicked(this, &SAvaStaggerTool::OnAlignToPlayhead)
-			[
-				SNew(SImage)
-				.DesiredSizeOverride(FVector2D(ImageSize))
-				.ColorAndOpacity(FSlateColor::UseForeground())
-				.Image(FAppStyle::GetBrush(TEXT("HorizontalAlignment_Left")))
-			]
-		]
-		+ SHorizontalBox::Slot()
-		.AutoWidth()
-		.VAlign(VAlign_Bottom)
-		.Padding(Spacing, 0.f, 0.f, 0.f)
-		[
-			SNew(SButton)
-			.Text(LOCTEXT("ApplyButtonText", "Apply"))
-			.OnClicked(this, &SAvaStaggerTool::OnApplyButtonClick)
-			.IsEnabled_Lambda([this]()
-				{
-					if (const TSharedPtr<FAvaStaggerTool> Tool = WeakTool.Pin())
-					{
-						return Tool->HasValidSelection() && !Tool->IsAutoAppying();
-					}
-					return false;
-				})
-		]
-		+ SHorizontalBox::Slot()
-		.AutoWidth()
-		.VAlign(VAlign_Bottom)
-		.Padding(Spacing, 0.f, 0.f, 0.f)
-		[
-			SNew(SBox)
-			.WidthOverride(ButtonSize)
-			.HeightOverride(ButtonSize)
+			SNew(SHorizontalBox)
+			+ SHorizontalBox::Slot()
+			.AutoWidth()
+			.VAlign(VAlign_Fill)
+			.Padding(Spacing, 0.f, Spacing, 0.f)
 			[
 				SNew(SCheckBox)
 				.Style(FAppStyle::Get(), TEXT("Sequencer.Outliner.ToggleButton"))
 				.ToolTipText(LOCTEXT("AutoApplyToolTip", "Auto apply options on change"))
 				.IsChecked_Lambda([this]()
 					{
-						return (WeakTool.IsValid() && WeakTool.Pin()->IsAutoAppying())
+						return (WeakTool.IsValid() && WeakTool.Pin()->IsAutoApplying())
 							? ECheckBoxState::Checked : ECheckBoxState::Unchecked;
 					})
 				.OnCheckStateChanged(this, &SAvaStaggerTool::OnToggleAutoUpdateClick)
 				[
-					SNew(SBox)
-					.Padding(0.f, 4.f)
+					SNew(SBox) 
+					.HAlign(HAlign_Center)
+					.VAlign(VAlign_Center)
 					[
 						SNew(SImage)
+						.DesiredSizeOverride(FVector2D(ImageSize))
 						.ColorAndOpacity(FSlateColor::UseForeground())
-						.Image_Lambda([this]()
-							{
-								const FName BrushName = (WeakTool.IsValid() && WeakTool.Pin()->IsAutoAppying())
-									? TEXT("Icons.Lock") : TEXT("Icons.Unlock");
-								return FAppStyle::GetBrush(BrushName);
-							})
+						.Image(FAppStyle::GetBrush(TEXT("MaterialEditor.Apply")))
 					]
 				]
 			]
-		]
-		+ SHorizontalBox::Slot()
-		.AutoWidth()
-		.VAlign(VAlign_Bottom)
-		.Padding(Spacing, 0.f, 0.f, 0.f)
-		[
-			SNew(SBox)
-			.WidthOverride(ButtonSize)
-			.HeightOverride(ButtonSize)
+			+ SHorizontalBox::Slot()
+			.HAlign(HAlign_Center)
+			.VAlign(VAlign_Bottom)
+			.Padding(Spacing, 0.f, 0.f, 0.f)
+			[
+				SNew(STextBlock)
+				.TextStyle(FAppStyle::Get(), TEXT("HintText"))
+				.Text(this, &SAvaStaggerTool::GetSelectionText)
+				.ToolTipText(this, &SAvaStaggerTool::GetSelectionToolTipText)
+				.ColorAndOpacity(this, &SAvaStaggerTool::GetSelectionTextColor)
+			]
+			+ SHorizontalBox::Slot()
+			.AutoWidth()
+			.VAlign(VAlign_Fill)
+			.Padding(Spacing, 0.f, 0.f, 0.f)
 			[
 				SNew(SButton)
 				.ContentPadding(2.f)
-				.ButtonStyle(FAppStyle::Get(), TEXT("SimpleButton"))
-				.ToolTipText(LOCTEXT("ResetToDefaultsToolTip", "Reset All to Defaults"))
-				.OnClicked(this, &SAvaStaggerTool::OnResetToDefaultsClick)
-				.IsEnabled_Lambda([this]()
-					{
-						return WeakTool.IsValid() ? WeakTool.Pin()->GetSettings()->CanResetToolOptions() : false;
-					})
+				.ToolTipText(FSequencerCommands::Get().AlignSelectionToPlayhead->GetDescription())
+				.IsEnabled(this, &SAvaStaggerTool::OnCanAlignToPlayhead)
+				.OnClicked(this, &SAvaStaggerTool::OnAlignToPlayhead)
 				[
 					SNew(SImage)
 					.DesiredSizeOverride(FVector2D(ImageSize))
 					.ColorAndOpacity(FSlateColor::UseForeground())
-					.Image(FAppStyle::GetBrush(TEXT("PropertyWindow.DiffersFromDefault")))
+					.Image(FAppStyle::GetBrush(TEXT("HorizontalAlignment_Left")))
+				]
+			]
+			+ SHorizontalBox::Slot()
+			.AutoWidth()
+			.VAlign(VAlign_Fill)
+			.Padding(Spacing, 0.f, 0.f, 0.f)
+			[
+				SNew(SButton)
+				.ContentPadding(2.f)
+				.Text(LOCTEXT("ApplyButtonText", "Apply"))
+				.OnClicked(this, &SAvaStaggerTool::OnApplyButtonClick)
+				.IsEnabled_Lambda([this]()
+					{
+						if (const TSharedPtr<FAvaStaggerTool> Tool = WeakTool.Pin())
+						{
+							return Tool->HasValidSelection() && !Tool->IsAutoApplying();
+						}
+						return false;
+					})
+			]
+			+ SHorizontalBox::Slot()
+			.AutoWidth()
+			.VAlign(VAlign_Fill)
+			.Padding(Spacing, 0.f, Spacing, 0.f)
+			[
+				SNew(SBox)
+				.WidthOverride(ButtonSize)
+				[
+					SNew(SButton)
+					.ContentPadding(2.f)
+					.ButtonStyle(FAppStyle::Get(), TEXT("SimpleButton"))
+					.ToolTipText(LOCTEXT("ResetToDefaultsToolTip", "Reset All to Defaults"))
+					.OnClicked(this, &SAvaStaggerTool::OnResetToDefaultsClick)
+					.IsEnabled_Lambda([this]()
+						{
+							return WeakTool.IsValid() ? WeakTool.Pin()->GetSettings()->CanResetToolOptions() : false;
+						})
+					[
+						SNew(SImage)
+						.DesiredSizeOverride(FVector2D(ImageSize))
+						.ColorAndOpacity(FSlateColor::UseForeground())
+						.Image(FAppStyle::GetBrush(TEXT("PropertyWindow.DiffersFromDefault")))
+					]
 				]
 			]
 		];
