[Editor Viewports] Fix a crash when starting mesh paint mode, caused by auto focus being set to a minimized viewport window in the avalanche plugin

diff --git a/Engine/Plugins/VirtualProduction/Avalanche/Source/AvalancheEditor/Private/Viewport/AvaLevelViewportExtension.cpp b/Engine/Plugins/VirtualProduction/Avalanche/Source/AvalancheEditor/Private/Viewport/AvaLevelViewportExtension.cpp
--- a/Engine/Plugins/VirtualProduction/Avalanche/Source/AvalancheEditor/Private/Viewport/AvaLevelViewportExtension.cpp
+++ b/Engine/Plugins/VirtualProduction/Avalanche/Source/AvalancheEditor/Private/Viewport/AvaLevelViewportExtension.cpp
@@ -63,28 +63,31 @@ namespace UE::AvaEditor::Private
 			return;
 		}
 
-		TSharedPtr<FEditorViewportClient> ViewportClientToFocus = nullptr;
+		FEditorViewportClient* ViewportClientToFocus = nullptr;
 
 		// Get the first valid Motion Design viewport
 		for (const TSharedPtr<SLevelViewport>& Viewport : Viewports)
 		{
-			TSharedPtr<FEditorViewportClient> ViewportClient = Viewport->GetViewportClient();
-			TSharedPtr<IAvaViewportClient> AvaViewportClient = FAvaViewportUtils::GetAsAvaViewportClient(ViewportClient.Get());
-
-			if (AvaViewportClient && AvaViewportClient->IsMotionDesignViewport())
+			if (Viewport->IsVisible())
 			{
-				ViewportClientToFocus = ViewportClient;
-				break;
+				TSharedPtr<FEditorViewportClient> ViewportClient = Viewport->GetViewportClient();
+				TSharedPtr<IAvaViewportClient> AvaViewportClient = FAvaViewportUtils::GetAsAvaViewportClient(ViewportClient.Get());
+
+				if (AvaViewportClient && AvaViewportClient->IsMotionDesignViewport())
+				{
+					ViewportClientToFocus = ViewportClient.Get();
+					break;
+				}
 			}
 		}
 
-		// if there isn't any Motion Design viewport, use first viewport client
+		// if there isn't any Motion Design viewport, use the global viewport client that was set by the level editor
 		if (!ViewportClientToFocus)
 		{
-			ViewportClientToFocus = Viewports[0]->GetViewportClient();
+			ViewportClientToFocus = GCurrentLevelEditingViewportClient;
 		}
 
-		if (ensure(ViewportClientToFocus.IsValid()))
+		if (ensure(ViewportClientToFocus))
 		{
 			ViewportClientToFocus->ReceivedFocus(ViewportClientToFocus->Viewport);
 		}
@@ -193,8 +196,6 @@ void FAvaLevelViewportExtension::Deactivate()
 		LevelEditorModule->UnregisterViewportType(FAvaLevelViewportLayoutEntity::GetStaticType());
 	}
 
-	CheckValidViewportType();
-
 	UnbindCameraCutDelegate();
 }
 
@@ -395,30 +396,8 @@ void FAvaLevelViewportExtension::OnMapChanged(UWorld* InWorld, EMapChangeType In
 		return;
 	}
 
-	const bool bViewportClientWasValid = !!GCurrentLevelEditingViewportClient;
-
 	CheckValidViewportType();
 
-	// Only deal with setting a Current Level Editing VP Client if it was unset in CheckValidViewportType
-	if (GEditor && bViewportClientWasValid && !GCurrentLevelEditingViewportClient)
-	{
-		// Logic copied from FLevelViewportTabContent::OnLayoutChanged to keep consistent behavior to how the CurrentLevelEditingVPClient is set
-		const TArray<FLevelEditorViewportClient*>& LevelViewportClients = GEditor->GetLevelViewportClients();
-		for (FLevelEditorViewportClient* LevelViewport : LevelViewportClients)
-		{
-			if (LevelViewport && LevelViewport->IsPerspective())
-			{
-				LevelViewport->SetCurrentViewport();
-				break;
-			}
-		}
-
-		if (!GCurrentLevelEditingViewportClient && LevelViewportClients.Num())
-		{
-			GCurrentLevelEditingViewportClient = LevelViewportClients[0];
-		}
-	}
-
 	if (AActor* LastCameraCutActor = LastCameraCutActorWeak.Get())
 	{
 		SetActiveCamera(LastCameraCutActor, true);
