[MotionDesign]
- Keep pinned drawer state when switching between sequences
- Remove old, not used, sidebar settings and function

diff --git a/Engine/Plugins/VirtualProduction/Avalanche/Source/AvalancheSequencer/Private/AvaSequencer.cpp b/Engine/Plugins/VirtualProduction/Avalanche/Source/AvalancheSequencer/Private/AvaSequencer.cpp
--- a/Engine/Plugins/VirtualProduction/Avalanche/Source/AvalancheSequencer/Private/AvaSequencer.cpp
+++ b/Engine/Plugins/VirtualProduction/Avalanche/Source/AvalancheSequencer/Private/AvaSequencer.cpp
@@ -1789,7 +1789,7 @@ TSharedRef<SWidget> FAvaSequencer::CreateSequenceWidget()
 
 	const TSharedRef<ISequencer> Sequencer = GetSequencer();
 
-	UAvaSequencerSettings* const SequencerSettings = GetMutableDefault<UAvaSequencerSettings>();
+	USequencerSettings* const SequencerSettings = Sequencer->GetSequencerSettings();
 	check(SequencerSettings);
 	FSidebarState& SidebarState = SequencerSettings->GetSidebarState();
 
@@ -2109,14 +2109,6 @@ TSharedRef<FAvaEaseCurveTool> FAvaSequencer::GetEaseCurveTool() const
 	return EaseCurveTool.ToSharedRef(); 
 }
 
-void FAvaSequencer::OnSidebarStateChanged(const FSidebarState& InNewState)
-{
-	if (UAvaSequencerSettings* const SequencerSettings = GetMutableDefault<UAvaSequencerSettings>())
-	{
-		SequencerSettings->SetSidebarState(InNewState);
-	}
-}
-
 void FAvaSequencer::ExtendSidebarSelectionMenu(FMenuBuilder& OutMenuBuilder)
 {
 	const TSharedPtr<FAvaSequencerEaseCurveToolSection> Section = MakeShared<FAvaSequencerEaseCurveToolSection>(SharedThis(this));
diff --git a/Engine/Plugins/VirtualProduction/Avalanche/Source/AvalancheSequencer/Private/AvaSequencer.h b/Engine/Plugins/VirtualProduction/Avalanche/Source/AvalancheSequencer/Private/AvaSequencer.h
--- a/Engine/Plugins/VirtualProduction/Avalanche/Source/AvalancheSequencer/Private/AvaSequencer.h
+++ b/Engine/Plugins/VirtualProduction/Avalanche/Source/AvalancheSequencer/Private/AvaSequencer.h
@@ -245,8 +245,6 @@ private:
 
 	void OnUpdateCameraCut(UObject* InCameraObject, bool bInJumpCut);
 
-	void OnSidebarStateChanged(const FSidebarState& InNewState);
-
 	void ExtendSidebarSelectionMenu(FMenuBuilder& OutMenuBuilder);
 	void ExtendSidebarMarkedFramesMenu(FMenuBuilder& OutMenuBuilder);
 
diff --git a/Engine/Plugins/VirtualProduction/Avalanche/Source/AvalancheSequencer/Public/Settings/AvaSequencerSettings.h b/Engine/Plugins/VirtualProduction/Avalanche/Source/AvalancheSequencer/Public/Settings/AvaSequencerSettings.h
--- a/Engine/Plugins/VirtualProduction/Avalanche/Source/AvalancheSequencer/Public/Settings/AvaSequencerSettings.h
+++ b/Engine/Plugins/VirtualProduction/Avalanche/Source/AvalancheSequencer/Public/Settings/AvaSequencerSettings.h
@@ -42,14 +42,16 @@ public:
 		return CustomSequencePresets;
 	}
 
+	UE_DEPRECATED(5.6, "The Motion Design sidebar drawer was moved to the Sequencer sidebar and now stores its state in USequencerSettings. Use USequencerSettings::GetSidebarState() instead.")
 	FSidebarState& GetSidebarState()
 	{
-		return SidebarState;
+		return SidebarState_DEPRECATED;
 	}
 
+	UE_DEPRECATED(5.6, "The Motion Design sidebar drawer was moved to the Sequencer sidebar and now stores its state in USequencerSettings. Use USequencerSettings::SetSidebarState() instead.")
 	void SetSidebarState(const FSidebarState& InSidebarState)
 	{
-		SidebarState = InSidebarState;
+		SidebarState_DEPRECATED = InSidebarState;
 	}
 
 private:
@@ -70,6 +72,6 @@ private:
 	TSet<FAvaSequencePreset> CustomSequencePresets;
 
 	/** The state of a sidebar to be restored when Sequencer is initialized */
-	UPROPERTY(Config)
-	FSidebarState SidebarState;
+	UPROPERTY(Config, meta = (DeprecatedProperty, DeprecationMessage = "The Motion Design sidebar drawer was moved to the Sequencer sidebar and now stores its state in USequencerSettings"))
+	FSidebarState SidebarState_DEPRECATED;
 };
