MotionDesign : Modifiers
- Fixed dynamic mesh converter modifier not removing dynamic mesh component when removed

diff --git a/Engine/Plugins/VirtualProduction/Avalanche/Source/AvalancheModifiers/Private/Modifiers/AvaDynamicMeshConverterModifier.cpp b/Engine/Plugins/VirtualProduction/Avalanche/Source/AvalancheModifiers/Private/Modifiers/AvaDynamicMeshConverterModifier.cpp
--- a/Engine/Plugins/VirtualProduction/Avalanche/Source/AvalancheModifiers/Private/Modifiers/AvaDynamicMeshConverterModifier.cpp
+++ b/Engine/Plugins/VirtualProduction/Avalanche/Source/AvalancheModifiers/Private/Modifiers/AvaDynamicMeshConverterModifier.cpp
@@ -121,10 +121,7 @@ void UAvaDynamicMeshConverterModifier::OnModifierRemoved(EActorModifierCoreDisab
 {
 	Super::OnModifierRemoved(InReason);
 
-	if (InReason != EActorModifierCoreDisableReason::Destroyed)
-	{
-		RemoveDynamicMeshComponent();
-	}
+	RemoveDynamicMeshComponent();
 }
 
 bool UAvaDynamicMeshConverterModifier::IsModifierDirtyable() const
