MotionDesign : Shapes
- Fixed scale from zero sequence bug where geometry would not appear after level reload (shapes, Text3D)
- Fixed set primary material everywhere not applying by default and on level reload
- Fixed 3D shape material translucency mode is set to auto instead of enabled

diff --git a/Engine/Plugins/Experimental/Text3D/Source/Text3D/Private/Text3DComponent.cpp b/Engine/Plugins/Experimental/Text3D/Source/Text3D/Private/Text3DComponent.cpp
--- a/Engine/Plugins/Experimental/Text3D/Source/Text3D/Private/Text3DComponent.cpp
+++ b/Engine/Plugins/Experimental/Text3D/Source/Text3D/Private/Text3DComponent.cpp
@@ -58,6 +58,11 @@ UText3DComponent::UText3DComponent()
 	LayoutExtension = CreateDefaultSubobject<UText3DLayoutExtensionBase, UText3DDefaultLayoutExtension>(TEXT("LayoutExtension"));
 	RenderingExtension = CreateDefaultSubobject<UText3DRenderingExtensionBase, UText3DDefaultRenderingExtension>(TEXT("RenderingExtension"));
 	CharacterExtension = CreateDefaultSubobject<UText3DCharacterExtensionBase, UText3DDefaultCharacterExtension>(TEXT("CharacterExtension"));
+
+	if (!IsTemplate())
+	{
+		TransformUpdated.AddUObject(this, &UText3DComponent::OnComponentTransformed);
+	}
 }
 
 void UText3DComponent::OnComponentCreated()
@@ -75,6 +80,8 @@ void UText3DComponent::OnComponentDestroyed(bool bInDestroyingHierarchy)
 	{
 		TextRenderer->Destroy();
 	}
+
+	TransformUpdated.RemoveAll(this);
 }
 
 void UText3DComponent::PostLoad()
@@ -224,6 +231,31 @@ void UText3DComponent::UpdateStatistics()
 	}
 }
 
+void UText3DComponent::OnComponentTransformed(USceneComponent*, EUpdateTransformFlags, ETeleportType)
+{
+	static TMap<FObjectKey, bool> TextComponentScale;
+
+	/*
+	 * Find previous zero scaled status or initialize it to false
+	 * TransformUpdated gets broadcasted when component is registered first time so it will have the correct starting value (see UpdateComponentToWorld in UActorComponent::OnRegister)
+	 * Further calls affecting transform value, will have the proper zero scaled status
+	 */
+	bool& bZeroScaled = TextComponentScale.FindOrAdd(this, /** Default */false);
+	const bool bIsZeroScaled = FMath::IsNearlyZero(GetComponentScale().GetMin(), UE_KINDA_SMALL_NUMBER);
+
+	if (bIsZeroScaled)
+	{
+		bZeroScaled = true;
+	}
+
+	// Only update when scale was zero and is not anymore
+	if (bZeroScaled && !bIsZeroScaled)
+	{
+		RequestUpdate(EText3DRendererFlags::Layout);
+		bZeroScaled = false;
+	}
+}
+
 void UText3DComponent::OnRegister()
 {
 	Super::OnRegister();
@@ -232,8 +264,9 @@ void UText3DComponent::OnRegister()
 	{
 		TextRenderer->Create();
 	}
-	
-	RequestUpdate(EText3DRendererFlags::None);
+
+	constexpr bool bImmediate = true;
+	RequestUpdate(UpdateFlags, bImmediate);
 }
 
 void UText3DComponent::OnUnregister()
@@ -914,6 +947,11 @@ void UText3DComponent::Rebuild(EText3DRendererFlags InUpdateFlags, bool bCleanCa
 {
 	EnumAddFlags(UpdateFlags, InUpdateFlags);
 
+	if (UpdateFlags == EText3DRendererFlags::None)
+	{
+		return;
+	}
+
 	if (IsTemplate(RF_ClassDefaultObject))
 	{
 		return;
@@ -1049,7 +1087,8 @@ void UText3DComponent::OnTextRendererClassChanged()
 	if (IsRegistered())
 	{
 		TextRenderer->Create();
-		RequestUpdate(EText3DRendererFlags::All);
+		constexpr bool bImmediate = true;
+		RequestUpdate(EText3DRendererFlags::All, bImmediate);
 	}
 }
 
diff --git a/Engine/Plugins/Experimental/Text3D/Source/Text3D/Public/Text3DComponent.h b/Engine/Plugins/Experimental/Text3D/Source/Text3D/Public/Text3DComponent.h
--- a/Engine/Plugins/Experimental/Text3D/Source/Text3D/Public/Text3DComponent.h
+++ b/Engine/Plugins/Experimental/Text3D/Source/Text3D/Public/Text3DComponent.h
@@ -513,6 +513,8 @@ private:
 
 	void UpdateStatistics();
 
+	void OnComponentTransformed(USceneComponent*, EUpdateTransformFlags, ETeleportType);
+
 #if WITH_EDITORONLY_DATA
 	UE_DEPRECATED(5.6, "Use Material Extension instead")
 	UPROPERTY()
diff --git a/Engine/Plugins/VirtualProduction/Avalanche/Source/AvalancheShapes/Private/DynamicMeshes/AvaShapeDynMeshBase.cpp b/Engine/Plugins/VirtualProduction/Avalanche/Source/AvalancheShapes/Private/DynamicMeshes/AvaShapeDynMeshBase.cpp
--- a/Engine/Plugins/VirtualProduction/Avalanche/Source/AvalancheShapes/Private/DynamicMeshes/AvaShapeDynMeshBase.cpp
+++ b/Engine/Plugins/VirtualProduction/Avalanche/Source/AvalancheShapes/Private/DynamicMeshes/AvaShapeDynMeshBase.cpp
@@ -10,6 +10,7 @@
 #include "DynamicMesh/DynamicMesh3.h"
 #include "DynamicMesh/DynamicMeshAttributeSet.h"
 #include "DynamicMesh/MeshTransforms.h"
+#include "DynamicMeshes/AvaShape2DDynMeshBase.h"
 #include "Engine/StaticMesh.h"
 #include "Framework/AvaGizmoComponent.h"
 #include "GeometryScript/MeshAssetFunctions.h"
@@ -127,27 +128,23 @@ void UAvaShapeDynamicMeshBase::SetupMeshes()
 			// register other meshes
 			RegisterMeshes();
 
+			const bool bSetPrimaryMaterial = bUsePrimaryMaterialEverywhere;
+			const bool bIs2DShape = IsA<UAvaShape2DDynMeshBase>();
 			for (const int32& MeshIdx : GetMeshesIndexes())
 			{
 				FAvaShapeMeshData* MeshData = GetMeshData(MeshIdx);
 
-				if (MeshIdx != MESH_INDEX_PRIMARY)
-				{
-					MeshData->bUsesPrimaryMaterialParams = bUsePrimaryMaterialEverywhere;
-				}
-
 				// Enable translucency by default
-				MeshData->ParametricMaterial.SetTranslucency(EAvaShapeParametricMaterialTranslucency::Enabled);
-
-				OnMaterialTypeChanged(MeshIdx);
+				MeshData->ParametricMaterial.SetTranslucency(bIs2DShape ? EAvaShapeParametricMaterialTranslucency::Enabled : EAvaShapeParametricMaterialTranslucency::Auto);
+				OnMaterialChanged(MeshIdx);
 			}
-		}
-		else
-		{
-			// handle duplication and loading saved asset here
-			OnUsePrimaryMaterialEverywhereChanged();
+
+			bUsePrimaryMaterialEverywhere = bSetPrimaryMaterial;
 		}
 
+		// handle duplication and loading saved asset here
+		OnUsePrimaryMaterialEverywhereChanged();
+
 		// meshes setup done
 		OnRegisteredMeshes();
 	}
@@ -193,12 +190,13 @@ void UAvaShapeDynamicMeshBase::PostLoad()
 		return;
 	}
 
-	const UDynamicMeshComponent* ShapeComponent = GetShapeMeshComponent();
+	UDynamicMeshComponent* ShapeComponent = GetShapeMeshComponent();
 	if (!IsValid(ShapeComponent))
 	{
 		return;
 	}
 
+	const bool bSetPrimaryMaterial = bUsePrimaryMaterialEverywhere;
 	for (TPair<int32, FAvaShapeMeshData>& MeshData : MeshDatas)
 	{
 		if (MeshData.Value.MaterialType != EMaterialType::Parametric)
@@ -225,6 +223,9 @@ void UAvaShapeDynamicMeshBase::PostLoad()
 		// If the shape has no valid Material Designer Instance, create it from the parametric material struct and apply
 		SetMaterialDirect(MeshData.Key, MeshData.Value.ParametricMaterial.GetOrCreateMaterial(ShapeActor));
 	}
+
+	bUsePrimaryMaterialEverywhere = bSetPrimaryMaterial;
+	OnUsePrimaryMaterialEverywhereChanged();
 }
 
 void UAvaShapeDynamicMeshBase::OnComponentDestroyed(bool bDestroyingHierarchy)
@@ -546,12 +547,23 @@ void UAvaShapeDynamicMeshBase::OnUsePrimaryMaterialEverywhereChanged()
 
 void UAvaShapeDynamicMeshBase::OnMeshTransformChanged(USceneComponent* InComponent, EUpdateTransformFlags InTransformFlags, ETeleportType InTeleport)
 {
-	const bool bIsZeroScaled = InComponent->GetComponentScale().IsNearlyZero();
+	/*
+	 * TransformUpdated gets broadcasted when component is registered first time so it will have the correct starting value (see UpdateComponentToWorld in UActorComponent::OnRegister)
+	 * Further calls affecting transform value, will have the proper zero scaled status
+	 */
+	const bool bIsZeroScaled = FMath::IsNearlyZero(InComponent->GetComponentScale().GetMin(), UE_KINDA_SMALL_NUMBER);
+
+	if (bIsZeroScaled)
+	{
+		bZeroScaled = true;
+	}
+
+	// Only update when scale was zero and is not anymore
 	if (bZeroScaled && !bIsZeroScaled)
 	{
-		OnSizeChanged();
+		MarkAllMeshesDirty();
+		bZeroScaled = false;
 	}
-	bZeroScaled = bIsZeroScaled;
 }
 
 void UAvaShapeDynamicMeshBase::SetOverridePrimaryUVParams(int32 InMeshIndex, bool bInOverride)
@@ -1856,7 +1868,7 @@ void UAvaShapeDynamicMeshBase::SetShapeComponentMaterial(int32 InMaterialIndex,
 
 void UAvaShapeDynamicMeshBase::ScaleVertices(const FVector& InScale)
 {
-	if (InScale.IsNearlyZero() || InScale.Equals(FVector::OneVector))
+	if (InScale.IsNearlyZero())
 	{
 		return;
 	}
@@ -2290,7 +2302,6 @@ void UAvaShapeDynamicMeshBase::InitializeDynamicMesh()
 		, /** Broadcast */false
 	);
 
-	bZeroScaled = MeshComponent->GetComponentScale().IsNearlyZero();
 	MeshComponent->TransformUpdated.AddUObject(this, &UAvaShapeDynamicMeshBase::OnMeshTransformChanged);
 }
 
