MotionDesign : Modifiers
- Fixed pattern modifier crash due to high instance/triangle count requested, added a limit to 10,000 instances and 10,000,000 triangles

diff --git a/Engine/Plugins/VirtualProduction/Avalanche/Source/AvalancheModifiers/Private/Modifiers/AvaPatternModifier.cpp b/Engine/Plugins/VirtualProduction/Avalanche/Source/AvalancheModifiers/Private/Modifiers/AvaPatternModifier.cpp
--- a/Engine/Plugins/VirtualProduction/Avalanche/Source/AvalancheModifiers/Private/Modifiers/AvaPatternModifier.cpp
+++ b/Engine/Plugins/VirtualProduction/Avalanche/Source/AvalancheModifiers/Private/Modifiers/AvaPatternModifier.cpp
@@ -58,14 +58,29 @@ void UAvaPatternModifier::Apply()
 		return;
 	}
 
-	TargetMeshComponent->GetDynamicMesh()->EditMesh([this](FDynamicMesh3& AppendToMesh)
+	// Get the original bounds before clearing the mesh
+	OriginalMeshBounds = GetMeshBounds();
+
+	const TArray<FTransform> Transforms = ActiveTool->GetTransformInstances(OriginalMeshBounds);
+	if (Transforms.IsEmpty())
+	{
+		Fail(LOCTEXT("InvalidInstanceCount", "Requested instance count is not supported"));
+		return;
+	}
+
+	const uint32 TriangleCount = TargetMeshComponent->GetDynamicMesh()->GetTriangleCount() * Transforms.Num();
+	constexpr uint32 MaxTriangleCount = 10000000;
+	if (TriangleCount > MaxTriangleCount)
+	{
+		Fail(FText::Format(LOCTEXT("HighInstanceCount", "Requested instance count is too high:\n{0} instances totalling in {1} triangles"), FText::AsNumber(Transforms.Num()), FText::AsNumber(TriangleCount)));
+		return;
+	}
+
+	TargetMeshComponent->GetDynamicMesh()->EditMesh([this, &Transforms](FDynamicMesh3& AppendToMesh)
 	{
 		// copy original mesh once into tmp mesh
 		FDynamicMesh3 TmpMesh = AppendToMesh;
 
-		// Get the original bounds before clearing the mesh
-		OriginalMeshBounds = GetMeshBounds();
-
 		// clear all since we have a copy
 		for (const int32 TId : AppendToMesh.TriangleIndicesItr())
 		{
@@ -74,7 +89,7 @@ void UAvaPatternModifier::Apply()
 
 		FMeshIndexMappings TmpMappings;
 		FDynamicMeshEditor Editor(&AppendToMesh);
-		for (const FTransform& TransformInstance : ActiveTool->GetTransformInstances(OriginalMeshBounds))
+		for (const FTransform& TransformInstance : Transforms)
 		{
 			MeshTransforms::ApplyTransform(TmpMesh, TransformInstance, /** ReverseIfNeeded */true);
 			Editor.AppendMesh(&TmpMesh, TmpMappings);
diff --git a/Engine/Plugins/VirtualProduction/Avalanche/Source/AvalancheModifiers/Private/Tools/AvaPatternModifierCircleTool.cpp b/Engine/Plugins/VirtualProduction/Avalanche/Source/AvalancheModifiers/Private/Tools/AvaPatternModifierCircleTool.cpp
--- a/Engine/Plugins/VirtualProduction/Avalanche/Source/AvalancheModifiers/Private/Tools/AvaPatternModifierCircleTool.cpp
+++ b/Engine/Plugins/VirtualProduction/Avalanche/Source/AvalancheModifiers/Private/Tools/AvaPatternModifierCircleTool.cpp
@@ -48,7 +48,7 @@ void UAvaPatternModifierCircleTool::SetCircleFullAngle(float InCircleFullAngle)
 
 void UAvaPatternModifierCircleTool::SetCircleCount(int32 InCircleCount)
 {
-	InCircleCount = FMath::Max(InCircleCount, 1);
+	InCircleCount = FMath::Clamp(InCircleCount, 1, 10000);
 	if (CircleCount == InCircleCount)
 	{
 		return;
diff --git a/Engine/Plugins/VirtualProduction/Avalanche/Source/AvalancheModifiers/Private/Tools/AvaPatternModifierGridTool.cpp b/Engine/Plugins/VirtualProduction/Avalanche/Source/AvalancheModifiers/Private/Tools/AvaPatternModifierGridTool.cpp
--- a/Engine/Plugins/VirtualProduction/Avalanche/Source/AvalancheModifiers/Private/Tools/AvaPatternModifierGridTool.cpp
+++ b/Engine/Plugins/VirtualProduction/Avalanche/Source/AvalancheModifiers/Private/Tools/AvaPatternModifierGridTool.cpp
@@ -26,7 +26,7 @@ void UAvaPatternModifierGridTool::SetGridAlignment(EAvaPatternModifierGridAlignm
 
 void UAvaPatternModifierGridTool::SetGridCountX(int32 InGridCountX)
 {
-	InGridCountX = FMath::Max(InGridCountX, 1);
+	InGridCountX = FMath::Clamp(InGridCountX, 1, 10000);
 	if (GridCountX == InGridCountX)
 	{
 		return;
@@ -38,7 +38,7 @@ void UAvaPatternModifierGridTool::SetGridCountX(int32 InGridCountX)
 
 void UAvaPatternModifierGridTool::SetGridCountY(int32 InGridCountY)
 {
-	InGridCountY = FMath::Max(InGridCountY, 1);
+	InGridCountY = FMath::Clamp(InGridCountY, 1, 10000);
 	if (GridCountY == InGridCountY)
 	{
 		return;
diff --git a/Engine/Plugins/VirtualProduction/Avalanche/Source/AvalancheModifiers/Private/Tools/AvaPatternModifierLineTool.cpp b/Engine/Plugins/VirtualProduction/Avalanche/Source/AvalancheModifiers/Private/Tools/AvaPatternModifierLineTool.cpp
--- a/Engine/Plugins/VirtualProduction/Avalanche/Source/AvalancheModifiers/Private/Tools/AvaPatternModifierLineTool.cpp
+++ b/Engine/Plugins/VirtualProduction/Avalanche/Source/AvalancheModifiers/Private/Tools/AvaPatternModifierLineTool.cpp
@@ -26,7 +26,7 @@ void UAvaPatternModifierLineTool::SetLineAlignment(EAvaPatternModifierLineAlignm
 
 void UAvaPatternModifierLineTool::SetLineCount(int32 InLineCount)
 {
-	InLineCount = FMath::Max(InLineCount, 1);
+	InLineCount = FMath::Clamp(InLineCount, 1, 10000);
 	if (LineCount == InLineCount)
 	{
 		return;
diff --git a/Engine/Plugins/VirtualProduction/Avalanche/Source/AvalancheModifiers/Public/Tools/AvaPatternModifierCircleTool.h b/Engine/Plugins/VirtualProduction/Avalanche/Source/AvalancheModifiers/Public/Tools/AvaPatternModifierCircleTool.h
--- a/Engine/Plugins/VirtualProduction/Avalanche/Source/AvalancheModifiers/Public/Tools/AvaPatternModifierCircleTool.h
+++ b/Engine/Plugins/VirtualProduction/Avalanche/Source/AvalancheModifiers/Public/Tools/AvaPatternModifierCircleTool.h
@@ -100,27 +100,27 @@ protected:
 	virtual FName GetToolName() const override;
 	//~ End UAvaPatternModifierTool
 
-	UPROPERTY(EditInstanceOnly, Setter, Getter, Category="Pattern", meta=(EditCondition="Layout == EAvaPatternModifierLayout::Circle", EditConditionHides))
+	UPROPERTY(EditInstanceOnly, Setter, Getter, Category="Pattern")
 	EAvaPatternModifierPlane CirclePlane = EAvaPatternModifierPlane::YZ;
 
-	UPROPERTY(EditInstanceOnly, Setter, Getter, Category="Pattern", meta=(EditCondition="Layout == EAvaPatternModifierLayout::Circle", EditConditionHides))
+	UPROPERTY(EditInstanceOnly, Setter, Getter, Category="Pattern")
 	float CircleRadius = 100.f;
 
-	UPROPERTY(EditInstanceOnly, Setter, Getter, Category="Pattern", meta=(EditCondition="Layout == EAvaPatternModifierLayout::Circle", EditConditionHides))
+	UPROPERTY(EditInstanceOnly, Setter, Getter, Category="Pattern")
 	float CircleStartAngle = 180.f;
 
-	UPROPERTY(EditInstanceOnly, Setter, Getter, Category="Pattern", meta=(EditCondition="Layout == EAvaPatternModifierLayout::Circle", EditConditionHides))
+	UPROPERTY(EditInstanceOnly, Setter, Getter, Category="Pattern")
 	float CircleFullAngle = 360.f;
 
-	UPROPERTY(EditInstanceOnly, Setter, Getter, Category="Pattern", meta=(ClampMin="1", EditCondition="Layout == EAvaPatternModifierLayout::Circle", EditConditionHides))
+	UPROPERTY(EditInstanceOnly, Setter, Getter, Category="Pattern", meta=(ClampMin="1", ClampMax="10000"))
 	int32 CircleCount = 4;
 
-	UPROPERTY(EditInstanceOnly, Setter="SetCircleAccumulateTransform", Getter="GetCircleAccumulateTransform", Category="Pattern", meta=(EditCondition="Layout == EAvaPatternModifierLayout::Circle", EditConditionHides))
+	UPROPERTY(EditInstanceOnly, Setter="SetCircleAccumulateTransform", Getter="GetCircleAccumulateTransform", Category="Pattern")
 	bool bCircleAccumulateTransform = false;
 
-	UPROPERTY(EditInstanceOnly, Setter, Getter, Category="Pattern", meta=(EditCondition="Layout == EAvaPatternModifierLayout::Circle", EditConditionHides))
+	UPROPERTY(EditInstanceOnly, Setter, Getter, Category="Pattern")
 	FRotator CircleRotation = FRotator::ZeroRotator;
 
-	UPROPERTY(EditInstanceOnly, Setter, Getter, Category="Pattern", meta=(EditCondition="Layout == EAvaPatternModifierLayout::Circle", EditConditionHides, MotionDesignVectorWidget, AllowPreserveRatio="XYZ"))
+	UPROPERTY(EditInstanceOnly, Setter, Getter, Category="Pattern", meta=(MotionDesignVectorWidget, AllowPreserveRatio="XYZ"))
 	FVector CircleScale = FVector::OneVector;
 };
diff --git a/Engine/Plugins/VirtualProduction/Avalanche/Source/AvalancheModifiers/Public/Tools/AvaPatternModifierGridTool.h b/Engine/Plugins/VirtualProduction/Avalanche/Source/AvalancheModifiers/Public/Tools/AvaPatternModifierGridTool.h
--- a/Engine/Plugins/VirtualProduction/Avalanche/Source/AvalancheModifiers/Public/Tools/AvaPatternModifierGridTool.h
+++ b/Engine/Plugins/VirtualProduction/Avalanche/Source/AvalancheModifiers/Public/Tools/AvaPatternModifierGridTool.h
@@ -109,30 +109,30 @@ protected:
 	virtual FName GetToolName() const override;
 	//~ End UAvaPatternModifierTool
 
-	UPROPERTY(EditInstanceOnly, Setter, Getter, Category="Pattern", meta=(EditCondition="Layout == EAvaPatternModifierLayout::Grid", EditConditionHides))
+	UPROPERTY(EditInstanceOnly, Setter, Getter, Category="Pattern")
 	EAvaPatternModifierPlane GridPlane = EAvaPatternModifierPlane::YZ;
 
-	UPROPERTY(EditInstanceOnly, Setter, Getter, Category="Pattern", meta=(EditCondition="Layout == EAvaPatternModifierLayout::Grid", EditConditionHides))
+	UPROPERTY(EditInstanceOnly, Setter, Getter, Category="Pattern")
 	EAvaPatternModifierGridAlignment GridAlignment = EAvaPatternModifierGridAlignment::Center;
 
-	UPROPERTY(EditInstanceOnly, Setter, Getter, Category="Pattern", meta=(ClampMin="1", EditCondition="Layout == EAvaPatternModifierLayout::Grid", EditConditionHides))
+	UPROPERTY(EditInstanceOnly, Setter, Getter, Category="Pattern", meta=(ClampMin="1", ClampMax="10000"))
 	int32 GridCountX = 2;
 
-	UPROPERTY(EditInstanceOnly, Setter, Getter, Category="Pattern", meta=(ClampMin="1", EditCondition="Layout == EAvaPatternModifierLayout::Grid", EditConditionHides))
+	UPROPERTY(EditInstanceOnly, Setter, Getter, Category="Pattern", meta=(ClampMin="1", ClampMax="10000"))
 	int32 GridCountY = 2;
 
-	UPROPERTY(EditInstanceOnly, Setter, Getter, Category="Pattern", meta=(EditCondition="Layout == EAvaPatternModifierLayout::Grid", EditConditionHides))
+	UPROPERTY(EditInstanceOnly, Setter, Getter, Category="Pattern")
 	float GridSpacingX = 0.f;
 
-	UPROPERTY(EditInstanceOnly, Setter, Getter, Category="Pattern", meta=(EditCondition="Layout == EAvaPatternModifierLayout::Grid", EditConditionHides))
+	UPROPERTY(EditInstanceOnly, Setter, Getter, Category="Pattern")
 	float GridSpacingY = 0.f;
 
-	UPROPERTY(EditInstanceOnly, Setter="SetGridAccumulateTransform", Getter="GetGridAccumulateTransform", Category="Pattern", meta=(EditCondition="Layout == EAvaPatternModifierLayout::Grid", EditConditionHides))
+	UPROPERTY(EditInstanceOnly, Setter="SetGridAccumulateTransform", Getter="GetGridAccumulateTransform", Category="Pattern")
 	bool bGridAccumulateTransform = false;
 
-	UPROPERTY(EditInstanceOnly, Setter, Getter, Category="Pattern", meta=(EditCondition="Layout == EAvaPatternModifierLayout::Grid", EditConditionHides))
+	UPROPERTY(EditInstanceOnly, Setter, Getter, Category="Pattern")
 	FRotator GridRotation = FRotator::ZeroRotator;
 
-	UPROPERTY(EditInstanceOnly, Setter, Getter, Category="Pattern", meta=(EditCondition="Layout == EAvaPatternModifierLayout::Grid", EditConditionHides, MotionDesignVectorWidget, AllowPreserveRatio="XYZ"))
+	UPROPERTY(EditInstanceOnly, Setter, Getter, Category="Pattern", meta=(MotionDesignVectorWidget, AllowPreserveRatio="XYZ"))
 	FVector GridScale = FVector::OneVector;
 };
\ No newline at end of file
diff --git a/Engine/Plugins/VirtualProduction/Avalanche/Source/AvalancheModifiers/Public/Tools/AvaPatternModifierLineTool.h b/Engine/Plugins/VirtualProduction/Avalanche/Source/AvalancheModifiers/Public/Tools/AvaPatternModifierLineTool.h
--- a/Engine/Plugins/VirtualProduction/Avalanche/Source/AvalancheModifiers/Public/Tools/AvaPatternModifierLineTool.h
+++ b/Engine/Plugins/VirtualProduction/Avalanche/Source/AvalancheModifiers/Public/Tools/AvaPatternModifierLineTool.h
@@ -91,24 +91,24 @@ protected:
 	virtual FName GetToolName() const override;
 	//~ End UAvaPatternModifierTool
 
-	UPROPERTY(EditInstanceOnly, Setter, Getter, Category="Pattern", meta=(EditCondition="Layout == EAvaPatternModifierLayout::Line", EditConditionHides))
+	UPROPERTY(EditInstanceOnly, Setter, Getter, Category="Pattern")
 	EAvaPatternModifierAxis LineAxis = EAvaPatternModifierAxis::Y;
 
-	UPROPERTY(EditInstanceOnly, Setter, Getter, Category="Pattern", meta=(EditCondition="Layout == EAvaPatternModifierLayout::Line", EditConditionHides))
+	UPROPERTY(EditInstanceOnly, Setter, Getter, Category="Pattern")
 	EAvaPatternModifierLineAlignment LineAlignment = EAvaPatternModifierLineAlignment::Center;
 
-	UPROPERTY(EditInstanceOnly, Setter, Getter, Category="Pattern", meta=(ClampMin="1", EditCondition="Layout == EAvaPatternModifierLayout::Line", EditConditionHides))
+	UPROPERTY(EditInstanceOnly, Setter, Getter, Category="Pattern", meta=(ClampMin="1", ClampMax="10000"))
 	int32 LineCount = 4;
 
-	UPROPERTY(EditInstanceOnly, Setter, Getter, Category="Pattern", meta=(EditCondition="Layout == EAvaPatternModifierLayout::Line", EditConditionHides))
+	UPROPERTY(EditInstanceOnly, Setter, Getter, Category="Pattern")
 	float LineSpacing = 0.f;
 
-	UPROPERTY(EditInstanceOnly, Setter="SetLineAccumulateTransform", Getter="GetLineAccumulateTransform", Category="Pattern", meta=(EditCondition="Layout == EAvaPatternModifierLayout::Line", EditConditionHides))
+	UPROPERTY(EditInstanceOnly, Setter="SetLineAccumulateTransform", Getter="GetLineAccumulateTransform", Category="Pattern")
 	bool bLineAccumulateTransform = false;
 
-	UPROPERTY(EditInstanceOnly, Setter, Getter, Category="Pattern", meta=(EditCondition="Layout == EAvaPatternModifierLayout::Line", EditConditionHides))
+	UPROPERTY(EditInstanceOnly, Setter, Getter, Category="Pattern")
 	FRotator LineRotation = FRotator::ZeroRotator;
 
-	UPROPERTY(EditInstanceOnly, Setter, Getter, Category="Pattern", meta=(EditCondition="Layout == EAvaPatternModifierLayout::Line", EditConditionHides, MotionDesignVectorWidget, AllowPreserveRatio="XYZ"))
+	UPROPERTY(EditInstanceOnly, Setter, Getter, Category="Pattern", meta=(MotionDesignVectorWidget, AllowPreserveRatio="XYZ"))
 	FVector LineScale = FVector::OneVector;
 };
