Motion Design: Fixes for the color picker.
- Now remembers the selected theme.
- Renaming a theme to a the same as another is blocked.
- Can no longer create a new theme with an existing name.

diff --git a/Engine/Plugins/VirtualProduction/Avalanche/Source/AvalancheLevelViewport/Private/Widgets/SAvaLevelViewportActorColorMenu.cpp b/Engine/Plugins/VirtualProduction/Avalanche/Source/AvalancheLevelViewport/Private/Widgets/SAvaLevelViewportActorColorMenu.cpp
--- a/Engine/Plugins/VirtualProduction/Avalanche/Source/AvalancheLevelViewport/Private/Widgets/SAvaLevelViewportActorColorMenu.cpp
+++ b/Engine/Plugins/VirtualProduction/Avalanche/Source/AvalancheLevelViewport/Private/Widgets/SAvaLevelViewportActorColorMenu.cpp
@@ -1,6 +1,7 @@
 // Copyright Epic Games, Inc. All Rights Reserved.
 
 #include "Widgets/SAvaLevelViewportActorColorMenu.h"
+#include "AvaLevelViewportModule.h"
 #include "AvaLevelViewportStyle.h"
 #include "Brushes/SlateImageBrush.h"
 #include "ColorPicker/AvaViewportColorPickerActorClassRegistry.h"
@@ -31,6 +32,8 @@
 
 #define LOCTEXT_NAMESPACE "SAvaLevelViewportActorColorMenu"
 
+int32 SAvaLevelViewportActorColorMenu::ActiveThemeIndex = 0;
+
 class FAvaColorPickerPaletteDragDropOperation : public FDragDropOperation
 {
 public:
@@ -110,7 +113,7 @@ void SAvaLevelViewportActorColorMenu::SetInactiveColorRGB_Direct(FLinearColor In
 
 void SAvaLevelViewportActorColorMenu::SetInactiveColorHSV_Direct(FLinearColor InHSV)
 {
-	InactiveColor = InHSV.HSVToLinearRGB();	
+	InactiveColor = InHSV.HSVToLinearRGB();
 }
 
 void SAvaLevelViewportActorColorMenu::SetColorStyle_Direct(EAvaColorStyle InStyle)
@@ -171,8 +174,6 @@ void SAvaLevelViewportActorColorMenu::Construct(const FArguments& InArgs, const
 	InactiveColor = FLinearColor::White;
 	ColorStyle = EAvaColorStyle::Solid;
 
-	ActiveThemeIndex = 0;
-
 	SatValueMaterialDynamic = UMaterialInstanceDynamic::Create(SatValueMaterial.Get(), GetTransientPackage());
 
 	HueBrush = MakeShared<FSlateMaterialBrush>(*HueMaterial.Get(), FVector2f(200.f, 200.f));
@@ -280,7 +281,7 @@ void SAvaLevelViewportActorColorMenu::Construct(const FArguments& InArgs, const
 				.OnClicked(this, &SAvaLevelViewportActorColorMenu::OnColorFormatToggleClicked)
 				[
 					SAssignNew(ColorFormatButtonText, STextBlock)
-					.Text(LOCTEXT("RGB","RGB"))
+					.Text(LOCTEXT("RGB", "RGB"))
 					.RenderTransform(FSlateRenderTransform(FVector2f(0.f, -2.f)))
 				]
 			]
@@ -573,7 +574,7 @@ void SAvaLevelViewportActorColorMenu::LoadColorThemesFromIni()
 		NewTheme->GetColors().Add({FLinearColor::Red, LOCTEXT("RedColor", "Red")});
 		NewTheme->GetColors().Add({FLinearColor::Green, LOCTEXT("GreenColor", "Green")});
 		NewTheme->GetColors().Add({FLinearColor::Blue, LOCTEXT("BlueColor", "Blue")});
-		SaveColorThemesToIni();		
+		SaveColorThemesToIni();
 	}
 
 	FAvaColorTheme* EditorAccentColors = FindTheme("Editor Accents", true);
@@ -654,7 +655,7 @@ void SAvaLevelViewportActorColorMenu::UpdateColorPalette()
 	// starting from last: most recent color is shown at the top
 	int32 ColorIdx = ThemeColors.Num() - 1;
 
-	do 
+	do
 	{
 		for (int32 Col = 0; Col < ColorsPerPaletteRow; ++Col)
 		{
@@ -686,8 +687,7 @@ void SAvaLevelViewportActorColorMenu::UpdateColorPalette()
 		{
 			break;
 		}
-	}
-	while (true);
+	} while (true);
 }
 
 void SAvaLevelViewportActorColorMenu::AddColorToTheme(const FLinearColor& InNewColor)
@@ -714,7 +714,7 @@ void SAvaLevelViewportActorColorMenu::AddColorToTheme(const FLinearColor& InNewC
 		{
 			// If we're almost identical to a previously set colour, set that colour instead.
 			if (InNewColor == ActiveColor)
-			{ 
+			{
 				SetActiveColorRGB(ThemeColors[ColorIdx].Color);
 			}
 			else if (InNewColor == InactiveColor)
@@ -758,7 +758,7 @@ TSharedRef<SWidget> SAvaLevelViewportActorColorMenu::GetRGBEntryBox()
 				SNew(STextBlock)
 				.MinDesiredWidth(10.f)
 				.Margin(FMargin(0.f, 5.f, 0.f, 0.f))
-				.Text(LOCTEXT("Red","R"))
+				.Text(LOCTEXT("Red", "R"))
 			]
 		]
 		+ SVerticalBox::Slot()
@@ -779,7 +779,7 @@ TSharedRef<SWidget> SAvaLevelViewportActorColorMenu::GetRGBEntryBox()
 				SNew(STextBlock)
 				.MinDesiredWidth(10.f)
 				.Margin(FMargin(0.f, 5.f, 0.f, 0.f))
-				.Text(LOCTEXT("Green","G"))
+				.Text(LOCTEXT("Green", "G"))
 			]
 		]
 		+ SVerticalBox::Slot()
@@ -800,7 +800,7 @@ TSharedRef<SWidget> SAvaLevelViewportActorColorMenu::GetRGBEntryBox()
 				SNew(STextBlock)
 				.MinDesiredWidth(10.f)
 				.Margin(FMargin(0.f, 5.f, 0.f, 0.f))
-				.Text(LOCTEXT("Blue","B"))
+				.Text(LOCTEXT("Blue", "B"))
 			]
 		];
 }
@@ -826,7 +826,7 @@ TSharedRef<SWidget> SAvaLevelViewportActorColorMenu::GetHSLEntryBox()
 				SNew(STextBlock)
 				.MinDesiredWidth(10.f)
 				.Margin(FMargin(0.f, 5.f, 0.f, 0.f))
-				.Text(LOCTEXT("Hue","H"))
+				.Text(LOCTEXT("Hue", "H"))
 			]
 		]
 		+ SVerticalBox::Slot()
@@ -847,7 +847,7 @@ TSharedRef<SWidget> SAvaLevelViewportActorColorMenu::GetHSLEntryBox()
 				SNew(STextBlock)
 				.MinDesiredWidth(10.f)
 				.Margin(FMargin(0.f, 5.f, 0.f, 0.f))
-				.Text(LOCTEXT("Saturation","S"))
+				.Text(LOCTEXT("Saturation", "S"))
 			]
 		]
 		+ SVerticalBox::Slot()
@@ -868,7 +868,7 @@ TSharedRef<SWidget> SAvaLevelViewportActorColorMenu::GetHSLEntryBox()
 				SNew(STextBlock)
 				.MinDesiredWidth(10.f)
 				.Margin(FMargin(0.f, 5.f, 0.f, 0.f))
-				.Text(LOCTEXT("Value","V"))
+				.Text(LOCTEXT("Value", "V"))
 			]
 		];
 }
@@ -1040,7 +1040,7 @@ FReply SAvaLevelViewportActorColorMenu::OnUnlitButtonClicked()
 {
 	bIsUnlit = !bIsUnlit;
 	BroadcastColorChange();
-	
+
 	return FReply::Handled();
 }
 
@@ -1148,7 +1148,7 @@ TSharedRef<SWidget> SAvaLevelViewportActorColorMenu::BuildRenameColorMenu(int32
 		ColorMenuBuilder.BeginSection("ColorOptions", LOCTEXT("ColorOptions", "Color Options"));
 
 		ColorMenuBuilder.AddEditableText(LOCTEXT("NameLabel", "Name"), LOCTEXT("NameTooltip", "The name of the color."),
-			FSlateIcon(), ThemeColors[InColorIdx].Label, 
+			FSlateIcon(), ThemeColors[InColorIdx].Label,
 			FOnTextCommitted::CreateLambda([Theme, ColorWidget, InColorIdx, CurrentColor](const FText& NewLabel, ETextCommit::Type CommitType)
 				{
 					TArray<FAvaColorInfo>& ThemeColors = Theme->GetColors();
@@ -1213,7 +1213,7 @@ TSharedRef<SWidget> SAvaLevelViewportActorColorMenu::BuildThemesMenu()
 	{
 		static FText RemoveThemeLabelText = LOCTEXT("RemoveThemeLabel", "Remove");
 		static FText RemoveThemeTooltipText = LOCTEXT("RemoveThemeTooltip", "Remove this theme.");
-		
+
 		if (ColorThemes.IsValidIndex(ActiveThemeIndex))
 		{
 			ThemeMenuBuilder.BeginSection("CurrentTheme", LOCTEXT("CurrentTheme", "Current Theme"));
@@ -1227,18 +1227,33 @@ TSharedRef<SWidget> SAvaLevelViewportActorColorMenu::BuildThemesMenu()
 				.MinDesiredWidth(100.f)
 				.IsReadOnly(false)
 				.ClearKeyboardFocusOnCommit(true)
-				.OnTextCommitted(FOnTextCommitted::CreateSPLambda(this,
+				.OnTextCommitted(FOnTextCommitted::CreateSPLambda(
+					this,
 					[this](const FText& NewName, ETextCommit::Type CommitType)
 					{
-						if (!ColorThemes.IsValidIndex(ActiveThemeIndex))
+						if (CommitType == ETextCommit::OnEnter)
 						{
-							return;
-						}
+							if (!ColorThemes.IsValidIndex(ActiveThemeIndex))
+							{
+								return;
+							}
 
-						ColorThemes[ActiveThemeIndex].SetName(NewName.ToString());
-						SaveColorThemesToIni();
+							const FString NewNameString = NewName.ToString();
+
+							for (int32 Index = 0; Index < ColorThemes.Num(); ++Index)
+							{
+									if (Index != ActiveThemeIndex && ColorThemes[Index].GetName().Equals(NewNameString, ESearchCase::IgnoreCase))
+								{
+									UE_LOG(AvaLevelViewportLog, Error, TEXT("Theme with that name already exists."));
+									return;
+								}
+							}
+
+							ColorThemes[ActiveThemeIndex].SetName(NewNameString);
+							SaveColorThemesToIni();
+						}
 					})
-				),				
+				),
 				LOCTEXT("ThemeNameLabel", "Name")
 			);
 
@@ -1275,10 +1290,12 @@ TSharedRef<SWidget> SAvaLevelViewportActorColorMenu::BuildThemesMenu()
 		{
 			ThemeMenuBuilder.AddSubMenu(
 				FText::FromString(ColorThemes[ThemeIdx].GetName()),
-				ActiveThemeIndex == ThemeIdx 
-					? LOCTEXT("CurrentThemeTooltip", "This is the current theme.") 
+				ActiveThemeIndex == ThemeIdx
+					? LOCTEXT("CurrentThemeTooltip", "This is the current theme.")
 					: LOCTEXT("ChangeToTheme", "Change to this theme."),
-				FNewMenuDelegate::CreateSPLambda(this, [this, ThemeIdx](FMenuBuilder& SubMenuBuilder)
+				FNewMenuDelegate::CreateSPLambda(
+					this, 
+					[this, ThemeIdx](FMenuBuilder& SubMenuBuilder)
 					{
 						SubMenuBuilder.BeginSection("ThemeOptions", LOCTEXT("Options", "Options"));
 
@@ -1329,10 +1346,10 @@ TSharedRef<SWidget> SAvaLevelViewportActorColorMenu::BuildThemesMenu()
 											SetActiveColorRGB(ColorThemes[ThemeIdx].GetColors()[ColorIdx].Color);
 
 											FSlateApplication::Get().DismissAllMenus();
-										})										
+										})
 									),
 									SNew(SHorizontalBox)
-										+SHorizontalBox::Slot()
+										+ SHorizontalBox::Slot()
 										.AutoWidth()
 										[
 											SNew(SBox)
@@ -1343,13 +1360,13 @@ TSharedRef<SWidget> SAvaLevelViewportActorColorMenu::BuildThemesMenu()
 												.Color(ColorThemes[ThemeIdx].GetColors()[ColorIdx].Color)
 											]
 										]
-										+SHorizontalBox::Slot()
+										+ SHorizontalBox::Slot()
 										.AutoWidth()
 										.VAlign(EVerticalAlignment::VAlign_Center)
 										.Padding(FMargin(5.f, 0.5f, 4.f, 0.f))
 										[
 											SNew(STextBlock)
-											.Text(ColorThemes[ThemeIdx].GetColors()[ColorIdx].Label)
+												.Text(ColorThemes[ThemeIdx].GetColors()[ColorIdx].Label)
 										],
 									FNewMenuDelegate::CreateSPLambda(this, [this, ThemeIdx, ColorIdx](FMenuBuilder& SubMenuBuilder)
 										{
@@ -1372,17 +1389,17 @@ TSharedRef<SWidget> SAvaLevelViewportActorColorMenu::BuildThemesMenu()
 
 														FSlateApplication::Get().DismissAllMenus();
 													})
-											));
+												));
 										})
-								);
+									);
 							}
 
 							SubMenuBuilder.EndSection();
 						}
-					}				
+					}
 				),
 				FUIAction(
-					FExecuteAction::CreateSPLambda(this, [this,ThemeIdx]() 
+					FExecuteAction::CreateSPLambda(this, [this, ThemeIdx]()
 						{
 							if (!ColorThemes.IsValidIndex(ThemeIdx))
 							{
@@ -1429,7 +1446,18 @@ TSharedRef<SWidget> SAvaLevelViewportActorColorMenu::BuildThemesMenu()
 				{
 					if (CommitType == ETextCommit::OnEnter)
 					{
-						ActiveThemeIndex = ColorThemes.Add(FAvaColorTheme(NewName.ToString()));
+						const FString NewNameString = NewName.ToString();
+
+						for (int32 Index = 0; Index < ColorThemes.Num(); ++Index)
+						{
+							if (ColorThemes[Index].GetName().Equals(NewNameString, ESearchCase::IgnoreCase))
+							{
+								UE_LOG(AvaLevelViewportLog, Error, TEXT("Theme with that name already exists."));
+								return;
+							}
+						}
+
+						ActiveThemeIndex = ColorThemes.Add(FAvaColorTheme(NewNameString));
 						SaveColorThemesToIni();
 						UpdateColorPalette();
 					}
@@ -1440,7 +1468,7 @@ TSharedRef<SWidget> SAvaLevelViewportActorColorMenu::BuildThemesMenu()
 
 		ThemeMenuBuilder.EndSection();
 	}
-	
+
 	return ThemeMenuBuilder.MakeWidget();
 }
 
@@ -1504,7 +1532,7 @@ void SAvaLevelViewportActorColorMenu::SetRedValue(float InNewValue, bool bInAddT
 {
 	Red = InNewValue;
 	ActiveColor.R = InNewValue;
-	
+
 	CopyRGBtoHSV();
 	ApplyActiveColor();
 	BroadcastColorChange();
@@ -1548,7 +1576,7 @@ void SAvaLevelViewportActorColorMenu::SetBlueValue(float InNewValue, bool bInAdd
 void SAvaLevelViewportActorColorMenu::SetHueValue(float InNewValue, bool bInAddToTheme)
 {
 	Hue = InNewValue;
-	
+
 	CopyHSVtoRGB();
 	ApplyActiveColor();
 	BroadcastColorChange();
@@ -1576,7 +1604,7 @@ void SAvaLevelViewportActorColorMenu::SetSaturationValue(float InNewValue, bool
 void SAvaLevelViewportActorColorMenu::SetValueValue(float InNewValue, bool InbAddToTheme)
 {
 	Value = InNewValue;
-	
+
 	CopyHSVtoRGB();
 	ApplyActiveColor();
 	BroadcastColorChange();
diff --git a/Engine/Plugins/VirtualProduction/Avalanche/Source/AvalancheLevelViewport/Private/Widgets/SAvaLevelViewportActorColorMenu.h b/Engine/Plugins/VirtualProduction/Avalanche/Source/AvalancheLevelViewport/Private/Widgets/SAvaLevelViewportActorColorMenu.h
--- a/Engine/Plugins/VirtualProduction/Avalanche/Source/AvalancheLevelViewport/Private/Widgets/SAvaLevelViewportActorColorMenu.h
+++ b/Engine/Plugins/VirtualProduction/Avalanche/Source/AvalancheLevelViewport/Private/Widgets/SAvaLevelViewportActorColorMenu.h
@@ -33,7 +33,7 @@ struct FAvaColorTheme
 {
 	FAvaColorTheme()
 		:FAvaColorTheme("Theme")
-	{		
+	{
 	}
 
 	FAvaColorTheme(const FString& InName)
@@ -117,6 +117,7 @@ private:
 	static const inline int32 MaxColorsForAutoAdd = 21;
 	static constexpr int32 ColorsPerPaletteRow = 3;
 
+	static int32 ActiveThemeIndex;
 	static TArray<FAvaColorTheme> ColorThemes;
 
 	static void LoadColorThemesFromIni();
@@ -144,7 +145,6 @@ private:
 	FLinearColor ActiveColor;
 	FLinearColor InactiveColor;
 
-	int32 ActiveThemeIndex;
 	EAvaColorStyle ColorStyle = EAvaColorStyle::Solid;
 	bool bIsUnlit = true;
 
@@ -218,7 +218,7 @@ private:
 	void CopyHSVtoRGB();
 
 	FReply OnColorSwitchedMouseDown(const FGeometry& InMyGeometry, const FPointerEvent& InPointerEvent);
-	
+
 	void OnColorStyleChanged();
 	void BroadcastColorChange();
 
