[Motion Design] Rundown Editor: Fixing Preview channel auto start error message.

The preview channel auto start is configurable (can be enable or disable) and will no longer generate an error if the preview channel is the default channel key (i.e. an actual preview channel hasn't been created).

diff --git a/Engine/Plugins/VirtualProduction/Avalanche/Source/AvalancheMedia/Private/Playback/AvaPlaybackGraph.cpp b/Engine/Plugins/VirtualProduction/Avalanche/Source/AvalancheMedia/Private/Playback/AvaPlaybackGraph.cpp
--- a/Engine/Plugins/VirtualProduction/Avalanche/Source/AvalancheMedia/Private/Playback/AvaPlaybackGraph.cpp
+++ b/Engine/Plugins/VirtualProduction/Avalanche/Source/AvalancheMedia/Private/Playback/AvaPlaybackGraph.cpp
@@ -1027,10 +1027,13 @@ int32 FAvaPlaybackGraphBuilder::GetPinIndexForChannel(const FString& InChannelNa
 	// the most sensible thing is to just connect to first index and continue
 	if (ConnectionPinIndex == INDEX_NONE)
 	{
-		UE_LOG(LogAvaPlayback, Warning
-			, TEXT("No channel with name (%s) was found. Connecting player node to first channel available")
-			, *InChannelName);
-	
+		// For the default channel key, we expect it to not have a broadcast channel, so we don't give the warning.
+		if (InChannelName != UE::AvaMedia::Playback::DefaultPreviewChannelName.ToString())
+		{
+			UE_LOG(LogAvaPlayback, Warning
+				, TEXT("No channel with name (%s) was found. Connecting player node to first channel available")
+				, *InChannelName);
+		}
 		ConnectionPinIndex = 0;
 	}
 	return ConnectionPinIndex;
diff --git a/Engine/Plugins/VirtualProduction/Avalanche/Source/AvalancheMedia/Private/Rundown/AvaRundown.cpp b/Engine/Plugins/VirtualProduction/Avalanche/Source/AvalancheMedia/Private/Rundown/AvaRundown.cpp
--- a/Engine/Plugins/VirtualProduction/Avalanche/Source/AvalancheMedia/Private/Rundown/AvaRundown.cpp
+++ b/Engine/Plugins/VirtualProduction/Avalanche/Source/AvalancheMedia/Private/Rundown/AvaRundown.cpp
@@ -1940,11 +1940,9 @@ UTextureRenderTarget2D* UAvaRundown::GetPreviewRenderTarget(const FName& InPrevi
 
 FName UAvaRundown::GetDefaultPreviewChannelName()
 {
-	// Even if the user selected preview channel is empty, we need a default
-	// name as a key for the playback manager.
-	static FName DefaultPreviewChannelFName = TEXT("_Preview");
+	// Even if the user selected preview channel is empty, we need a default name as a key for the playback manager.
 	const UAvaMediaSettings& Settings = UAvaMediaSettings::Get();
-	return !Settings.PreviewChannelName.IsEmpty() ? FName(Settings.PreviewChannelName) : DefaultPreviewChannelFName;
+	return !Settings.PreviewChannelName.IsEmpty() ? FName(Settings.PreviewChannelName) : UE::AvaMedia::Playback::DefaultPreviewChannelName;
 }
 
 void UAvaRundown::OnParentWordBeginTearDown()
diff --git a/Engine/Plugins/VirtualProduction/Avalanche/Source/AvalancheMedia/Public/Playback/AvaPlaybackDefines.h b/Engine/Plugins/VirtualProduction/Avalanche/Source/AvalancheMedia/Public/Playback/AvaPlaybackDefines.h
--- a/Engine/Plugins/VirtualProduction/Avalanche/Source/AvalancheMedia/Public/Playback/AvaPlaybackDefines.h
+++ b/Engine/Plugins/VirtualProduction/Avalanche/Source/AvalancheMedia/Public/Playback/AvaPlaybackDefines.h
@@ -34,4 +34,13 @@ enum class EAvaPlaybackAnimAction
 	Stop,
 	PreviewFrame,
 	CameraCut
-};
\ No newline at end of file
+};
+
+namespace UE::AvaMedia::Playback
+{
+	/**
+	 * Playback manager needs a channel name as a key.
+	 * The default key for a preview channel that may not have a defined broadcast channel.
+	 */
+	static inline const FLazyName DefaultPreviewChannelName = FLazyName(TEXT("_Preview"));
+}
\ No newline at end of file
diff --git a/Engine/Plugins/VirtualProduction/Avalanche/Source/AvalancheMediaEditor/Private/Rundown/AvaRundownEditor.cpp b/Engine/Plugins/VirtualProduction/Avalanche/Source/AvalancheMediaEditor/Private/Rundown/AvaRundownEditor.cpp
--- a/Engine/Plugins/VirtualProduction/Avalanche/Source/AvalancheMediaEditor/Private/Rundown/AvaRundownEditor.cpp
+++ b/Engine/Plugins/VirtualProduction/Avalanche/Source/AvalancheMediaEditor/Private/Rundown/AvaRundownEditor.cpp
@@ -1007,11 +1007,24 @@ void FAvaRundownEditor::HandleOnPagePlayerAdded(UAvaRundown* InRundown, UAvaRund
 		return;
 	}
 
-	// TODO: this should have an editor setting to enable it. "EnablePreviewAutoBroadcast" or something...
 	// Auto start preview channel.
 	if (InPagePlayer->bIsPreview)
 	{
-		UAvaBroadcast::Get().ConditionalStartBroadcastChannel(InPagePlayer->ChannelFName);
+		const UAvaRundownEditorSettings* Settings = UAvaRundownEditorSettings::Get();
+		if (Settings && Settings->bAutoStartPreviewBroadcastChannel)
+		{
+			UAvaBroadcast& Broadcast = UAvaBroadcast::Get();
+			if (Broadcast.GetCurrentProfile().GetChannel(InPagePlayer->ChannelFName).IsValidChannel())
+			{
+				Broadcast.ConditionalStartBroadcastChannel(InPagePlayer->ChannelFName);
+			}
+			else if (InPagePlayer->ChannelFName != UE::AvaMedia::Playback::DefaultPreviewChannelName)
+			{
+				UE_LOG(LogAvaBroadcast, Warning,
+					TEXT("Start Broadcast failed: Channel \"%s\" of Profile \"%s\" is not valid."),
+					*InPagePlayer->ChannelFName.ToString(), *GetCurrentProfileName().ToString());
+			}
+		}
 	}
 }
 
diff --git a/Engine/Plugins/VirtualProduction/Avalanche/Source/AvalancheMediaEditor/Private/Rundown/AvaRundownEditorSettings.h b/Engine/Plugins/VirtualProduction/Avalanche/Source/AvalancheMediaEditor/Private/Rundown/AvaRundownEditorSettings.h
--- a/Engine/Plugins/VirtualProduction/Avalanche/Source/AvalancheMediaEditor/Private/Rundown/AvaRundownEditorSettings.h
+++ b/Engine/Plugins/VirtualProduction/Avalanche/Source/AvalancheMediaEditor/Private/Rundown/AvaRundownEditorSettings.h
@@ -34,6 +34,10 @@ public:
 	UPROPERTY(Config, EditAnywhere, Category="Preview")
 	bool bPreviewCheckerBoard = false;
 
+	/** Automatically start the currently selected preview broadcast channel when a preview playback starts. */
+	UPROPERTY(Config, EditAnywhere, Category="Preview")
+	bool bAutoStartPreviewBroadcastChannel = true;
+
 	/** Whether to show controller properties in Rundown Page Details */
 	UPROPERTY(Config, EditAnywhere, Category = "Page Details", meta = (DisplayName = "Show RC Controlled Properties"))
 	bool bPageDetailsShowProperties = false;
diff --git a/Engine/Plugins/VirtualProduction/Avalanche/Source/AvalancheMediaEditor/Private/Rundown/ChannelLayerStatus/SAvaRundownChannelLayerStatusList.cpp b/Engine/Plugins/VirtualProduction/Avalanche/Source/AvalancheMediaEditor/Private/Rundown/ChannelLayerStatus/SAvaRundownChannelLayerStatusList.cpp
--- a/Engine/Plugins/VirtualProduction/Avalanche/Source/AvalancheMediaEditor/Private/Rundown/ChannelLayerStatus/SAvaRundownChannelLayerStatusList.cpp
+++ b/Engine/Plugins/VirtualProduction/Avalanche/Source/AvalancheMediaEditor/Private/Rundown/ChannelLayerStatus/SAvaRundownChannelLayerStatusList.cpp
@@ -262,7 +262,7 @@ void SAvaRundownChannelLayerStatusList::RefreshList()
 		constexpr int32 ColumnPageId = 2;
 		constexpr int32 ColumnCount = 3;
 
-		static const FName PreviewChannelName = "_Preview";
+		const FName PreviewChannelName =  UE::AvaMedia::Playback::DefaultPreviewChannelName;
 		static const FText PreviewChannelDescription = LOCTEXT("Preview", "Preview");
 
 		int32 RowIndex = 0;
