MotionDesign : Text3D
- Fix robomerge conflict coming from 5.6.1 to 5.7

diff --git a/Engine/Plugins/Experimental/Text3D/Source/Text3D/Private/GeometryBuilders/Text3DGlyphLoader.cpp b/Engine/Plugins/Experimental/Text3D/Source/Text3D/Private/GeometryBuilders/Text3DGlyphLoader.cpp
--- a/Engine/Plugins/Experimental/Text3D/Source/Text3D/Private/GeometryBuilders/Text3DGlyphLoader.cpp
+++ b/Engine/Plugins/Experimental/Text3D/Source/Text3D/Private/GeometryBuilders/Text3DGlyphLoader.cpp
@@ -432,7 +432,7 @@ TSharedPtr<FPolygon2f> FText3DGlyphLoader::ProcessFreetypeOutline(const FT_Outli
 
 bool FText3DGlyphLoader::IsInside(const TSharedPtr<const FPolygon2f> ContourA, const TSharedPtr<const FPolygon2f> ContourB) const
 {
-	return ContourB->Contains(*ContourA) || ContourB->Bounds().Contains(ContourA->Bounds());
+	return ContourB->Contains(*ContourA);
 }
 
 void FText3DGlyphLoader::Insert(const TText3DGlyphContourNodeShared NodeA, const TText3DGlyphContourNodeShared NodeB) const
diff --git a/Engine/Plugins/VirtualProduction/Avalanche/Source/AvalancheText/Private/AvaText3DComponent.cpp b/Engine/Plugins/VirtualProduction/Avalanche/Source/AvalancheText/Private/AvaText3DComponent.cpp
--- a/Engine/Plugins/VirtualProduction/Avalanche/Source/AvalancheText/Private/AvaText3DComponent.cpp
+++ b/Engine/Plugins/VirtualProduction/Avalanche/Source/AvalancheText/Private/AvaText3DComponent.cpp
@@ -11,6 +11,7 @@
 #include "GameFramework/Actor.h"
 #include "Materials/Material.h"
 #include "Materials/MaterialInstanceDynamic.h"
+#include "Settings/Text3DProjectSettings.h"
 #include "Text3DComponentVersion.h"
 
 #if WITH_EDITOR
@@ -41,6 +42,11 @@ void UAvaText3DComponent::Serialize(FArchive& InArchive)
 #if WITH_EDITORONLY_DATA
 			PRAGMA_DISABLE_DEPRECATION_WARNINGS
 
+			if (const UText3DProjectSettings* TextSettings = UText3DProjectSettings::Get())
+			{
+				SetFont(TextSettings->GetFallbackFont());
+			}
+
 #if WITH_EDITOR
 			SetFontByName(MotionDesignFont_DEPRECATED.GetFontNameAsString());
 #else
