AvalancheMedia: Support RHI immediate captures.

This is how nDisplay captures media, for example.

#rnx

diff --git a/Engine/Plugins/VirtualProduction/Avalanche/Source/AvalancheMedia/Private/Broadcast/OutputDevices/AvaBroadcastRenderTargetMediaCapture.cpp b/Engine/Plugins/VirtualProduction/Avalanche/Source/AvalancheMedia/Private/Broadcast/OutputDevices/AvaBroadcastRenderTargetMediaCapture.cpp
--- a/Engine/Plugins/VirtualProduction/Avalanche/Source/AvalancheMedia/Private/Broadcast/OutputDevices/AvaBroadcastRenderTargetMediaCapture.cpp
+++ b/Engine/Plugins/VirtualProduction/Avalanche/Source/AvalancheMedia/Private/Broadcast/OutputDevices/AvaBroadcastRenderTargetMediaCapture.cpp
@@ -82,6 +82,11 @@ bool UAvaBroadcastRenderTargetMediaCapture::PostInitializeCaptureRenderTarget(UT
 	return StartNewCapture(FIntPoint(InRenderTarget->SizeX, InRenderTarget->SizeY), InRenderTarget->GetFormat());
 }
 
+bool UAvaBroadcastRenderTargetMediaCapture::PostInitializeCaptureRHIResource(const FRHICaptureResourceDescription& InResourceDescription)
+{
+	return StartNewCapture(InResourceDescription.ResourceSize, InResourceDescription.PixelFormat);
+}
+
 void UAvaBroadcastRenderTargetMediaCapture::StopCaptureImpl(bool /*bAllowPendingFrameToBeProcess*/)
 {
 	TRACE_BOOKMARK(TEXT("UAvaBroadcastRenderTargetMediaCapture::StopCapture"));
diff --git a/Engine/Plugins/VirtualProduction/Avalanche/Source/AvalancheMedia/Private/Broadcast/OutputDevices/AvaBroadcastRenderTargetMediaCapture.h b/Engine/Plugins/VirtualProduction/Avalanche/Source/AvalancheMedia/Private/Broadcast/OutputDevices/AvaBroadcastRenderTargetMediaCapture.h
--- a/Engine/Plugins/VirtualProduction/Avalanche/Source/AvalancheMedia/Private/Broadcast/OutputDevices/AvaBroadcastRenderTargetMediaCapture.h
+++ b/Engine/Plugins/VirtualProduction/Avalanche/Source/AvalancheMedia/Private/Broadcast/OutputDevices/AvaBroadcastRenderTargetMediaCapture.h
@@ -20,6 +20,7 @@ protected:
 	virtual bool InitializeCapture() override;
 	virtual bool PostInitializeCaptureViewport(TSharedPtr<FSceneViewport>& InSceneViewport) override;
 	virtual bool PostInitializeCaptureRenderTarget(UTextureRenderTarget2D* InRenderTarget) override;
+	virtual bool PostInitializeCaptureRHIResource(const FRHICaptureResourceDescription& InResourceDescription) override;
 	virtual void StopCaptureImpl(bool bAllowPendingFrameToBeProcess) override;
 
 	UPROPERTY(Transient)
