MotionDesign : Modifiers
- Adjust default for spline sweep modifier "steps" property to 64 instead of 10

diff --git a/Engine/Plugins/VirtualProduction/Avalanche/Source/AvalancheModifiers/Public/Modifiers/AvaSplineSweepModifier.h b/Engine/Plugins/VirtualProduction/Avalanche/Source/AvalancheModifiers/Public/Modifiers/AvaSplineSweepModifier.h
--- a/Engine/Plugins/VirtualProduction/Avalanche/Source/AvalancheModifiers/Public/Modifiers/AvaSplineSweepModifier.h
+++ b/Engine/Plugins/VirtualProduction/Avalanche/Source/AvalancheModifiers/Public/Modifiers/AvaSplineSweepModifier.h
@@ -172,7 +172,7 @@ protected:
 
 	/** The sample count defines the precision of the sweep */
 	UPROPERTY(EditInstanceOnly, Setter="SetSteps", Getter="GetSteps", Category="SplineSweep", meta=(ClampMin="3", ClampMax="200", AllowPrivateAccess="true"))
-	int32 Steps = 10;
+	int32 Steps = 64;
 
 	/** Sample range offset of the spline, for closed loop spline range can wrap around */
 	UPROPERTY(EditInstanceOnly, Setter="SetProgressOffset", Getter="GetProgressOffset", Category="SplineSweep", meta=(AllowPrivateAccess="true"))
