[MotionDesign] Remove unneeded instances of IsValid() in favor of default null checks

diff --git a/Engine/Plugins/VirtualProduction/Avalanche/Source/AvalancheEditor/Private/DetailView/Customizations/AvaViewportQualitySettingsPropertyTypeCustomization.cpp b/Engine/Plugins/VirtualProduction/Avalanche/Source/AvalancheEditor/Private/DetailView/Customizations/AvaViewportQualitySettingsPropertyTypeCustomization.cpp
--- a/Engine/Plugins/VirtualProduction/Avalanche/Source/AvalancheEditor/Private/DetailView/Customizations/AvaViewportQualitySettingsPropertyTypeCustomization.cpp
+++ b/Engine/Plugins/VirtualProduction/Avalanche/Source/AvalancheEditor/Private/DetailView/Customizations/AvaViewportQualitySettingsPropertyTypeCustomization.cpp
@@ -289,8 +289,7 @@ void FAvaViewportQualitySettingsPropertyTypeCustomization::OpenEditorSettings()
 {
 	if (ISettingsModule* const SettingsModule = FModuleManager::GetModulePtr<ISettingsModule>(TEXT("Settings")))
 	{
-		const UAvaEditorSettings* const Settings = GetDefault<UAvaEditorSettings>();
-		if (ensure(IsValid(Settings)))
+		if (const UAvaEditorSettings* const Settings = GetDefault<UAvaEditorSettings>())
 		{
 			SettingsModule->ShowViewer(Settings->GetContainerName(), Settings->GetCategoryName(), Settings->GetSectionName());
 		}
diff --git a/Engine/Plugins/VirtualProduction/Avalanche/Source/AvalancheEditor/Private/Scene/AvaSceneRigExtension.cpp b/Engine/Plugins/VirtualProduction/Avalanche/Source/AvalancheEditor/Private/Scene/AvaSceneRigExtension.cpp
--- a/Engine/Plugins/VirtualProduction/Avalanche/Source/AvalancheEditor/Private/Scene/AvaSceneRigExtension.cpp
+++ b/Engine/Plugins/VirtualProduction/Avalanche/Source/AvalancheEditor/Private/Scene/AvaSceneRigExtension.cpp
@@ -66,13 +66,13 @@ void FAvaSceneRigExtension::BindCommands(const TSharedRef<FUICommandList>& InCom
 
 void FAvaSceneRigExtension::ExtendOutlinerItemContextMenu(UToolMenu* const InToolMenu)
 {
-	if (!IsValid(InToolMenu))
+	if (!InToolMenu)
 	{
 		return;
 	}
 
 	UAvaOutlinerItemsContext* const ItemsContext = InToolMenu->FindContext<UAvaOutlinerItemsContext>();
-	if (!IsValid(ItemsContext))
+	if (!ItemsContext)
 	{
 		return;
 	}
@@ -111,19 +111,15 @@ void FAvaSceneRigExtension::CreateSubMenu(UToolMenu* const InToolMenu)
 
 			FToolMenuSection& NewSection = InToolMenu->AddSection(TEXT("SceneRig"), LOCTEXT("SceneRigActions", "Scene Rig"));
 
-			UWorld* const World = GetWorld();
-			if (IsValid(World))
+			if (UWorld* const World = GetWorld())
 			{
-				UAvaSceneRigSubsystem* const SceneRigSubsystem = UAvaSceneRigSubsystem::ForWorld(World);
-				if (IsValid(SceneRigSubsystem))
+				if (UAvaSceneRigSubsystem* const SceneRigSubsystem = UAvaSceneRigSubsystem::ForWorld(World))
 				{
-					const ULevelStreaming* const StreamingLevel = SceneRigSubsystem->FindFirstActiveSceneRig();
-					if (IsValid(StreamingLevel))
+					if (const ULevelStreaming* const StreamingLevel = SceneRigSubsystem->FindFirstActiveSceneRig())
 					{
 						TArray<AActor*> SelectedActors;
 
-						const UAvaOutlinerItemsContext* const ItemsContext = ItemsContextWeak.Get();
-						if (IsValid(ItemsContext))
+						if (const UAvaOutlinerItemsContext* const ItemsContext = ItemsContextWeak.Get())
 						{
 							const TConstArrayView<FAvaOutlinerItemWeakPtr> SelectedItems = ItemsContext->GetItems();
 							SelectedActors = OutlinerItemsToActors(ItemsContext->GetItems(), true);
@@ -165,13 +161,13 @@ void FAvaSceneRigExtension::PromptToSaveSceneRigFromOutlinerItems()
 	}
 
 	const UAvaOutlinerItemsContext* const ItemsContext = ItemsContextWeak.Get();
-	if (!IsValid(ItemsContext))
+	if (!ItemsContext)
 	{
 		return;
 	}
 
 	UWorld* const World = GetWorld();
-	if (!IsValid(World))
+	if (!World)
 	{
 		return;
 	}
@@ -187,7 +183,7 @@ void FAvaSceneRigExtension::PromptToSaveSceneRigFromOutlinerItems()
 	const FSoftObjectPath NewAssetPath = SceneRigEditorModule.CreateSceneRigAssetWithDialog();
 	if (NewAssetPath.IsValid())
 	{
-		if (IsValid(SceneRigEditorModule.SetActiveSceneRig(World, NewAssetPath)))
+		if (SceneRigEditorModule.SetActiveSceneRig(World, NewAssetPath))
 		{
 			SceneRigEditorModule.AddActiveSceneRigActors(World, SelectedActors);
 		}
@@ -214,7 +210,7 @@ void FAvaSceneRigExtension::AddOutlinerItemsToSceneRig()
 	}
 
 	const UAvaOutlinerItemsContext* const ItemsContext = ItemsContextWeak.Get();
-	if (!IsValid(ItemsContext))
+	if (!ItemsContext)
 	{
 		return;
 	}
@@ -226,22 +222,20 @@ void FAvaSceneRigExtension::AddOutlinerItemsToSceneRig()
 	}
 
 	UWorld* const World = GetWorld();
-	if (!IsValid(World))
+	if (!World)
 	{
 		return;
 	}
 
 	UAvaSceneRigSubsystem* const SceneRigSubsystem = UAvaSceneRigSubsystem::ForWorld(World);
-	if (!IsValid(SceneRigSubsystem))
+	if (!SceneRigSubsystem)
 	{
 		return;
 	}
 
 	IAvaSceneRigEditorModule& SceneRigEditorModule = IAvaSceneRigEditorModule::Get();
 
-	ULevelStreaming* const SceneRig = SceneRigSubsystem->FindFirstActiveSceneRig();
-
-	if (IsValid(SceneRig))
+	if (ULevelStreaming* const SceneRig = SceneRigSubsystem->FindFirstActiveSceneRig())
 	{
 		FScopedTransaction Transaction(LOCTEXT("AddSceneRig", "Add Scene Rig Actors"));
 		SceneRig->Modify();
@@ -253,7 +247,7 @@ void FAvaSceneRigExtension::AddOutlinerItemsToSceneRig()
 		const FSoftObjectPath NewAssetPath = SceneRigEditorModule.CreateSceneRigAssetWithDialog();
 		if (NewAssetPath.IsValid())
 		{
-			if (IsValid(SceneRigEditorModule.SetActiveSceneRig(World, NewAssetPath)))
+			if (SceneRigEditorModule.SetActiveSceneRig(World, NewAssetPath))
 			{
 				SceneRigEditorModule.AddActiveSceneRigActors(World, SelectedActors);
 			}
@@ -264,7 +258,7 @@ void FAvaSceneRigExtension::AddOutlinerItemsToSceneRig()
 bool FAvaSceneRigExtension::CanAddOutlinerItemsToSceneRig() const
 {
 	const UAvaOutlinerItemsContext* const ItemsContext = ItemsContextWeak.Get();
-	if (!IsValid(ItemsContext))
+	if (!ItemsContext)
 	{
 		return false;
 	}
@@ -276,13 +270,13 @@ bool FAvaSceneRigExtension::CanAddOutlinerItemsToSceneRig() const
 	}
 
 	const UAvaSceneRigSubsystem* const SceneRigSubsystem = UAvaSceneRigSubsystem::ForWorld(GetWorld());
-	if (!IsValid(SceneRigSubsystem))
+	if (!SceneRigSubsystem)
 	{
 		return false;
 	}
 
 	UWorld* const SceneRigAsset = SceneRigSubsystem->FindFirstActiveSceneRigAsset();
-	if (!IsValid(SceneRigAsset) || !IsValid(SceneRigAsset->PersistentLevel))
+	if (!SceneRigAsset || !SceneRigAsset->PersistentLevel)
 	{
 		return false;
 	}
@@ -298,25 +292,25 @@ void FAvaSceneRigExtension::RemoveOutlinerItemsFromSceneRig()
 	}
 
 	const UAvaOutlinerItemsContext* const ItemsContext = ItemsContextWeak.Get();
-	if (!IsValid(ItemsContext))
+	if (!ItemsContext)
 	{
 		return;
 	}
 
 	UWorld* const World = GetWorld();
-	if (!IsValid(World))
+	if (!World)
 	{
 		return;
 	}
 	
 	UAvaSceneRigSubsystem* const SceneRigSubsystem = UAvaSceneRigSubsystem::ForWorld(World);
-	if (!IsValid(SceneRigSubsystem))
+	if (!SceneRigSubsystem)
 	{
 		return;
 	}
 
 	ULevelStreaming* const SceneRig = SceneRigSubsystem->FindFirstActiveSceneRig();
-	if (!IsValid(SceneRig))
+	if (!SceneRig)
 	{
 		return;
 	}
@@ -332,7 +326,7 @@ void FAvaSceneRigExtension::RemoveOutlinerItemsFromSceneRig()
 bool FAvaSceneRigExtension::CanRemoveOutlinerItemsFromSceneRig() const
 {
 	const UAvaOutlinerItemsContext* const ItemsContext = ItemsContextWeak.Get();
-	if (!IsValid(ItemsContext))
+	if (!ItemsContext)
 	{
 		return false;
 	}
@@ -345,13 +339,13 @@ bool FAvaSceneRigExtension::CanRemoveOutlinerItemsFromSceneRig() const
 	}
 
 	const UAvaSceneRigSubsystem* const SceneRigSubsystem = UAvaSceneRigSubsystem::ForWorld(GetWorld());
-	if (!IsValid(SceneRigSubsystem))
+	if (!SceneRigSubsystem)
 	{
 		return false;
 	}
 
 	UWorld* const SceneRigAsset = SceneRigSubsystem->FindFirstActiveSceneRigAsset();
-	if (!IsValid(SceneRigAsset) || !IsValid(SceneRigAsset->PersistentLevel))
+	if (!SceneRigAsset || !SceneRigAsset->PersistentLevel)
 	{
 		return false;
 	}
@@ -372,8 +366,7 @@ TArray<AActor*> FAvaSceneRigExtension::OutlinerItemsToActors(const TConstArrayVi
 			{
 				if (bInIncludeLocked || !ActorItem->IsLocked())
 				{
-					AActor* const Actor = ActorItem->GetActor();
-					if (IsValid(Actor))
+					if (AActor* const Actor = ActorItem->GetActor())
 					{
 						OutActors.Add(Actor);
 					}
@@ -387,13 +380,13 @@ TArray<AActor*> FAvaSceneRigExtension::OutlinerItemsToActors(const TConstArrayVi
 
 UAvaSceneState* FAvaSceneRigExtension::GetSceneState(UWorld* const InWorld)
 {
-	if (!IsValid(InWorld))
+	if (!InWorld)
 	{
 		return nullptr;
 	}
 
 	const UAvaSceneSubsystem* const SceneSubsystem = InWorld->GetSubsystem<UAvaSceneSubsystem>();
-	if (!IsValid(SceneSubsystem))
+	if (!SceneSubsystem)
 	{
 		return nullptr;
 	}
diff --git a/Engine/Plugins/VirtualProduction/Avalanche/Source/AvalancheEditor/Private/Sequencer/AvaSequencerExtension.cpp b/Engine/Plugins/VirtualProduction/Avalanche/Source/AvalancheEditor/Private/Sequencer/AvaSequencerExtension.cpp
--- a/Engine/Plugins/VirtualProduction/Avalanche/Source/AvalancheEditor/Private/Sequencer/AvaSequencerExtension.cpp
+++ b/Engine/Plugins/VirtualProduction/Avalanche/Source/AvalancheEditor/Private/Sequencer/AvaSequencerExtension.cpp
@@ -41,7 +41,7 @@ TSharedPtr<ISequencer> FAvaSequencerExtension::GetSequencer() const
 
 void FAvaSequencerExtension::OnObjectRenamed(UObject* InRenamedObject, const FText& InDisplayNameText)
 {
-	if (!IsValid(InRenamedObject) || !AvaSequencer.IsValid())
+	if (!InRenamedObject || !AvaSequencer.IsValid())
 	{
 		return;
 	}
diff --git a/Engine/Plugins/VirtualProduction/Avalanche/Source/AvalancheModifiersEditor/Private/Outliner/AvaOutlinerModifierProxy.cpp b/Engine/Plugins/VirtualProduction/Avalanche/Source/AvalancheModifiersEditor/Private/Outliner/AvaOutlinerModifierProxy.cpp
--- a/Engine/Plugins/VirtualProduction/Avalanche/Source/AvalancheModifiersEditor/Private/Outliner/AvaOutlinerModifierProxy.cpp
+++ b/Engine/Plugins/VirtualProduction/Avalanche/Source/AvalancheModifiersEditor/Private/Outliner/AvaOutlinerModifierProxy.cpp
@@ -90,7 +90,7 @@ FText FAvaOutlinerModifierProxy::GetIconTooltipText() const
 
 bool FAvaOutlinerModifierProxy::CanDelete() const
 {
-	return IsValid(GetModifierStack());
+	return GetModifierStack() != nullptr;
 }
 
 bool FAvaOutlinerModifierProxy::Delete()
@@ -98,7 +98,7 @@ bool FAvaOutlinerModifierProxy::Delete()
 	UActorModifierCoreStack* ModifierStack = GetModifierStack();
 	const UActorModifierCoreSubsystem* ModifierSubsystem = UActorModifierCoreSubsystem::Get();
 
-	if (IsValid(ModifierStack) && ModifierSubsystem)
+	if (ModifierStack && ModifierSubsystem)
 	{
 		FText FailReason;
 		FActorModifierCoreStackRemoveOp RemoveOp;
diff --git a/Engine/Plugins/VirtualProduction/Avalanche/Source/AvalancheOutliner/Private/AvaOutliner.cpp b/Engine/Plugins/VirtualProduction/Avalanche/Source/AvalancheOutliner/Private/AvaOutliner.cpp
--- a/Engine/Plugins/VirtualProduction/Avalanche/Source/AvalancheOutliner/Private/AvaOutliner.cpp
+++ b/Engine/Plugins/VirtualProduction/Avalanche/Source/AvalancheOutliner/Private/AvaOutliner.cpp
@@ -102,13 +102,13 @@ UAvaOutlinerSubsystem* FAvaOutliner::GetOutlinerSubsystem() const
 
 bool FAvaOutliner::IsActorAllowedInOutliner(const AActor* InActor) const
 {
-	if (!IsValid(InActor))
+	if (!InActor)
 	{
 		return false;
 	}
 
 	const UWorld* const ActorWorld = InActor->GetWorld();
-	if (!IsValid(ActorWorld))
+	if (!ActorWorld)
 	{
 		return false;
 	}
diff --git a/Engine/Plugins/VirtualProduction/Avalanche/Source/AvalancheOutliner/Private/Item/AvaOutlinerActor.cpp b/Engine/Plugins/VirtualProduction/Avalanche/Source/AvalancheOutliner/Private/Item/AvaOutlinerActor.cpp
--- a/Engine/Plugins/VirtualProduction/Avalanche/Source/AvalancheOutliner/Private/Item/AvaOutlinerActor.cpp
+++ b/Engine/Plugins/VirtualProduction/Avalanche/Source/AvalancheOutliner/Private/Item/AvaOutlinerActor.cpp
@@ -28,7 +28,7 @@ void FAvaOutlinerActor::FindChildren(TArray<FAvaOutlinerItemPtr>& OutChildren, b
 	Super::FindChildren(OutChildren, bRecursive);
 
 	AActor* const UnderlyingActor = GetActor();
-	if (!IsValid(UnderlyingActor))
+	if (!UnderlyingActor)
 	{
 		return;
 	}
diff --git a/Engine/Plugins/VirtualProduction/Avalanche/Source/AvalancheSceneRigEditor/Private/Outliner/AvaOutlinerSceneRigProxy.cpp b/Engine/Plugins/VirtualProduction/Avalanche/Source/AvalancheSceneRigEditor/Private/Outliner/AvaOutlinerSceneRigProxy.cpp
--- a/Engine/Plugins/VirtualProduction/Avalanche/Source/AvalancheSceneRigEditor/Private/Outliner/AvaOutlinerSceneRigProxy.cpp
+++ b/Engine/Plugins/VirtualProduction/Avalanche/Source/AvalancheSceneRigEditor/Private/Outliner/AvaOutlinerSceneRigProxy.cpp
@@ -22,8 +22,7 @@ FAvaOutlinerSceneRigProxy::FAvaOutlinerSceneRigProxy(IAvaOutliner& InOutliner, c
 {
 	Icon = FSlateIcon(FAppStyle::GetAppStyleSetName(), TEXT("LandscapeEditor.NoiseTool"));
 
-	UAvaSceneRigSubsystem* const SceneRigSubsystem = UAvaSceneRigSubsystem::ForWorld(Outliner.GetWorld());
-	if (IsValid(SceneRigSubsystem))
+	if (UAvaSceneRigSubsystem* const SceneRigSubsystem = UAvaSceneRigSubsystem::ForWorld(Outliner.GetWorld()))
 	{
 		StreamingLevelWeak = SceneRigSubsystem->FindFirstActiveSceneRig();
 	}
@@ -67,14 +66,12 @@ void FAvaOutlinerSceneRigProxy::GetProxiedItems(const TSharedRef<IAvaOutlinerIte
 {
 	if (const FAvaOutlinerActor* const ActorItem = InParent->CastTo<FAvaOutlinerActor>())
 	{
-		AActor* const Actor = ActorItem->GetActor();
-		if (IsValid(Actor))
+		if (AActor* const Actor = ActorItem->GetActor())
 		{
-			ULevelStreaming* const ActiveSceneRig = StreamingLevelWeak.Get();
-			if (IsValid(ActiveSceneRig))
+			if (ULevelStreaming* const ActiveSceneRig = StreamingLevelWeak.Get())
 			{
 				const UWorld* const World = ActiveSceneRig->GetWorldAsset().Get();
-				if (IsValid(World) && IsValid(World->PersistentLevel))
+				if (World && World->PersistentLevel)
 				{
 					if (World->PersistentLevel->Actors.Contains(Actor))
 					{
@@ -154,13 +151,13 @@ ULevelStreaming* FAvaOutlinerSceneRigProxy::GetSceneRigAsset() const
 	}
 
 	UAvaSceneRigSubsystem* const SceneRigSubsystem = UAvaSceneRigSubsystem::ForWorld(Outliner.GetWorld());
-	if (!IsValid(SceneRigSubsystem))
+	if (!SceneRigSubsystem)
 	{
 		return nullptr;
 	}
 
 	ULevelStreaming* const StreamingLevel = SceneRigSubsystem->SceneRigFromActor(ActorItem->GetActor());
-	if (!IsValid(StreamingLevel))
+	if (!StreamingLevel)
 	{
 		return nullptr;
 	}
diff --git a/Engine/Plugins/VirtualProduction/Avalanche/Source/AvalancheSequencer/Private/AvaSequencerUtils.cpp b/Engine/Plugins/VirtualProduction/Avalanche/Source/AvalancheSequencer/Private/AvaSequencerUtils.cpp
--- a/Engine/Plugins/VirtualProduction/Avalanche/Source/AvalancheSequencer/Private/AvaSequencerUtils.cpp
+++ b/Engine/Plugins/VirtualProduction/Avalanche/Source/AvalancheSequencer/Private/AvaSequencerUtils.cpp
@@ -26,8 +26,7 @@ UWorld* FAvaSequencerUtils::GetSequencerWorld(const TSharedRef<ISequencer>& InSe
 
 UAvaSequencerSubsystem* FAvaSequencerUtils::GetSequencerSubsystem(const TSharedRef<ISequencer>& InSequencer)
 {
-	const UWorld* const World = GetSequencerWorld(InSequencer);
-	if (IsValid(World))
+	if (const UWorld* const World = GetSequencerWorld(InSequencer))
 	{
 		return World->GetSubsystem<UAvaSequencerSubsystem>();
 	}
@@ -36,8 +35,7 @@ UAvaSequencerSubsystem* FAvaSequencerUtils::GetSequencerSubsystem(const TSharedR
 
 UAvaSceneSubsystem* FAvaSequencerUtils::GetSceneSubsystem(const TSharedRef<ISequencer>& InSequencer)
 {
-	const UWorld* const World = GetSequencerWorld(InSequencer);
-	if (IsValid(World))
+	if (const UWorld* const World = GetSequencerWorld(InSequencer))
 	{
 		return World->GetSubsystem<UAvaSceneSubsystem>();
 	}
@@ -46,8 +44,7 @@ UAvaSceneSubsystem* FAvaSequencerUtils::GetSceneSubsystem(const TSharedRef<ISequ
 
 IAvaSceneInterface* FAvaSequencerUtils::GetSceneInterface(const TSharedRef<ISequencer>& InSequencer)
 {
-	UAvaSceneSubsystem* const SceneSubsystem = GetSceneSubsystem(InSequencer);
-	if (IsValid(SceneSubsystem))
+	if (UAvaSceneSubsystem* const SceneSubsystem = GetSceneSubsystem(InSequencer))
 	{
 		return SceneSubsystem->GetSceneInterface();
 	}
@@ -65,8 +62,7 @@ IAvaSequenceProvider* FAvaSequencerUtils::GetSequenceProvider(const TSharedRef<I
 
 TSharedPtr<IAvaSequencer> FAvaSequencerUtils::GetAvaSequencer(const TSharedRef<ISequencer>& InSequencer)
 {
-	UAvaSequencerSubsystem* const SequencerSubsystem = GetSequencerSubsystem(InSequencer);
-	if (IsValid(SequencerSubsystem))
+	if (UAvaSequencerSubsystem* const SequencerSubsystem = GetSequencerSubsystem(InSequencer))
 	{
 		return SequencerSubsystem->GetSequencer();
 	}
diff --git a/Engine/Plugins/VirtualProduction/Avalanche/Source/AvalancheSequencer/Private/Outliner/AvaOutlinerSequenceProxy.cpp b/Engine/Plugins/VirtualProduction/Avalanche/Source/AvalancheSequencer/Private/Outliner/AvaOutlinerSequenceProxy.cpp
--- a/Engine/Plugins/VirtualProduction/Avalanche/Source/AvalancheSequencer/Private/Outliner/AvaOutlinerSequenceProxy.cpp
+++ b/Engine/Plugins/VirtualProduction/Avalanche/Source/AvalancheSequencer/Private/Outliner/AvaOutlinerSequenceProxy.cpp
@@ -103,8 +103,8 @@ void FAvaOutlinerSequenceProxy::GetProxiedItems(const TSharedRef<IAvaOutlinerIte
 		TSharedPtr<IAvaSequencer> AvaSequencer = AvaSequencerWeak.Pin();
 
 		AActor* const Actor = ActorItem->GetActor();
-		
-		if (AvaSequencer.IsValid() && IsValid(Actor))
+
+		if (AvaSequencer.IsValid() && Actor)
 		{
 			const TArray<UAvaSequence*> Sequences = AvaSequencer->GetSequencesForObject(Actor);
 			
diff --git a/Engine/Plugins/VirtualProduction/Avalanche/Source/AvalancheSequencer/Private/StaggerTool/AvaStaggerTool.cpp b/Engine/Plugins/VirtualProduction/Avalanche/Source/AvalancheSequencer/Private/StaggerTool/AvaStaggerTool.cpp
--- a/Engine/Plugins/VirtualProduction/Avalanche/Source/AvalancheSequencer/Private/StaggerTool/AvaStaggerTool.cpp
+++ b/Engine/Plugins/VirtualProduction/Avalanche/Source/AvalancheSequencer/Private/StaggerTool/AvaStaggerTool.cpp
@@ -442,7 +442,7 @@ TArray<FAvaStaggerTool::FAvaStaggerKeyElement> FAvaStaggerTool::GatherSelectionK
 			if (const TViewModelPtr<FChannelModel> KeyModel = SequencerSelection->KeySelection.GetModelForKey(KeyHandle))
 			{
 				const UMovieSceneSection* const KeySection = KeyModel->GetSection();
-				if (IsValid(KeySection) && !KeySection->IsReadOnly())
+				if (KeySection && !KeySection->IsReadOnly())
 				{
 					FFrameNumber KeyFrameNumber;
 					KeyModel->GetChannel()->GetKeyTime(KeyHandle, KeyFrameNumber);
