Motion Design: fix issue with removed sequence causing world gc leak

diff --git a/Engine/Plugins/VirtualProduction/Avalanche/Source/Avalanche/Private/AvaScene.cpp b/Engine/Plugins/VirtualProduction/Avalanche/Source/Avalanche/Private/AvaScene.cpp
--- a/Engine/Plugins/VirtualProduction/Avalanche/Source/Avalanche/Private/AvaScene.cpp
+++ b/Engine/Plugins/VirtualProduction/Avalanche/Source/Avalanche/Private/AvaScene.cpp
@@ -206,6 +206,11 @@ bool AAvaScene::AddSequence(UAvaSequence* InSequence)
 
 void AAvaScene::RemoveSequence(UAvaSequence* InSequence)
 {
+	if (InSequence)
+	{
+		InSequence->Cleanup();
+	}
+
 	Animations.Remove(InSequence);
 	ScheduleRebuildSequenceTree();
 }
diff --git a/Engine/Plugins/VirtualProduction/Avalanche/Source/AvalancheEditor/Private/Sequencer/AvaSequencerExtension.cpp b/Engine/Plugins/VirtualProduction/Avalanche/Source/AvalancheEditor/Private/Sequencer/AvaSequencerExtension.cpp
--- a/Engine/Plugins/VirtualProduction/Avalanche/Source/AvalancheEditor/Private/Sequencer/AvaSequencerExtension.cpp
+++ b/Engine/Plugins/VirtualProduction/Avalanche/Source/AvalancheEditor/Private/Sequencer/AvaSequencerExtension.cpp
@@ -4,7 +4,6 @@
 #include "AvaEditorModule.h"
 #include "AvaSequencerSubsystem.h"
 #include "AvaSequencerTabSpawner.h"
-#include "BlueprintActionDatabase.h"
 #include "EditorModeManager.h"
 #include "Framework/Docking/LayoutExtender.h"
 #include "IAssetTools.h"
@@ -98,19 +97,12 @@ void FAvaSequencerExtension::Cleanup()
 		return;
 	}
 
-	FBlueprintActionDatabase& BlueprintActionDatabase = FBlueprintActionDatabase::Get();
-
 	// Clear the Asset Actions for each Director Blueprint so that the World gets properly GC'd
 	for (UAvaSequence* Sequence : SequenceProvider->GetSequences())
 	{
-		if (!Sequence)
-		{
-			continue;;
-		}
-
-		if (UBlueprint* DirectorBlueprint = Sequence->GetDirectorBlueprint())
+		if (Sequence)
 		{
-			BlueprintActionDatabase.ClearAssetActions(DirectorBlueprint);
+			Sequence->Cleanup();
 		}
 	}
 }
diff --git a/Engine/Plugins/VirtualProduction/Avalanche/Source/AvalancheSequence/AvalancheSequence.Build.cs b/Engine/Plugins/VirtualProduction/Avalanche/Source/AvalancheSequence/AvalancheSequence.Build.cs
--- a/Engine/Plugins/VirtualProduction/Avalanche/Source/AvalancheSequence/AvalancheSequence.Build.cs
+++ b/Engine/Plugins/VirtualProduction/Avalanche/Source/AvalancheSequence/AvalancheSequence.Build.cs
@@ -36,6 +36,7 @@ public class AvalancheSequence : ModuleRules
             PrivateDependencyModuleNames.AddRange(
                 new string[]
                 {
+                    "BlueprintGraph",
                     "Sequencer",
                 });
         }
diff --git a/Engine/Plugins/VirtualProduction/Avalanche/Source/AvalancheSequence/Private/AvaSequence.cpp b/Engine/Plugins/VirtualProduction/Avalanche/Source/AvalancheSequence/Private/AvaSequence.cpp
--- a/Engine/Plugins/VirtualProduction/Avalanche/Source/AvalancheSequence/Private/AvaSequence.cpp
+++ b/Engine/Plugins/VirtualProduction/Avalanche/Source/AvalancheSequence/Private/AvaSequence.cpp
@@ -29,6 +29,11 @@
 #include "UObject/CoreRedirects.h"
 #include "UObject/Package.h"
 
+#if WITH_EDITOR
+#include "BlueprintActionDatabase.h"
+#include "Engine/Blueprint.h"
+#endif
+
 #define LOCTEXT_NAMESPACE "AvaSequence"
 
 #if WITH_EDITOR
@@ -660,6 +665,17 @@ void UAvaSequence::OnOuterWorldRenamed(UObject* InNewOuter)
 }
 #endif
 
+void UAvaSequence::Cleanup()
+{
+#if WITH_EDITOR
+	if (DirectorBlueprint)
+	{
+		FBlueprintActionDatabase& BlueprintActionDatabase = FBlueprintActionDatabase::Get();
+		BlueprintActionDatabase.ClearAssetActions(DirectorBlueprint.Get());
+	}
+#endif
+}
+
 void UAvaSequence::OnWorldCleanup(UWorld* InWorld, bool bInSessionEnded, bool bInCleanupResources)
 {
 	// Ignore cleanups from other worlds
@@ -689,6 +705,8 @@ void UAvaSequence::OnWorldCleanup(UWorld* InWorld, bool bInSessionEnded, bool bI
 			, REN_DontCreateRedirectors | REN_NonTransactional | REN_DoNotDirty);
 	};
 
+	Cleanup();
+
 #if WITH_EDITORONLY_DATA
 	CleanupObject(DirectorBlueprint);
 #endif
diff --git a/Engine/Plugins/VirtualProduction/Avalanche/Source/AvalancheSequence/Public/AvaSequence.h b/Engine/Plugins/VirtualProduction/Avalanche/Source/AvalancheSequence/Public/AvaSequence.h
--- a/Engine/Plugins/VirtualProduction/Avalanche/Source/AvalancheSequence/Public/AvaSequence.h
+++ b/Engine/Plugins/VirtualProduction/Avalanche/Source/AvalancheSequence/Public/AvaSequence.h
@@ -153,6 +153,9 @@ public:
 	void OnOuterWorldRenamed(UObject* InNewOuter);
 #endif
 
+	/** Called to clean up resources */
+	void Cleanup();
+
 	void OnWorldCleanup(UWorld* InWorld, bool bInSessionEnded, bool bInCleanupResources);
 
 protected:
