MotionDesign : ActorModifier
- Fixed typo in modifier description

diff --git a/Engine/Plugins/VirtualProduction/Avalanche/Source/AvalancheModifiers/Private/Modifiers/AvaDynamicMeshConverterModifier.cpp b/Engine/Plugins/VirtualProduction/Avalanche/Source/AvalancheModifiers/Private/Modifiers/AvaDynamicMeshConverterModifier.cpp
--- a/Engine/Plugins/VirtualProduction/Avalanche/Source/AvalancheModifiers/Private/Modifiers/AvaDynamicMeshConverterModifier.cpp
+++ b/Engine/Plugins/VirtualProduction/Avalanche/Source/AvalancheModifiers/Private/Modifiers/AvaDynamicMeshConverterModifier.cpp
@@ -50,7 +50,7 @@ void UAvaDynamicMeshConverterModifier::OnModifierCDOSetup(FActorModifierCoreMeta
 	InMetadata.AllowTick(true);
 #if WITH_EDITOR
 	InMetadata.SetDisplayName(LOCTEXT("ModifierDisplayName", "Dynamic Mesh Converter"));
-	InMetadata.SetDescription(LOCTEXT("ModifierDescription", "Converts various actor mesh types into a single dynamic mesh, this is an heavy operation"));
+	InMetadata.SetDescription(LOCTEXT("ModifierDescription", "Converts various actor mesh types into a single dynamic mesh, this is a heavy operation"));
 #endif
 	InMetadata.SetCompatibilityRule([](const AActor* InActor)->bool
 	{
