Avalanche Outliner - Minor tweaks on automated tests to fix an error on Horde

diff --git a/Engine/Plugins/VirtualProduction/Avalanche/Source/AvalancheOutliner/Private/Tests/AvaOutliner.spec.cpp b/Engine/Plugins/VirtualProduction/Avalanche/Source/AvalancheOutliner/Private/Tests/AvaOutliner.spec.cpp
--- a/Engine/Plugins/VirtualProduction/Avalanche/Source/AvalancheOutliner/Private/Tests/AvaOutliner.spec.cpp
+++ b/Engine/Plugins/VirtualProduction/Avalanche/Source/AvalancheOutliner/Private/Tests/AvaOutliner.spec.cpp
@@ -11,292 +11,291 @@
 
 #if WITH_AUTOMATION_TESTS
 
-BEGIN_DEFINE_SPEC(FAvaOutlinerSpec, "Avalanche.Outliner"
-	, EAutomationTestFlags::EngineFilter
+BEGIN_DEFINE_SPEC(FAvaOutlinerSpec, "Avalanche.Outliner",
+	EAutomationTestFlags::EngineFilter
 	| EAutomationTestFlags::EditorContext
 	| EAutomationTestFlags_ApplicationContextMask)
 
 	// Types of common situations that are used across multiple tests
 	enum class EWhenType
-	{
-		NoItemsSelected,
-		ItemsSelected
-	};
-	
-	void When(const FString& InDescription, TFunction<void()>&& DoWork)
-	{
-		Describe("When " + InDescription, MoveTemp(DoWork));
-	}
+{
+	NoItemsSelected,
+	ItemsSelected
+};
+
+void When(const FString& InDescription, TFunction<void()>&& DoWork)
+{
+	Describe("When " + InDescription, MoveTemp(DoWork));
+}
 
-	void When(EWhenType InWhenType, TFunction<void()>&& DoWork)
+void When(EWhenType InWhenType, TFunction<void()>&& DoWork)
+{
+	switch (InWhenType)
 	{
-		switch (InWhenType)
-		{
-			case EWhenType::NoItemsSelected:
-				When("no items are selected", MoveTemp(DoWork));
-				break;
+	case EWhenType::NoItemsSelected:
+		When("no items are selected", MoveTemp(DoWork));
+		break;
 
-			case EWhenType::ItemsSelected:
-				When("multiple items are selected", [this, DoWork = MoveTemp(DoWork)]()
+	case EWhenType::ItemsSelected:
+		When("multiple items are selected", [this, DoWork = MoveTemp(DoWork)]()
+			{
+				BeforeEach([this]()
 					{
-						BeforeEach([this]()
-							{
-								OutlinerView->SelectItems(TestItems);
-							});
-
-						DoWork();
+						OutlinerView->SelectItems(TestItems);
 					});
-				break;
-		}
+				DoWork();
+			});
+		break;
 	}
+}
 
-	TSharedPtr<FAvaOutliner> Outliner;
-
-	// The outliner view instance created before each test
-	TSharedPtr<FAvaOutlinerView> OutlinerView;
-
-	// Custom test provider to fill some of the info/actions requested by the outliner (e.g. USelection Instances, or Actor Duplication)
-	TSharedPtr<UE::AvaOutliner::Private::FAvaOutlinerProviderTest> OutlinerProvider;
-
-	// A list of outliner items to test things with
-	TArray<FAvaOutlinerItemPtr> TestItems;
+TSharedPtr<FAvaOutliner> Outliner;
+TSharedPtr<FAvaOutlinerView> OutlinerView;
+TSharedPtr<UE::AvaOutliner::Private::FAvaOutlinerProviderTest> OutlinerProvider;
+TArray<FAvaOutlinerItemPtr> TestItems;
 
 END_DEFINE_SPEC(FAvaOutlinerSpec)
 
 void FAvaOutlinerSpec::Define()
 {
 	BeforeEach([this]()
-	{
-		OutlinerProvider = MakeShared<UE::AvaOutliner::Private::FAvaOutlinerProviderTest>();
-
-		constexpr int32 OutlinerViewId = 0;
-
-		Outliner = OutlinerProvider->GetOutliner();
-		Outliner->RegisterOutlinerView(OutlinerViewId);
-		Outliner->Refresh();
-
-		OutlinerView = Outliner->GetMostRecentOutlinerView();
-
-		const TArray<UObject*> TestObjects
 		{
-			OutlinerProvider->GetDirectionalLight(),
-			OutlinerProvider->GetFloor(),
-			OutlinerProvider->GetSkySphere(),
-		};
+			// Always start with a clean slate
+			TestItems.Reset();
+			OutlinerView.Reset();
+			Outliner.Reset();
+			OutlinerProvider.Reset();
 
-		TestItems = OutlinerProvider->GetOutlinerItems(TestObjects);
-	 });
+			// Create new test provider and outliner for each test
+			OutlinerProvider = MakeShared<UE::AvaOutliner::Private::FAvaOutlinerProviderTest>();
 
-	Describe("Spawning an Actor", [this]()
-	{
-		BeforeEach([this]()
-		{
-			OutlinerProvider->TestSpawnActor();
+			constexpr int32 OutlinerViewId = 0;
+
+			Outliner = OutlinerProvider->GetOutliner();
+			Outliner->RegisterOutlinerView(OutlinerViewId);
 			Outliner->Refresh();
-		});
 
-		It("should register an item for the actor", [this]()
-		{
-			const FAvaOutlinerItemPtr Item = Outliner->FindItem(OutlinerProvider->GetTestSpawnActor());
-			TestTrue(TEXT("Item found for the actor"), Item.IsValid());
+			OutlinerView = Outliner->GetMostRecentOutlinerView();
 
-			if (!TestTrue(TEXT("Item is an actor item"), Item.IsValid() && Item->IsA<FAvaOutlinerActor>()))
+			const TArray<UObject*> TestObjects
 			{
-				return;
-			}
+				OutlinerProvider->GetDirectionalLight(),
+				OutlinerProvider->GetFloor(),
+				OutlinerProvider->GetSkySphere(),
+			};
 
-			const AActor* const Actor = Item->CastTo<FAvaOutlinerActor>()->GetActor();
-			TestTrue(TEXT("Actor item matches the spawned actor"), Actor == OutlinerProvider->GetTestSpawnActor());	
-		});
-	});
-	
-	Describe("Duplicating Selected Items", [this]()
-	{
-		When(EWhenType::NoItemsSelected, [this]()
-		{
-			It("should keep the outliner with no selections", [this]()
-			{
-				OutlinerView->DuplicateSelected();
-				Outliner->Refresh();
-				TestEqual(TEXT("Selected Item count is 0")
-					, Outliner->GetSelectedItemCount()
-					, 0);
-			});
+			TestItems = OutlinerProvider->GetOutlinerItems(TestObjects);
 		});
 
-		When(EWhenType::ItemsSelected, [this]()
+	Describe("Spawning an Actor", [this]()
 		{
-			Describe("When selected objects have duplicated", [this]()
-			{
-				BeforeEach([this]()
+			BeforeEach([this]()
 				{
-					OutlinerView->DuplicateSelected();
-				});
-				
-				It("should have a matching pending item action count", [this]()
-				{
-					TestEqual(TEXT("Number of items enqueued match the requested number")
-						, Outliner->GetPendingItemActionCount()
-						, TestItems.Num());
+					OutlinerProvider->TestSpawnActor();
+					Outliner->Refresh();
 				});
 
-				Describe("When outliner has refreshed", [this]()
+			It("should register an item for the actor", [this]()
 				{
-					BeforeEach([this]()
-					{
-						Outliner->Refresh();
-					});
-
-					//It("should have a matching pending item action count", [this]()
-					//{
-					//	TestEqual(TEXT("The template item count and the number of copies match")
-					//		, Outliner->GetSelectedItemCount()
-					//		, TestItems.Num());	
-					//});
+					const FAvaOutlinerItemPtr Item = Outliner->FindItem(OutlinerProvider->GetTestSpawnActor());
+					TestTrue(TEXT("Item found for the actor"), Item.IsValid());
 
-					It("should have the duplicate copies selected instead of the template", [this]()
+					if (!TestTrue(TEXT("Item is an actor item"), Item.IsValid() && Item->IsA<FAvaOutlinerActor>()))
 					{
-						// These should be the new items that were created from duplication
-						const TArray<FAvaOutlinerItemPtr> DuplicateItems = Outliner->GetSelectedItems();
-
-						TestFalse(TEXT("The newly selected item count is not zero"), DuplicateItems.IsEmpty());
-
-						constexpr int32 TestItemIndex = 0;
-							
-						FAvaOutliner::SortItems(TestItems);
-							
-						TestNotEqual(TEXT("The newly selected items are different from the old selected items")
-							, DuplicateItems[TestItemIndex]
-							, TestItems[TestItemIndex]);
-					});
-
-					It("should have each duplicate item right below its respective template item", [this]()
-					{
-						constexpr int32 TestItemIndex = 0;
-
-						FAvaOutliner::SortItems(TestItems);
-						
-						const FAvaOutlinerItemPtr TemplateItem  = TestItems[TestItemIndex];
-						const FAvaOutlinerItemPtr DuplicateItem = Outliner->GetSelectedItems()[TestItemIndex];
-
-						TestEqual(TEXT("The template item and duplicate item share the same parent item")
-							, TemplateItem->GetParent()
-							, DuplicateItem->GetParent());
-
-						const FAvaOutlinerItemPtr ParentItem = TemplateItem->GetParent();
-
-						const int32 TemplateItemIndex  = ParentItem->GetChildIndex(TemplateItem);
-						const int32 DuplicateItemIndex = ParentItem->GetChildIndex(DuplicateItem);
+						return;
+					}
 
-						TestEqual(TEXT("The duplicate item is right below template item")
-							, TemplateItemIndex + 1
-							, DuplicateItemIndex);
-					});
+					const AActor* Actor = Item->CastTo<FAvaOutlinerActor>()->GetActor();
+					TestTrue(TEXT("Actor item matches the spawned actor"), Actor == OutlinerProvider->GetTestSpawnActor());
 				});
-			});
-		});	
-	});
+		});
 
-	Describe("Grouping Selected Items", [this]()
-	{
-		When("grouping with null", [this]()
+	Describe("Duplicating Selected Items", [this]()
 		{
 			When(EWhenType::NoItemsSelected, [this]()
-			{
-				It("should keep the outliner with no selections", [this]()
 				{
-					Outliner->GroupSelection(nullptr);
-					Outliner->Refresh();
-					TestEqual(TEXT("Selected Item count is 0")
-						, Outliner->GetSelectedItemCount()
-						, 0);
+					It("should keep the outliner with no selections", [this]()
+						{
+							OutlinerView->DuplicateSelected();
+							Outliner->Refresh();
+							TestEqual(TEXT("Selected Item count is 0"),
+								Outliner->GetSelectedItemCount(),
+								0);
+						});
 				});
-			});
 
 			When(EWhenType::ItemsSelected, [this]()
-			{
-				It("should keep the outliner with the same selected items", [this]()
 				{
-					const TArray<FAvaOutlinerItemPtr> PrevSelectedItems = Outliner->GetSelectedItems();
-
-					Outliner->GroupSelection(nullptr);
-					Outliner->Refresh();
-
-					const TArray<FAvaOutlinerItemPtr> CurrSelectedItems = Outliner->GetSelectedItems();
-
-					TestEqual(TEXT("Selected Item count remains the same")
-						, PrevSelectedItems.Num()
-						, CurrSelectedItems.Num());
-
-					TestEqual(TEXT("Selected Items remain the same")
-						, PrevSelectedItems[0]
-						, CurrSelectedItems[0]);
+					Describe("When selected objects have duplicated", [this]()
+						{
+							BeforeEach([this]()
+								{
+									OutlinerView->DuplicateSelected();
+								});
+
+							It("should have a matching pending item action count", [this]()
+								{
+									TestEqual(TEXT("Number of items enqueued match the requested number"),
+										Outliner->GetPendingItemActionCount(),
+										TestItems.Num());
+								});
+
+							Describe("When outliner has refreshed", [this]()
+								{
+									BeforeEach([this]()
+										{
+											Outliner->Refresh();
+										});
+
+									 It("should have a matching pending item action count", [this]()
+									 {
+									 	TestEqual(TEXT("The template item count and the number of copies match"),
+									 		Outliner->GetSelectedItemCount(),
+									 		TestItems.Num());
+									 });
+
+									It("should have the duplicate copies selected instead of the template", [this]()
+										{
+											const TArray<FAvaOutlinerItemPtr> DuplicateItems = Outliner->GetSelectedItems();
+											TestFalse(TEXT("The newly selected item count is not zero"), DuplicateItems.IsEmpty());
+
+											constexpr int32 TestItemIndex = 0;
+											FAvaOutliner::SortItems(TestItems);
+
+											TestNotEqual(TEXT("The newly selected items are different from the old selected items"),
+												DuplicateItems[TestItemIndex],
+												TestItems[TestItemIndex]);
+										});
+
+									It("should have each duplicate item right below its respective template item", [this]()
+										{
+											constexpr int32 TestItemIndex = 0;
+											FAvaOutliner::SortItems(TestItems);
+
+											const FAvaOutlinerItemPtr TemplateItem = TestItems[TestItemIndex];
+											const FAvaOutlinerItemPtr DuplicateItem = Outliner->GetSelectedItems()[TestItemIndex];
+
+											TestEqual(TEXT("The template item and duplicate item share the same parent item"),
+												TemplateItem->GetParent(),
+												DuplicateItem->GetParent());
+
+											const FAvaOutlinerItemPtr ParentItem = TemplateItem->GetParent();
+
+											const int32 TemplateItemIndex = ParentItem->GetChildIndex(TemplateItem);
+											const int32 DuplicateItemIndex = ParentItem->GetChildIndex(DuplicateItem);
+
+											TestEqual(TEXT("The duplicate item is right below template item"),
+												TemplateItemIndex + 1,
+												DuplicateItemIndex);
+										});
+								});
+						});
 				});
-			});
 		});
 
-		When("grouping with an already existing actor", [this]()
+	Describe("Grouping Selected Items", [this]()
 		{
-			When(EWhenType::NoItemsSelected, [this]()
-			{
-				BeforeEach([this]()
+			When("grouping with null", [this]()
 				{
-					AActor* const GroupingActor = OutlinerProvider->GetVolumeActor();
-					TestNotNull(TEXT("Grouping actor exists"), GroupingActor);
-					Outliner->GroupSelection(GroupingActor);
-					Outliner->Refresh();
+					When(EWhenType::NoItemsSelected, [this]()
+						{
+							It("should keep the outliner with no selections", [this]()
+								{
+									Outliner->GroupSelection(nullptr);
+									Outliner->Refresh();
+									TestEqual(TEXT("Selected Item count is 0"),
+										Outliner->GetSelectedItemCount(),
+										0);
+								});
+						});
+
+					When(EWhenType::ItemsSelected, [this]()
+						{
+							It("should keep the outliner with the same selected items", [this]()
+								{
+									const TArray<FAvaOutlinerItemPtr> PrevSelectedItems = Outliner->GetSelectedItems();
+
+									Outliner->GroupSelection(nullptr);
+									Outliner->Refresh();
+
+									const TArray<FAvaOutlinerItemPtr> CurrSelectedItems = Outliner->GetSelectedItems();
+
+									TestEqual(TEXT("Selected Item count remains the same"),
+										PrevSelectedItems.Num(),
+										CurrSelectedItems.Num());
+
+									TestEqual(TEXT("Selected Items remain the same"),
+										PrevSelectedItems[0],
+										CurrSelectedItems[0]);
+								});
+						});
 				});
 
-				It("should keep the outliner with no selections", [this]()
+			When("grouping with an already existing actor", [this]()
 				{
-					TestEqual(TEXT("Selected item count is zero"), Outliner->GetSelectedItemCount(), 0);
+					When(EWhenType::NoItemsSelected, [this]()
+						{
+							BeforeEach([this]()
+								{
+									AActor* GroupingActor = OutlinerProvider->GetVolumeActor();
+									TestNotNull(TEXT("Grouping actor exists"), GroupingActor);
+									Outliner->GroupSelection(GroupingActor);
+									Outliner->Refresh();
+								});
+
+							It("should keep the outliner with no selections", [this]()
+								{
+									TestEqual(TEXT("Selected item count is zero"), Outliner->GetSelectedItemCount(), 0);
+								});
+						});
+
+					When(EWhenType::ItemsSelected, [this]()
+						{
+							BeforeEach([this]()
+								{
+									AActor* GroupingActor = OutlinerProvider->GetVolumeActor();
+									TestNotNull(TEXT("Grouping actor exists"), GroupingActor);
+									Outliner->GroupSelection(GroupingActor);
+									Outliner->Refresh();
+								});
+
+							It("should keep the outliner with the selected items", [this]()
+								{
+									const TArray<FAvaOutlinerItemPtr> SelectedItems = Outliner->GetSelectedItems();
+
+									TestEqual(TEXT("Selected item count is consistent"), SelectedItems.Num(), TestItems.Num());
+
+									if (TestNotEqual(TEXT("Selected item count is not zero"), SelectedItems.Num(), 0))
+									{
+										TestEqual(TEXT("The selected items are consistent"), TestItems[0], SelectedItems[0]);
+									}
+								});
+
+							It("should have placed the selected items under the grouping actor", [this]()
+								{
+									const AActor* GroupingActor = OutlinerProvider->GetVolumeActor();
+									const FAvaOutlinerItemPtr GroupingItem = Outliner->FindItem(GroupingActor);
+
+									TestNotNull(TEXT("Grouping item exists in outliner"), GroupingItem.Get());
+
+									const int32 ChildIndex = GroupingItem->GetChildIndex(TestItems[0]);
+
+									TestTrue(TEXT("the selected items are under the grouping item"), ChildIndex != INDEX_NONE);
+								});
+						});
 				});
-			});
-
-			When(EWhenType::ItemsSelected, [this]()
-			{
-				BeforeEach([this]()
-				{
-					AActor* const GroupingActor = OutlinerProvider->GetVolumeActor();
-					TestNotNull(TEXT("Grouping actor exists"), GroupingActor);
-					Outliner->GroupSelection(GroupingActor);
-					Outliner->Refresh();
-				});
-
-				It("should keep the outliner with the selected items", [this]()
-				{
-					const TArray<FAvaOutlinerItemPtr> SelectedItems = Outliner->GetSelectedItems();
-
-					TestEqual(TEXT("Selected item count is consistent"), SelectedItems.Num(), TestItems.Num());
-
-					if (TestNotEqual(TEXT("Selected item count is not zero"), SelectedItems.Num(), 0))
-					{
-						TestEqual(TEXT("The selected items are consistent"), TestItems[0], SelectedItems[0]);
-					}
-				});
-
-				It("should have placed the selected items under the grouping actor", [this]()
-				{
-					const AActor* const GroupingActor = OutlinerProvider->GetVolumeActor();
-					const FAvaOutlinerItemPtr GroupingItem = Outliner->FindItem(GroupingActor);
-
-					TestNotNull(TEXT("Grouping item exists in outliner"), GroupingItem.Get());
-
-					const int32 ChildIndex = GroupingItem->GetChildIndex(TestItems[0]);
-
-					TestTrue(TEXT("the selected items are under the grouping item"), ChildIndex != INDEX_NONE);
-				});
-			});
 		});
-	});
 
 	AfterEach([this]()
-	{
-		TestItems.Reset();
-		OutlinerView.Reset();
-		Outliner.Reset();
-		OutlinerProvider.Reset();
-	});
+		{
+			TestItems.Reset();
+			OutlinerView.Reset();
+			Outliner.Reset();
+			OutlinerProvider.Reset();
+
+			// Force Unreal to clean up any lingering objects between tests
+			CollectGarbage(GARBAGE_COLLECTION_KEEPFLAGS);
+		});
 }
 
 #endif
