Avalanche Outliner - tweaks on automated tests to ensure TypedElement registration avoids crash

diff --git a/Engine/Plugins/VirtualProduction/Avalanche/Source/AvalancheOutliner/Private/Tests/AvaOutliner.spec.cpp b/Engine/Plugins/VirtualProduction/Avalanche/Source/AvalancheOutliner/Private/Tests/AvaOutliner.spec.cpp
--- a/Engine/Plugins/VirtualProduction/Avalanche/Source/AvalancheOutliner/Private/Tests/AvaOutliner.spec.cpp
+++ b/Engine/Plugins/VirtualProduction/Avalanche/Source/AvalancheOutliner/Private/Tests/AvaOutliner.spec.cpp
@@ -16,7 +16,6 @@ BEGIN_DEFINE_SPEC(FAvaOutlinerSpec, "Avalanche.Outliner",
 	| EAutomationTestFlags::EditorContext
 	| EAutomationTestFlags_ApplicationContextMask)
 
-	// Types of common situations that are used across multiple tests
 	enum class EWhenType
 {
 	NoItemsSelected,
@@ -41,7 +40,15 @@ void When(EWhenType InWhenType, TFunction<void()>&& DoWork)
 			{
 				BeforeEach([this]()
 					{
-						OutlinerView->SelectItems(TestItems);
+						const bool bAllValid = !TestItems.ContainsByPredicate([](const FAvaOutlinerItemPtr& Item)
+							{
+								return !Item.IsValid();
+							});
+
+						if (TestItems.Num() > 0 && bAllValid)
+						{
+							OutlinerView->SelectItems(TestItems);
+						}
 					});
 				DoWork();
 			});
@@ -60,28 +67,27 @@ void FAvaOutlinerSpec::Define()
 {
 	BeforeEach([this]()
 		{
-			// Always start with a clean slate
 			TestItems.Reset();
 			OutlinerView.Reset();
 			Outliner.Reset();
 			OutlinerProvider.Reset();
 
-			// Create new test provider and outliner for each test
 			OutlinerProvider = MakeShared<UE::AvaOutliner::Private::FAvaOutlinerProviderTest>();
 
 			constexpr int32 OutlinerViewId = 0;
-
 			Outliner = OutlinerProvider->GetOutliner();
+			check(Outliner.IsValid());
+
 			Outliner->RegisterOutlinerView(OutlinerViewId);
 			Outliner->Refresh();
 
 			OutlinerView = Outliner->GetMostRecentOutlinerView();
+			check(OutlinerView.IsValid());
 
-			const TArray<UObject*> TestObjects
-			{
+			const TArray<UObject*> TestObjects = {
 				OutlinerProvider->GetDirectionalLight(),
 				OutlinerProvider->GetFloor(),
-				OutlinerProvider->GetSkySphere(),
+				OutlinerProvider->GetSkySphere()
 			};
 
 			TestItems = OutlinerProvider->GetOutlinerItems(TestObjects);
@@ -98,14 +104,12 @@ void FAvaOutlinerSpec::Define()
 			It("should register an item for the actor", [this]()
 				{
 					const FAvaOutlinerItemPtr Item = Outliner->FindItem(OutlinerProvider->GetTestSpawnActor());
-					TestTrue(TEXT("Item found for the actor"), Item.IsValid());
+					TestNotNull(TEXT("Item found for the actor"), Item.Get());
 
-					if (!TestTrue(TEXT("Item is an actor item"), Item.IsValid() && Item->IsA<FAvaOutlinerActor>()))
-					{
-						return;
-					}
+					const TSharedPtr<FAvaOutlinerActor> ActorItem = StaticCastSharedPtr<FAvaOutlinerActor>(Item);
+					TestNotNull(TEXT("Item cast to FAvaOutlinerActor is valid"), ActorItem.Get());
 
-					const AActor* Actor = Item->CastTo<FAvaOutlinerActor>()->GetActor();
+					const AActor* Actor = ActorItem->GetActor();
 					TestTrue(TEXT("Actor item matches the spawned actor"), Actor == OutlinerProvider->GetTestSpawnActor());
 				});
 		});
@@ -118,9 +122,7 @@ void FAvaOutlinerSpec::Define()
 						{
 							OutlinerView->DuplicateSelected();
 							Outliner->Refresh();
-							TestEqual(TEXT("Selected Item count is 0"),
-								Outliner->GetSelectedItemCount(),
-								0);
+							TestEqual(TEXT("Selected Item count is 0"), Outliner->GetSelectedItemCount(), 0);
 						});
 				});
 
@@ -147,46 +149,61 @@ void FAvaOutlinerSpec::Define()
 											Outliner->Refresh();
 										});
 
-									 It("should have a matching pending item action count", [this]()
-									 {
-									 	TestEqual(TEXT("The template item count and the number of copies match"),
-									 		Outliner->GetSelectedItemCount(),
-									 		TestItems.Num());
-									 });
+									It("should have a matching selected item count", [this]()
+										{
+											TestEqual(TEXT("The template item count and the number of copies match"),
+												Outliner->GetSelectedItemCount(),
+												TestItems.Num());
+										});
 
 									It("should have the duplicate copies selected instead of the template", [this]()
 										{
 											const TArray<FAvaOutlinerItemPtr> DuplicateItems = Outliner->GetSelectedItems();
 											TestFalse(TEXT("The newly selected item count is not zero"), DuplicateItems.IsEmpty());
 
-											constexpr int32 TestItemIndex = 0;
+											constexpr int32 Index = 0;
 											FAvaOutliner::SortItems(TestItems);
 
-											TestNotEqual(TEXT("The newly selected items are different from the old selected items"),
-												DuplicateItems[TestItemIndex],
-												TestItems[TestItemIndex]);
+											if (DuplicateItems.IsValidIndex(Index) && TestItems.IsValidIndex(Index) &&
+												DuplicateItems[Index].IsValid() && TestItems[Index].IsValid())
+											{
+												TestNotEqual(TEXT("The newly selected items are different from the old selected items"),
+													DuplicateItems[Index],
+													TestItems[Index]);
+											}
 										});
 
 									It("should have each duplicate item right below its respective template item", [this]()
 										{
-											constexpr int32 TestItemIndex = 0;
+											constexpr int32 Index = 0;
 											FAvaOutliner::SortItems(TestItems);
 
-											const FAvaOutlinerItemPtr TemplateItem = TestItems[TestItemIndex];
-											const FAvaOutlinerItemPtr DuplicateItem = Outliner->GetSelectedItems()[TestItemIndex];
+											if (TestItems.IsValidIndex(Index) && TestItems[Index].IsValid())
+											{
+												const FAvaOutlinerItemPtr TemplateItem = TestItems[Index];
+												const TArray<FAvaOutlinerItemPtr> SelectedItems = Outliner->GetSelectedItems();
 
-											TestEqual(TEXT("The template item and duplicate item share the same parent item"),
-												TemplateItem->GetParent(),
-												DuplicateItem->GetParent());
+												if (SelectedItems.IsValidIndex(Index) && SelectedItems[Index].IsValid())
+												{
+													const FAvaOutlinerItemPtr DuplicateItem = SelectedItems[Index];
 
-											const FAvaOutlinerItemPtr ParentItem = TemplateItem->GetParent();
+													TestEqual(TEXT("The template item and duplicate item share the same parent item"),
+														TemplateItem->GetParent(),
+														DuplicateItem->GetParent());
 
-											const int32 TemplateItemIndex = ParentItem->GetChildIndex(TemplateItem);
-											const int32 DuplicateItemIndex = ParentItem->GetChildIndex(DuplicateItem);
+													const FAvaOutlinerItemPtr ParentItem = TemplateItem->GetParent();
 
-											TestEqual(TEXT("The duplicate item is right below template item"),
-												TemplateItemIndex + 1,
-												DuplicateItemIndex);
+													if (ParentItem.IsValid())
+													{
+														const int32 TemplateIdx = ParentItem->GetChildIndex(TemplateItem);
+														const int32 DuplicateIdx = ParentItem->GetChildIndex(DuplicateItem);
+
+														TestEqual(TEXT("The duplicate item is right below template item"),
+															TemplateIdx + 1,
+															DuplicateIdx);
+													}
+												}
+											}
 										});
 								});
 						});
@@ -203,9 +220,7 @@ void FAvaOutlinerSpec::Define()
 								{
 									Outliner->GroupSelection(nullptr);
 									Outliner->Refresh();
-									TestEqual(TEXT("Selected Item count is 0"),
-										Outliner->GetSelectedItemCount(),
-										0);
+									TestEqual(TEXT("Selected Item count is 0"), Outliner->GetSelectedItemCount(), 0);
 								});
 						});
 
@@ -213,20 +228,22 @@ void FAvaOutlinerSpec::Define()
 						{
 							It("should keep the outliner with the same selected items", [this]()
 								{
-									const TArray<FAvaOutlinerItemPtr> PrevSelectedItems = Outliner->GetSelectedItems();
+									const TArray<FAvaOutlinerItemPtr> PrevItems = Outliner->GetSelectedItems();
 
 									Outliner->GroupSelection(nullptr);
 									Outliner->Refresh();
 
-									const TArray<FAvaOutlinerItemPtr> CurrSelectedItems = Outliner->GetSelectedItems();
+									const TArray<FAvaOutlinerItemPtr> CurrItems = Outliner->GetSelectedItems();
 
 									TestEqual(TEXT("Selected Item count remains the same"),
-										PrevSelectedItems.Num(),
-										CurrSelectedItems.Num());
+										PrevItems.Num(), CurrItems.Num());
 
-									TestEqual(TEXT("Selected Items remain the same"),
-										PrevSelectedItems[0],
-										CurrSelectedItems[0]);
+									if (PrevItems.Num() > 0 && CurrItems.Num() > 0 &&
+										PrevItems[0].IsValid() && CurrItems[0].IsValid())
+									{
+										TestEqual(TEXT("Selected Items remain the same"),
+											PrevItems[0], CurrItems[0]);
+									}
 								});
 						});
 				});
@@ -237,9 +254,9 @@ void FAvaOutlinerSpec::Define()
 						{
 							BeforeEach([this]()
 								{
-									AActor* GroupingActor = OutlinerProvider->GetVolumeActor();
-									TestNotNull(TEXT("Grouping actor exists"), GroupingActor);
-									Outliner->GroupSelection(GroupingActor);
+									AActor* GroupActor = OutlinerProvider->GetVolumeActor();
+									TestNotNull(TEXT("Grouping actor exists"), GroupActor);
+									Outliner->GroupSelection(GroupActor);
 									Outliner->Refresh();
 								});
 
@@ -253,34 +270,36 @@ void FAvaOutlinerSpec::Define()
 						{
 							BeforeEach([this]()
 								{
-									AActor* GroupingActor = OutlinerProvider->GetVolumeActor();
-									TestNotNull(TEXT("Grouping actor exists"), GroupingActor);
-									Outliner->GroupSelection(GroupingActor);
+									AActor* GroupActor = OutlinerProvider->GetVolumeActor();
+									TestNotNull(TEXT("Grouping actor exists"), GroupActor);
+									Outliner->GroupSelection(GroupActor);
 									Outliner->Refresh();
 								});
 
 							It("should keep the outliner with the selected items", [this]()
 								{
-									const TArray<FAvaOutlinerItemPtr> SelectedItems = Outliner->GetSelectedItems();
-
-									TestEqual(TEXT("Selected item count is consistent"), SelectedItems.Num(), TestItems.Num());
+									const TArray<FAvaOutlinerItemPtr> Selected = Outliner->GetSelectedItems();
+									TestEqual(TEXT("Selected item count is consistent"), Selected.Num(), TestItems.Num());
 
-									if (TestNotEqual(TEXT("Selected item count is not zero"), SelectedItems.Num(), 0))
+									if (Selected.Num() > 0 && TestItems.Num() > 0 &&
+										Selected[0].IsValid() && TestItems[0].IsValid())
 									{
-										TestEqual(TEXT("The selected items are consistent"), TestItems[0], SelectedItems[0]);
+										TestEqual(TEXT("The selected items are consistent"), TestItems[0], Selected[0]);
 									}
 								});
 
 							It("should have placed the selected items under the grouping actor", [this]()
 								{
-									const AActor* GroupingActor = OutlinerProvider->GetVolumeActor();
-									const FAvaOutlinerItemPtr GroupingItem = Outliner->FindItem(GroupingActor);
+									const AActor* GroupActor = OutlinerProvider->GetVolumeActor();
+									const FAvaOutlinerItemPtr GroupItem = Outliner->FindItem(GroupActor);
 
-									TestNotNull(TEXT("Grouping item exists in outliner"), GroupingItem.Get());
+									TestNotNull(TEXT("Grouping item exists in outliner"), GroupItem.Get());
 
-									const int32 ChildIndex = GroupingItem->GetChildIndex(TestItems[0]);
-
-									TestTrue(TEXT("the selected items are under the grouping item"), ChildIndex != INDEX_NONE);
+									if (TestItems.Num() > 0 && TestItems[0].IsValid())
+									{
+										const int32 ChildIdx = GroupItem->GetChildIndex(TestItems[0]);
+										TestTrue(TEXT("The selected items are under the grouping item"), ChildIdx != INDEX_NONE);
+									}
 								});
 						});
 				});
@@ -293,9 +312,8 @@ void FAvaOutlinerSpec::Define()
 			Outliner.Reset();
 			OutlinerProvider.Reset();
 
-			// Force Unreal to clean up any lingering objects between tests
 			CollectGarbage(GARBAGE_COLLECTION_KEEPFLAGS);
 		});
 }
 
-#endif
+#endif // WITH_AUTOMATION_TESTS
