MotionDesign : Shapes
- Fixed translucent materials behavior with shapes (reset defaults)
- Fixed self intersections when using extrude modifier on 2D shapes

#virtualized

diff --git a/Engine/Plugins/VirtualProduction/Avalanche/Source/AvalancheModifiers/Private/Modifiers/AvaExtrudeModifier.cpp b/Engine/Plugins/VirtualProduction/Avalanche/Source/AvalancheModifiers/Private/Modifiers/AvaExtrudeModifier.cpp
--- a/Engine/Plugins/VirtualProduction/Avalanche/Source/AvalancheModifiers/Private/Modifiers/AvaExtrudeModifier.cpp
+++ b/Engine/Plugins/VirtualProduction/Avalanche/Source/AvalancheModifiers/Private/Modifiers/AvaExtrudeModifier.cpp
@@ -8,11 +8,12 @@
 #include "DynamicMeshEditor.h"
 #include "DynamicMesh/DynamicMeshAttributeSet.h"
 #include "DynamicMesh/MeshTransforms.h"
+#include "DynamicMesh/Operations/MergeCoincidentMeshEdges.h"
 #include "DynamicMeshes/AvaShapeDynMeshBase.h"
 #include "GeometryScript/MeshModelingFunctions.h"
+#include "Operations/MeshSelfUnion.h"
 #include "Operations/OffsetMeshRegion.h"
 #include "UDynamicMesh.h"
-#include "DynamicMesh/Operations/MergeCoincidentMeshEdges.h"
 
 #define LOCTEXT_NAMESPACE "AvaExtrudeModifier"
 
@@ -107,6 +108,7 @@ void UAvaExtrudeModifier::Apply()
 		{
 			OriginalTriangleIds.Add(TId);
 		}
+
 		// extrude options
 		const FVector3d ExtrudeVec = static_cast<double>(Options.ExtrudeDistance) * Options.ExtrudeDirection;
 		FOffsetMeshRegion Extruder(&EditMesh);
@@ -119,8 +121,9 @@ void UAvaExtrudeModifier::Apply()
 		Extruder.UVScaleFactor = Options.UVScale;
 		Extruder.bOffsetFullComponentsAsSolids = Options.bSolidsToShells;
 		Extruder.Apply();
+
+		// get new UV ids to scale them properly
 		{
-			// get new UV ids to scale them properly
 			TSet<int32> ExtrudeUVIds;
 			TArray<int32> ExtrudeTIds;
 			for (int32 TId : Extruder.AllModifiedAndNewTriangles)
@@ -154,21 +157,21 @@ void UAvaExtrudeModifier::Apply()
 			}
 		}
 
+		// Offset mesh
 		FVector3d MeshOffset = FVector3d::ZeroVector;
-
 		switch (ExtrudeMode)
 		{
 			case EAvaExtrudeMode::Opposite:
-				MeshOffset = ExtrudeVec;
-				break;
+			MeshOffset = ExtrudeVec;
+			break;
 
-			case EAvaExtrudeMode::Symmetrical:
-				MeshOffset = ExtrudeVec * 0.5f;
-				break;
+		case EAvaExtrudeMode::Symmetrical:
+			MeshOffset = ExtrudeVec * 0.5f;
+			break;
 
-			case EAvaExtrudeMode::Front:
-				break;
-			default: ;
+		case EAvaExtrudeMode::Front:
+			break;
+		default: ;
 		}
 
 		// move mesh backwards if needed
@@ -182,6 +185,7 @@ void UAvaExtrudeModifier::Apply()
 				EditMesh.SetVertex(VId, Loc);
 			}
 		}
+
 		// close back
 		if (bCloseBack)
 		{
@@ -198,6 +202,15 @@ void UAvaExtrudeModifier::Apply()
 			TmpMappings.GetTriangleMap().GetForwardMap().GenerateValueArray(OutBackTriangles);
 			FDynamicMeshPolygroupAttribute* const BackPolygroup = FindOrCreatePolygroupLayer(EditMesh, UAvaExtrudeModifier::BackPolygroupLayerName, &OutBackTriangles);
 		}
+
+		// Remove geometry self intersections
+		{
+			FMeshSelfUnion Union(&EditMesh);
+			Union.bCollapseDegenerateEdgesOnCut = true;
+			Union.bTrimFlaps = false;
+			Union.Compute();
+		}
+
 	}, EDynamicMeshChangeType::GeneralEdit, EDynamicMeshAttributeChangeFlags::Unknown);
 
 	Next();
