[Motion Design] Crash fix when StopLayer command has the wrong number of arguments.

diff --git a/Engine/Plugins/VirtualProduction/Avalanche/Source/AvalancheMediaEditor/Private/Rundown/AvaRundownEditor.cpp b/Engine/Plugins/VirtualProduction/Avalanche/Source/AvalancheMediaEditor/Private/Rundown/AvaRundownEditor.cpp
--- a/Engine/Plugins/VirtualProduction/Avalanche/Source/AvalancheMediaEditor/Private/Rundown/AvaRundownEditor.cpp
+++ b/Engine/Plugins/VirtualProduction/Avalanche/Source/AvalancheMediaEditor/Private/Rundown/AvaRundownEditor.cpp
@@ -1398,6 +1398,7 @@ void FAvaRundownEditor::StopLayersCommand(const TArray<FString>& InArgs, bool bI
 	if (InArgs.Num() < 2)
 	{
 		UE_LOG(LogAvaRundown, Error, TEXT("Errors in \"Stop Layers\" command: Too few arguments, need at least ChannelName and one layer name"));
+		return;
 	}
 
 	const FName ChannelName(*InArgs[0]);
