Motion Design: initial implementation of ticker actor and component. Still missing visualization (and scene proxy) and queue limit

diff --git a/Engine/Plugins/VirtualProduction/Avalanche/Source/Avalanche/Internal/Framework/Ticker/AvaTickerActor.h b/Engine/Plugins/VirtualProduction/Avalanche/Source/Avalanche/Internal/Framework/Ticker/AvaTickerActor.h
new file mode 100644
--- /dev/null
+++ b/Engine/Plugins/VirtualProduction/Avalanche/Source/Avalanche/Internal/Framework/Ticker/AvaTickerActor.h
@@ -0,0 +1,27 @@
+// Copyright Epic Games, Inc. All Rights Reserved.
+
+#pragma once
+
+#include "GameFramework/Actor.h"
+#include "AvaTickerActor.generated.h"
+
+class UAvaTickerComponent;
+
+UCLASS(MinimalAPI, DisplayName="Motion Design Ticker Actor")
+class AAvaTickerActor : public AActor
+{
+	GENERATED_BODY()
+
+public:
+	AAvaTickerActor();
+
+	//~ Begin AActor
+#if WITH_EDITOR
+	virtual FString GetDefaultActorLabel() const override;
+#endif
+	//~ End AActor
+
+protected:
+	UPROPERTY(VisibleAnywhere, BlueprintReadOnly, Category="Ticker")
+	TObjectPtr<UAvaTickerComponent> TickerComponent;
+};
diff --git a/Engine/Plugins/VirtualProduction/Avalanche/Source/Avalanche/Internal/Framework/Ticker/AvaTickerComponent.h b/Engine/Plugins/VirtualProduction/Avalanche/Source/Avalanche/Internal/Framework/Ticker/AvaTickerComponent.h
new file mode 100644
--- /dev/null
+++ b/Engine/Plugins/VirtualProduction/Avalanche/Source/Avalanche/Internal/Framework/Ticker/AvaTickerComponent.h
@@ -0,0 +1,137 @@
+// Copyright Epic Games, Inc. All Rights Reserved.
+
+#pragma once
+
+#include "Components/PrimitiveComponent.h"
+#include "AvaTickerComponent.generated.h"
+
+#define UE_API AVALANCHE_API
+
+USTRUCT()
+struct FAvaTickerElement
+{
+	GENERATED_BODY()
+
+	FAvaTickerElement() = default;
+
+	UE_API explicit FAvaTickerElement(AActor* InActor);
+
+	UE_API bool operator==(const FAvaTickerElement& InOther) const;
+
+	/** Sets the game and editor visibility of the element */
+	void SetVisibility(bool bInVisible);
+
+	/** Retrieves the world space location of the element */
+	FVector GetElementLocation() const;
+
+	/** Moves the given element to the given world space location */
+	void SetLocation(const FVector& InLocation);
+
+	/** Moves the given element by the given displacement vector */
+	void Displace(const FVector& InDisplacement);
+
+	/** Destroys the element */
+	void Destroy();
+
+	/** Gets the bounds for the element */
+	void GetBounds(FVector& OutOrigin, FVector& OutExtents) const;
+
+	/** Actor mapped to this ticker element */
+	UPROPERTY()
+	TObjectPtr<AActor> Actor;
+};
+
+UCLASS(MinimalAPI, Blueprintable, DisplayName="Motion Design Ticker Component")
+class UAvaTickerComponent : public UPrimitiveComponent
+{
+	GENERATED_BODY()
+
+public:
+	UE_API UAvaTickerComponent();
+
+	/** Adds the given actor to the queue */
+	UFUNCTION(BlueprintCallable, Category="Ticker")
+	UE_API void QueueActor(AActor* InActor);
+
+	UFUNCTION(BlueprintPure, Category="Ticker")
+	FVector GetStartLocation() const
+	{
+		return StartLocation;
+	}
+
+	UFUNCTION(BlueprintCallable, Category="Ticker")
+	UE_API void SetStartLocation(const FVector& InStartLocation);
+
+	UFUNCTION(BlueprintPure, Category="Ticker")
+	double GetDestroyDistance() const
+	{
+		return DestroyDistance;
+	}
+
+	UFUNCTION(BlueprintCallable, Category="Ticker")
+	UE_API void SetDestroyDistance(double InDestroyDistance);
+
+	UFUNCTION(BlueprintPure, Category="Ticker")
+	FVector GetVelocity() const
+	{
+		return Velocity;
+	}
+
+	UFUNCTION(BlueprintCallable, Category="Ticker")
+	UE_API void SetVelocity(const FVector& InVelocity);
+
+	UFUNCTION(BlueprintPure, Category="Ticker")
+	double GetPadding() const
+	{
+		return Padding;
+	}
+
+	UFUNCTION(BlueprintCallable, Category="Ticker")
+	UE_API void SetPadding(double InPadding);
+
+	//~ Begin UActorComponent
+	UE_API virtual void TickComponent(float InDeltaTime, ELevelTick InTickType, FActorComponentTickFunction* InThisTickFunction) override;
+	//~ End UActorComponent
+
+protected:
+	/** Retrieves the start location transformed to World Space */
+	FVector GetWorldStartLocation() const;
+
+	/** Retrieves the displacement transformed to World Space */
+	FVector GetWorldDisplacement() const;
+
+	/** Retrieves the velocity transformed to World Space */
+	FVector GetWorldVelocity() const;
+
+	/** Removes the first element from the queue and activates it if the last active element isn't on the way */
+	void TryDequeueElement();
+
+	/** Updates the active elements to move with the given velocity */
+	void UpdateActiveElements(float InDeltaTime);
+
+	/** Elements being moved by this ticker */
+	UPROPERTY()
+	TArray<FAvaTickerElement> ActiveElements;
+
+	/** Elements currently hidden, awaiting to become available */
+	UPROPERTY()
+	TArray<FAvaTickerElement> QueuedElements;
+
+	/** The start location (in local space) where active elements start at */
+	UPROPERTY(EditAnywhere, Category="Ticker", meta=(AllowPrivateAccess="true"))
+	FVector StartLocation = FVector(0, 500, 0);
+
+	/** The distance the last part of the active element must cross before it gets destroyed */
+	UPROPERTY(EditAnywhere, Category="Ticker", meta=(AllowPrivateAccess="true"))
+	double DestroyDistance = 1000.0;
+
+	/** The velocity (in local space) to move the active elements at */
+	UPROPERTY(EditAnywhere, Category="Ticker", meta=(AllowPrivateAccess="true"))
+	FVector Velocity = FVector(0, -200, 0);
+
+	/** Padding between the active ticker elements */
+	UPROPERTY(EditAnywhere, Category="Ticker", meta=(AllowPrivateAccess="true"))
+	double Padding = 0.0;
+};
+
+#undef UE_API
diff --git a/Engine/Plugins/VirtualProduction/Avalanche/Source/Avalanche/Private/Framework/Ticker/AvaTickerActor.cpp b/Engine/Plugins/VirtualProduction/Avalanche/Source/Avalanche/Private/Framework/Ticker/AvaTickerActor.cpp
new file mode 100644
--- /dev/null
+++ b/Engine/Plugins/VirtualProduction/Avalanche/Source/Avalanche/Private/Framework/Ticker/AvaTickerActor.cpp
@@ -0,0 +1,16 @@
+// Copyright Epic Games, Inc. All Rights Reserved.
+
+#include "Framework/Ticker/AvaTickerActor.h"
+#include "Framework/Ticker/AvaTickerComponent.h"
+
+AAvaTickerActor::AAvaTickerActor()
+{
+	TickerComponent = CreateDefaultSubobject<UAvaTickerComponent>(TEXT("TickerComponent"));
+}
+
+#if WITH_EDITOR
+FString AAvaTickerActor::GetDefaultActorLabel() const
+{
+	return TEXT("Motion Design Ticker Actor");
+}
+#endif
diff --git a/Engine/Plugins/VirtualProduction/Avalanche/Source/Avalanche/Private/Framework/Ticker/AvaTickerComponent.cpp b/Engine/Plugins/VirtualProduction/Avalanche/Source/Avalanche/Private/Framework/Ticker/AvaTickerComponent.cpp
new file mode 100644
--- /dev/null
+++ b/Engine/Plugins/VirtualProduction/Avalanche/Source/Avalanche/Private/Framework/Ticker/AvaTickerComponent.cpp
@@ -0,0 +1,222 @@
+// Copyright Epic Games, Inc. All Rights Reserved.
+
+#include "Framework/Ticker/AvaTickerComponent.h"
+#include "Misc/EnumerateRange.h"
+
+namespace UE::Avalanche::Private
+{
+	/** For each vector component returns 1 if positive, -1 if negative and 0 if 0. */
+	FVector GetSignVector(const FVector& InVector)
+	{
+		// Not using FVector::GetSignVector as it only returns +1 or -1
+		return FVector(FMath::Sign(InVector.X), FMath::Sign(InVector.Y), FMath::Sign(InVector.Z));
+	}
+}
+
+FAvaTickerElement::FAvaTickerElement(AActor* InActor)
+	: Actor(InActor)
+{
+}
+
+bool FAvaTickerElement::operator==(const FAvaTickerElement& InOther) const
+{
+	return Actor == InOther.Actor;
+}
+
+void FAvaTickerElement::SetVisibility(bool bInVisible)
+{
+	Actor->SetActorHiddenInGame(!bInVisible);
+#if WITH_EDITOR
+	Actor->SetIsTemporarilyHiddenInEditor(!bInVisible);
+#endif
+}
+
+FVector FAvaTickerElement::GetElementLocation() const
+{
+	return Actor->GetActorLocation();
+}
+
+void FAvaTickerElement::SetLocation(const FVector& InLocation)
+{
+	Actor->SetActorLocation(InLocation, /*bSweep*/false, /*OutSweepHitResult*/nullptr, ETeleportType::ResetPhysics);
+}
+
+void FAvaTickerElement::Displace(const FVector& InDisplacement)
+{
+	Actor->AddActorWorldOffset(InDisplacement);
+}
+
+void FAvaTickerElement::Destroy()
+{
+	Actor->Destroy();
+}
+
+void FAvaTickerElement::GetBounds(FVector& OutOrigin, FVector& OutExtents) const
+{
+	constexpr bool bOnlyCollidingComponents = false;
+	constexpr bool bIncludeFromChildActors = true;
+	Actor->GetActorBounds(bOnlyCollidingComponents, OutOrigin, OutExtents, bIncludeFromChildActors);
+}
+
+UAvaTickerComponent::UAvaTickerComponent()
+{
+	PrimaryComponentTick.bCanEverTick = true;
+	PrimaryComponentTick.bStartWithTickEnabled = true;
+}
+
+void UAvaTickerComponent::QueueActor(AActor* InActor)
+{
+	if (!InActor)
+	{
+		return;
+	}
+
+	FAvaTickerElement Element(InActor);
+	Element.SetVisibility(/*bVisible*/false);
+	QueuedElements.AddUnique(MoveTemp(Element));
+}
+
+void UAvaTickerComponent::SetStartLocation(const FVector& InStartLocation)
+{
+	StartLocation = InStartLocation;
+}
+
+void UAvaTickerComponent::SetDestroyDistance(double InDestroyDistance)
+{
+	DestroyDistance = InDestroyDistance;
+}
+
+void UAvaTickerComponent::SetVelocity(const FVector& InVelocity)
+{
+	Velocity = InVelocity;
+}
+
+void UAvaTickerComponent::SetPadding(double InPadding)
+{
+	Padding = InPadding;
+}
+
+void UAvaTickerComponent::TickComponent(float InDeltaTime, ELevelTick InTickType, FActorComponentTickFunction* InThisTickFunction)
+{
+	Super::TickComponent(InDeltaTime, InTickType, InThisTickFunction);
+	TryDequeueElement();
+	UpdateActiveElements(InDeltaTime);
+}
+
+FVector UAvaTickerComponent::GetWorldStartLocation() const
+{
+	return GetComponentTransform().TransformPosition(StartLocation);
+}
+
+FVector UAvaTickerComponent::GetWorldDisplacement() const
+{
+	return GetComponentTransform().TransformVector(Velocity.GetSafeNormal() * DestroyDistance);
+}
+
+FVector UAvaTickerComponent::GetWorldVelocity() const
+{
+	// Todo: Consider giving an option whether to scale the velocity or not?
+	return GetComponentTransform().TransformVectorNoScale(Velocity);
+}
+
+void UAvaTickerComponent::TryDequeueElement()
+{
+	using namespace UE::Avalanche;
+
+	if (QueuedElements.IsEmpty())
+	{
+		return;
+	}
+
+	const FAvaTickerElement& Element = QueuedElements[0];
+
+	FVector ElementOrigin;
+	FVector ElementExtents;
+	Element.GetBounds(ElementOrigin, ElementExtents);
+
+	const FVector WorldStartLocation = GetWorldStartLocation();
+	const FVector NormalizedVelocity = GetWorldVelocity().GetSafeNormal();
+
+	// Switch the sign extents so that it can be projected onto the velocity direction
+	const FVector SignedExtents = ElementExtents * Private::GetSignVector(NormalizedVelocity);
+	const FVector StartExtentsToOrigin = -1.0 * (ElementOrigin + NormalizedVelocity * SignedExtents.ProjectOnTo(NormalizedVelocity).Length());
+	const FVector OriginToLocation = Element.GetElementLocation() - ElementOrigin;
+
+	// Check if the last active element is in the way
+	if (!ActiveElements.IsEmpty())
+	{
+		FVector TargetOrigin;
+		FVector TargetExtents;
+		ActiveElements.Last().GetBounds(TargetOrigin, TargetExtents);
+
+		// Check if the queued element bounds is clear from any overlap with the last active element's bounds
+		if (FBox::BuildAABB(WorldStartLocation + StartExtentsToOrigin, ElementExtents).Overlap(FBox::BuildAABB(TargetOrigin, TargetExtents)).IsValid)
+		{
+			return;
+		}
+	}
+
+	FVector ElementStartLocation = WorldStartLocation + StartExtentsToOrigin + OriginToLocation;
+
+	// Add padding in the opposite direction of world velocity
+	ElementStartLocation += -1.0 * NormalizedVelocity * Padding;
+
+	FAvaTickerElement DequeuedElement = MoveTemp(QueuedElements[0]);
+	QueuedElements.RemoveAt(0);
+
+	DequeuedElement.SetLocation(ElementStartLocation);
+	DequeuedElement.SetVisibility(/*bVisible*/true);
+
+	ActiveElements.Add(MoveTemp(DequeuedElement));
+}
+
+void UAvaTickerComponent::UpdateActiveElements(float InDeltaTime)
+{
+	using namespace UE::Avalanche;
+
+	if (ActiveElements.IsEmpty())
+	{
+		return;
+	}
+
+	const FVector WorldVelocity = GetWorldVelocity();
+	const FVector Displacement = WorldVelocity * InDeltaTime;
+
+	for (FAvaTickerElement& ActiveElement : ActiveElements)
+	{
+		ActiveElement.Displace(Displacement);
+	}
+
+	const FVector WorldStartLocation = GetWorldStartLocation();
+	const FVector NormalizedVelocity = WorldVelocity.GetSafeNormal();
+	const double DistanceSquared = GetWorldDisplacement().SquaredLength();
+
+	int32 IndexRangeToRemove = INDEX_NONE;
+
+	// Find the last active element that is more than the given distance away from 
+	for (TEnumerateRef<FAvaTickerElement> ActiveElement : EnumerateRange(ActiveElements))
+	{
+		FVector ElementOrigin;
+		FVector ElementExtents;
+		ActiveElement->GetBounds(ElementOrigin, ElementExtents);
+
+		// Switch the sign extents so that it can be projected onto the velocity direction
+		const FVector SignedExtents = ElementExtents * Private::GetSignVector(NormalizedVelocity);
+		const FVector EndExtents = ElementOrigin + (-1.0 * NormalizedVelocity * SignedExtents.ProjectOnTo(NormalizedVelocity).Length());
+		if (FVector::DistSquared(WorldStartLocation, EndExtents) >= DistanceSquared)
+		{
+			IndexRangeToRemove = ActiveElement.GetIndex();
+		}
+		else
+		{
+			// The elements coming after will all still be within range
+			break;
+		}
+	}
+
+	while (IndexRangeToRemove-- >= 0)
+	{
+		ActiveElements[0].Destroy();
+		ActiveElements.RemoveAt(0);
+	}
+}
