Motion Design: fixed actor null check

diff --git a/Engine/Plugins/VirtualProduction/Avalanche/Source/Avalanche/Private/Framework/Ticker/AvaTickerComponent.cpp b/Engine/Plugins/VirtualProduction/Avalanche/Source/Avalanche/Private/Framework/Ticker/AvaTickerComponent.cpp
--- a/Engine/Plugins/VirtualProduction/Avalanche/Source/Avalanche/Private/Framework/Ticker/AvaTickerComponent.cpp
+++ b/Engine/Plugins/VirtualProduction/Avalanche/Source/Avalanche/Private/Framework/Ticker/AvaTickerComponent.cpp
@@ -25,37 +25,61 @@ bool FAvaTickerElement::operator==(const FAvaTickerElement& InOther) const
 
 void FAvaTickerElement::SetVisibility(bool bInVisible)
 {
-	Actor->SetActorHiddenInGame(!bInVisible);
+	if (Actor)
+	{
+		Actor->SetActorHiddenInGame(!bInVisible);
 #if WITH_EDITOR
-	Actor->SetIsTemporarilyHiddenInEditor(!bInVisible);
+    	Actor->SetIsTemporarilyHiddenInEditor(!bInVisible);
 #endif
+	}
 }
 
 FVector FAvaTickerElement::GetElementLocation() const
 {
-	return Actor->GetActorLocation();
+	if (Actor)
+	{
+		return Actor->GetActorLocation();
+	}
+	return FVector::ZeroVector;
 }
 
 void FAvaTickerElement::SetLocation(const FVector& InLocation)
 {
-	Actor->SetActorLocation(InLocation, /*bSweep*/false, /*OutSweepHitResult*/nullptr, ETeleportType::ResetPhysics);
+	if (Actor)
+	{
+		Actor->SetActorLocation(InLocation, /*bSweep*/false, /*OutSweepHitResult*/nullptr, ETeleportType::ResetPhysics);
+	}
 }
 
 void FAvaTickerElement::Displace(const FVector& InDisplacement)
 {
-	Actor->AddActorWorldOffset(InDisplacement);
+	if (Actor)
+	{
+		Actor->AddActorWorldOffset(InDisplacement);
+	}
 }
 
 void FAvaTickerElement::Destroy()
 {
-	Actor->Destroy();
+	if (Actor)
+	{
+		Actor->Destroy();
+	}
 }
 
 void FAvaTickerElement::GetBounds(FVector& OutOrigin, FVector& OutExtents) const
 {
-	constexpr bool bOnlyCollidingComponents = false;
-	constexpr bool bIncludeFromChildActors = true;
-	Actor->GetActorBounds(bOnlyCollidingComponents, OutOrigin, OutExtents, bIncludeFromChildActors);
+	if (Actor)
+	{
+		constexpr bool bOnlyCollidingComponents = false;
+    	constexpr bool bIncludeFromChildActors = true;
+    	Actor->GetActorBounds(bOnlyCollidingComponents, OutOrigin, OutExtents, bIncludeFromChildActors);
+	}
+	else
+	{
+		OutOrigin = FVector::ZeroVector;
+		OutExtents = FVector::ZeroVector;
+	}
 }
 
 UAvaTickerComponent::UAvaTickerComponent()
