Motion Design: added scene proxy for ticker + fixed issue where actors with bounds larger than destroy distance would be destroyed immediately

diff --git a/Engine/Plugins/VirtualProduction/Avalanche/Source/Avalanche/Internal/Framework/Ticker/AvaTickerComponent.h b/Engine/Plugins/VirtualProduction/Avalanche/Source/Avalanche/Internal/Framework/Ticker/AvaTickerComponent.h
--- a/Engine/Plugins/VirtualProduction/Avalanche/Source/Avalanche/Internal/Framework/Ticker/AvaTickerComponent.h
+++ b/Engine/Plugins/VirtualProduction/Avalanche/Source/Avalanche/Internal/Framework/Ticker/AvaTickerComponent.h
@@ -89,6 +89,14 @@ public:
 	UFUNCTION(BlueprintCallable, Category="Ticker")
 	UE_API void SetPadding(double InPadding);
 
+	//~ Begin UPrimitiveComponent
+	UE_API virtual FPrimitiveSceneProxy* CreateSceneProxy() override;
+	//~ End UPrimitiveComponent
+
+	//~ Begin USceneComponent
+	UE_API virtual FBoxSphereBounds CalcBounds(const FTransform& InLocalToWorld) const override;
+	//~ End USceneComponent
+
 	//~ Begin UActorComponent
 	UE_API virtual void TickComponent(float InDeltaTime, ELevelTick InTickType, FActorComponentTickFunction* InThisTickFunction) override;
 	//~ End UActorComponent
@@ -97,8 +105,8 @@ protected:
 	/** Retrieves the start location transformed to World Space */
 	FVector GetWorldStartLocation() const;
 
-	/** Retrieves the displacement transformed to World Space */
-	FVector GetWorldDisplacement() const;
+	/** Retrieves the destroy displacement transformed to World Space */
+	FVector GetWorldDestroyDisplacement() const;
 
 	/** Retrieves the velocity transformed to World Space */
 	FVector GetWorldVelocity() const;
@@ -110,11 +118,11 @@ protected:
 	void UpdateActiveElements(float InDeltaTime);
 
 	/** Elements being moved by this ticker */
-	UPROPERTY()
+	UPROPERTY(Transient, DuplicateTransient, TextExportTransient)
 	TArray<FAvaTickerElement> ActiveElements;
 
 	/** Elements currently hidden, awaiting to become available */
-	UPROPERTY()
+	UPROPERTY(Transient, DuplicateTransient, TextExportTransient)
 	TArray<FAvaTickerElement> QueuedElements;
 
 	/** The start location (in local space) where active elements start at */
@@ -131,7 +139,7 @@ protected:
 
 	/** Padding between the active ticker elements */
 	UPROPERTY(EditAnywhere, Category="Ticker", meta=(AllowPrivateAccess="true"))
-	double Padding = 0.0;
+	double Padding = 30.0;
 };
 
 #undef UE_API
diff --git a/Engine/Plugins/VirtualProduction/Avalanche/Source/Avalanche/Private/Framework/Ticker/AvaTickerComponent.cpp b/Engine/Plugins/VirtualProduction/Avalanche/Source/Avalanche/Private/Framework/Ticker/AvaTickerComponent.cpp
--- a/Engine/Plugins/VirtualProduction/Avalanche/Source/Avalanche/Private/Framework/Ticker/AvaTickerComponent.cpp
+++ b/Engine/Plugins/VirtualProduction/Avalanche/Source/Avalanche/Private/Framework/Ticker/AvaTickerComponent.cpp
@@ -1,6 +1,7 @@
 // Copyright Epic Games, Inc. All Rights Reserved.
 
 #include "Framework/Ticker/AvaTickerComponent.h"
+#include "AvaTickerSceneProxy.h"
 #include "Misc/EnumerateRange.h"
 
 namespace UE::Avalanche::Private
@@ -86,6 +87,7 @@ UAvaTickerComponent::UAvaTickerComponent()
 {
 	PrimaryComponentTick.bCanEverTick = true;
 	PrimaryComponentTick.bStartWithTickEnabled = true;
+	bHiddenInGame = true;
 }
 
 void UAvaTickerComponent::QueueActor(AActor* InActor)
@@ -120,6 +122,17 @@ void UAvaTickerComponent::SetPadding(double InPadding)
 	Padding = InPadding;
 }
 
+FPrimitiveSceneProxy* UAvaTickerComponent::CreateSceneProxy()
+{
+	return new UE::Avalanche::FTickerSceneProxy(this);
+}
+
+FBoxSphereBounds UAvaTickerComponent::CalcBounds(const FTransform& InLocalToWorld) const
+{
+	const FVector HalfDestroyDisplacement = 0.5 * Velocity.GetSafeNormal() * DestroyDistance;
+	return FBoxSphereBounds(FBox::BuildAABB(StartLocation + HalfDestroyDisplacement, HalfDestroyDisplacement)).TransformBy(InLocalToWorld);
+}
+
 void UAvaTickerComponent::TickComponent(float InDeltaTime, ELevelTick InTickType, FActorComponentTickFunction* InThisTickFunction)
 {
 	Super::TickComponent(InDeltaTime, InTickType, InThisTickFunction);
@@ -132,7 +145,7 @@ FVector UAvaTickerComponent::GetWorldStartLocation() const
 	return GetComponentTransform().TransformPosition(StartLocation);
 }
 
-FVector UAvaTickerComponent::GetWorldDisplacement() const
+FVector UAvaTickerComponent::GetWorldDestroyDisplacement() const
 {
 	return GetComponentTransform().TransformVector(Velocity.GetSafeNormal() * DestroyDistance);
 }
@@ -213,7 +226,7 @@ void UAvaTickerComponent::UpdateActiveElements(float InDeltaTime)
 
 	const FVector WorldStartLocation = GetWorldStartLocation();
 	const FVector NormalizedVelocity = WorldVelocity.GetSafeNormal();
-	const double DistanceSquared = GetWorldDisplacement().SquaredLength();
+	const double DistanceSquared = GetWorldDestroyDisplacement().SquaredLength();
 
 	int32 IndexRangeToRemove = INDEX_NONE;
 
@@ -227,7 +240,11 @@ void UAvaTickerComponent::UpdateActiveElements(float InDeltaTime)
 		// Switch the sign extents so that it can be projected onto the velocity direction
 		const FVector SignedExtents = ElementExtents * Private::GetSignVector(NormalizedVelocity);
 		const FVector EndExtents = ElementOrigin + (-1.0 * NormalizedVelocity * SignedExtents.ProjectOnTo(NormalizedVelocity).Length());
-		if (FVector::DistSquared(WorldStartLocation, EndExtents) >= DistanceSquared)
+
+		// Ensure that the end extents has passed the destroy distance by checking distance
+		// while also preventing cases where an actors extents are large so that its end extents are at a distance larger than the destroy distance
+		if (FVector::DistSquared(WorldStartLocation, EndExtents) >= DistanceSquared
+			&& FVector::DotProduct(EndExtents - WorldStartLocation, NormalizedVelocity) >= 0)
 		{
 			IndexRangeToRemove = ActiveElement.GetIndex();
 		}
diff --git a/Engine/Plugins/VirtualProduction/Avalanche/Source/Avalanche/Private/Framework/Ticker/AvaTickerSceneProxy.cpp b/Engine/Plugins/VirtualProduction/Avalanche/Source/Avalanche/Private/Framework/Ticker/AvaTickerSceneProxy.cpp
new file mode 100644
--- /dev/null
+++ b/Engine/Plugins/VirtualProduction/Avalanche/Source/Avalanche/Private/Framework/Ticker/AvaTickerSceneProxy.cpp
@@ -0,0 +1,234 @@
+// Copyright Epic Games, Inc. All Rights Reserved.
+
+#include "AvaTickerSceneProxy.h"
+#include "Engine/Engine.h"
+#include "Framework/Ticker/AvaTickerComponent.h"
+#include "Materials/Material.h"
+#include "Materials/MaterialRenderProxy.h"
+#include "SceneInterface.h"
+#include "SceneView.h"
+
+namespace UE::Avalanche
+{
+
+FTickerSceneProxy::FTickerSceneProxy(UAvaTickerComponent* InComponent)
+	: FPrimitiveSceneProxy(InComponent)
+	, ArrowVertexFactory(GetScene().GetFeatureLevel(), "FAvaTickerSceneProxyArrow")
+	, EndPlaneVertexFactory(GetScene().GetFeatureLevel(), "FAvaTickerSceneProxyEndPlane")
+	, StartLocation(InComponent->GetStartLocation())
+	, Velocity(InComponent->GetVelocity())
+	, DestroyDistance(InComponent->GetDestroyDistance())
+	, DrawColor(GEngine ? GEngine->C_AddWire : FColor(127, 127, 255, 255))
+{
+	bWillEverBeLit = false;
+
+	LocalArrowMatrix = FRotationMatrix::MakeFromX(Velocity);
+	LocalArrowMatrix.SetOrigin(StartLocation);
+
+	LocalEndPlaneMatrix = FRotationMatrix::MakeFromX(-Velocity);
+	LocalEndPlaneMatrix.SetOrigin(StartLocation + Velocity.GetSafeNormal() * DestroyDistance);
+
+	BuildArrowMesh();
+	BuildEndPlaneMesh();
+}
+
+FTickerSceneProxy::~FTickerSceneProxy()
+{
+	ArrowVertexBuffers.PositionVertexBuffer.ReleaseResource();
+	ArrowVertexBuffers.StaticMeshVertexBuffer.ReleaseResource();
+	ArrowVertexBuffers.ColorVertexBuffer.ReleaseResource();
+	ArrowIndexBuffer.ReleaseResource();
+	ArrowVertexFactory.ReleaseResource();
+
+	EndPlaneVertexBuffers.PositionVertexBuffer.ReleaseResource();
+	EndPlaneVertexBuffers.StaticMeshVertexBuffer.ReleaseResource();
+	EndPlaneVertexBuffers.ColorVertexBuffer.ReleaseResource();
+	EndPlaneIndexBuffer.ReleaseResource();
+	EndPlaneVertexFactory.ReleaseResource();
+}
+
+SIZE_T FTickerSceneProxy::GetTypeHash() const
+{
+	static SIZE_T UniquePointer;
+	return reinterpret_cast<SIZE_T>(&UniquePointer);
+}
+
+void FTickerSceneProxy::GetDynamicMeshElements(const TArray<const FSceneView*>& InViews, const FSceneViewFamily& InViewFamily, uint32 InVisibilityMap, FMeshElementCollector& InCollector) const
+{
+	QUICK_SCOPE_CYCLE_COUNTER(STAT_AvaTickerSceneProxy_DrawDynamicElements);
+
+	const FMatrix ArrowWorldMatrix = LocalArrowMatrix * GetLocalToWorld();
+	const FMatrix EndPlaneWorldMatrix = LocalEndPlaneMatrix * GetLocalToWorld();
+
+	FColoredMaterialRenderProxy* MaterialRenderProxy = new FColoredMaterialRenderProxy(GEngine->ArrowMaterial->GetRenderProxy()
+		, DrawColor
+		, "GizmoColor");
+
+	InCollector.RegisterOneFrameMaterialProxy(MaterialRenderProxy);
+
+	for (int32 ViewIndex = 0; ViewIndex < InViews.Num(); ViewIndex++)
+	{
+		if (InVisibilityMap & (1 << ViewIndex))
+		{
+			DrawArrowElement(InCollector, ViewIndex, ArrowWorldMatrix, MaterialRenderProxy);
+			DrawEndPlaneElement(InCollector, ViewIndex, EndPlaneWorldMatrix, MaterialRenderProxy);
+		}
+	}
+}
+
+FPrimitiveViewRelevance FTickerSceneProxy::GetViewRelevance(const FSceneView* InView) const
+{
+	FPrimitiveViewRelevance Result;
+	Result.bDrawRelevance = IsShown(InView) && InView->Family->EngineShowFlags.BillboardSprites;
+	Result.bDynamicRelevance = true;
+	Result.bShadowRelevance = IsShadowCast(InView);
+	Result.bEditorPrimitiveRelevance = UseEditorCompositing(InView);
+	Result.bVelocityRelevance = DrawsVelocity() && Result.bOpaque && Result.bRenderInMainPass;
+	return Result;
+}
+
+uint32 FTickerSceneProxy::GetMemoryFootprint() const
+{
+	return sizeof(*this) + GetAllocatedSize();
+}
+
+uint32 FTickerSceneProxy::GetAllocatedSize() const
+{
+	return FPrimitiveSceneProxy::GetAllocatedSize();
+}
+
+void FTickerSceneProxy::BuildArrowMesh()
+{
+	TArray<FDynamicMeshVertex> Vertices;
+
+	constexpr float HeadAngle   = FMath::DegreesToRadians(10.f);
+	constexpr float HeadLength  = 16.f;
+	constexpr float ShaftRadius = 1.f;
+	constexpr float TotalLength = 50.f;
+	constexpr float ShaftLength = TotalLength - HeadLength * 0.5; // 10% overlap between shaft and head
+
+	const FVector ShaftCenter = FVector(-0.5 * ShaftLength - HeadLength, 0.0, 0.0);
+
+	BuildConeVerts(HeadAngle, HeadAngle
+		, -HeadLength
+		, /*XOffset*/0.0
+		, /*Sides*/32
+		, Vertices
+		, ArrowIndexBuffer.Indices);
+
+	BuildCylinderVerts(ShaftCenter
+		, /*X*/FVector(0, 0, 1)
+		, /*Y*/FVector(0, 1, 0)
+		, /*Z*/FVector(1, 0, 0)
+		, ShaftRadius
+		, 0.5f * ShaftLength
+		, /*Sides*/16
+		, Vertices
+		, ArrowIndexBuffer.Indices);
+
+	ArrowVertexBuffers.InitFromDynamicVertex(&ArrowVertexFactory, Vertices);
+
+	// Enqueue initialization of render resource
+	BeginInitResource(&ArrowIndexBuffer);
+}
+
+void FTickerSceneProxy::BuildEndPlaneMesh()
+{
+	constexpr float HalfWidth = 25.f;
+	constexpr float HalfHeight = 25.f;
+
+	const FVector XOffset = FVector::RightVector * HalfWidth;
+	const FVector YOffset = FVector::UpVector * HalfHeight;
+
+	// Calculate verts for a face lying in plane defined by the X and Y vectors
+	TArray<FDynamicMeshVertex> Vertices;
+	Vertices.SetNum(4);
+
+	Vertices[0].Position = static_cast<FVector3f>(-XOffset - YOffset);
+	Vertices[1].Position = static_cast<FVector3f>(XOffset - YOffset);
+	Vertices[2].Position = static_cast<FVector3f>(XOffset + YOffset);
+	Vertices[3].Position = static_cast<FVector3f>(-XOffset + YOffset);
+
+	EndPlaneIndexBuffer.Indices.Reserve(12);
+	EndPlaneIndexBuffer.Indices.Append({0, 1, 2});
+	EndPlaneIndexBuffer.Indices.Append({0, 2, 3});
+	EndPlaneIndexBuffer.Indices.Append({0, 2, 1});
+	EndPlaneIndexBuffer.Indices.Append({0, 3, 2});
+
+	EndPlaneVertexBuffers.InitFromDynamicVertex(&EndPlaneVertexFactory, Vertices);
+
+	// Enqueue initialization of render resource
+	BeginInitResource(&EndPlaneIndexBuffer);
+}
+
+void FTickerSceneProxy::DrawArrowElement(FMeshElementCollector& InCollector, int32 InViewIndex, const FMatrix& InLocalToWorld, FMaterialRenderProxy* InMaterialRenderProxy) const
+{
+	FMeshBatch& Mesh = InCollector.AllocateMesh();
+	Mesh.bWireframe = false;
+	Mesh.VertexFactory = &ArrowVertexFactory;
+	Mesh.MaterialRenderProxy = InMaterialRenderProxy;
+	Mesh.ReverseCulling = IsLocalToWorldDeterminantNegative();
+	Mesh.Type = PT_TriangleList;
+	Mesh.DepthPriorityGroup = SDPG_World;
+	Mesh.bCanApplyViewModeOverrides = false;
+
+	constexpr bool bShouldReceiveDecals = false;
+	constexpr bool bHasPrecomputedVolumetricLightmap = false;
+
+	FDynamicPrimitiveUniformBuffer& DynamicPrimitiveUniformBuffer = InCollector.AllocateOneFrameResource<FDynamicPrimitiveUniformBuffer>();
+	DynamicPrimitiveUniformBuffer.Set(InCollector.GetRHICommandList()
+		, InLocalToWorld
+		, InLocalToWorld
+		, GetBounds()
+		, GetLocalBounds()
+		, bShouldReceiveDecals
+		, bHasPrecomputedVolumetricLightmap
+		, AlwaysHasVelocity());
+
+	FMeshBatchElement& BatchElement = Mesh.Elements[0];
+	BatchElement.IndexBuffer = &ArrowIndexBuffer;
+	BatchElement.PrimitiveUniformBufferResource = &DynamicPrimitiveUniformBuffer.UniformBuffer;
+	BatchElement.FirstIndex = 0;
+	BatchElement.NumPrimitives = ArrowIndexBuffer.Indices.Num() / 3;
+	BatchElement.MinVertexIndex = 0;
+	BatchElement.MaxVertexIndex = ArrowVertexBuffers.PositionVertexBuffer.GetNumVertices() - 1;
+
+	InCollector.AddMesh(InViewIndex, Mesh);
+}
+
+void FTickerSceneProxy::DrawEndPlaneElement(FMeshElementCollector& InCollector, int32 InViewIndex, const FMatrix& InLocalToWorld, FMaterialRenderProxy* InMaterialRenderProxy) const
+{
+	FMeshBatch& Mesh = InCollector.AllocateMesh();
+	Mesh.bWireframe = true;
+	Mesh.VertexFactory = &EndPlaneVertexFactory;
+	Mesh.MaterialRenderProxy = InMaterialRenderProxy;
+	Mesh.ReverseCulling = IsLocalToWorldDeterminantNegative();
+	Mesh.Type = PT_TriangleList;
+	Mesh.DepthPriorityGroup = SDPG_World;
+	Mesh.bCanApplyViewModeOverrides = false;
+
+	constexpr bool bShouldReceiveDecals = false;
+	constexpr bool bHasPrecomputedVolumetricLightmap = false;
+
+	FDynamicPrimitiveUniformBuffer& DynamicPrimitiveUniformBuffer = InCollector.AllocateOneFrameResource<FDynamicPrimitiveUniformBuffer>();
+	DynamicPrimitiveUniformBuffer.Set(InCollector.GetRHICommandList()
+		, InLocalToWorld
+		, InLocalToWorld
+		, GetBounds()
+		, GetLocalBounds()
+		, bShouldReceiveDecals
+		, bHasPrecomputedVolumetricLightmap
+		, AlwaysHasVelocity());
+
+	FMeshBatchElement& BatchElement = Mesh.Elements[0];
+	BatchElement.IndexBuffer = &EndPlaneIndexBuffer;
+	BatchElement.PrimitiveUniformBufferResource = &DynamicPrimitiveUniformBuffer.UniformBuffer;
+	BatchElement.FirstIndex = 0;
+	BatchElement.NumPrimitives = EndPlaneIndexBuffer.Indices.Num() / 3;
+	BatchElement.MinVertexIndex = 0;
+	BatchElement.MaxVertexIndex = EndPlaneVertexBuffers.PositionVertexBuffer.GetNumVertices() - 1;
+
+	InCollector.AddMesh(InViewIndex, Mesh);
+}
+
+} // UE::Avalanche
diff --git a/Engine/Plugins/VirtualProduction/Avalanche/Source/Avalanche/Private/Framework/Ticker/AvaTickerSceneProxy.h b/Engine/Plugins/VirtualProduction/Avalanche/Source/Avalanche/Private/Framework/Ticker/AvaTickerSceneProxy.h
new file mode 100644
--- /dev/null
+++ b/Engine/Plugins/VirtualProduction/Avalanche/Source/Avalanche/Private/Framework/Ticker/AvaTickerSceneProxy.h
@@ -0,0 +1,68 @@
+// Copyright Epic Games, Inc. All Rights Reserved.
+
+#pragma once
+
+#include "DynamicMeshBuilder.h"
+#include "PrimitiveSceneProxy.h"
+#include "StaticMeshResources.h"
+
+class UAvaTickerComponent;
+
+namespace UE::Avalanche
+{
+
+/** Represents a UAvaTickerComponent to the scene manager. */
+class FTickerSceneProxy final : public FPrimitiveSceneProxy
+{
+public:
+	explicit FTickerSceneProxy(UAvaTickerComponent* InComponent);
+
+	virtual ~FTickerSceneProxy() override;
+
+	//~ Begin FPrimitiveSceneProxy
+	virtual SIZE_T GetTypeHash() const override;
+	virtual void GetDynamicMeshElements(const TArray<const FSceneView*>& InViews, const FSceneViewFamily& InViewFamily, uint32 InVisibilityMap, FMeshElementCollector& InCollector) const override;
+	virtual FPrimitiveViewRelevance GetViewRelevance(const FSceneView* InView) const override;
+	virtual uint32 GetMemoryFootprint() const override;
+	uint32 GetAllocatedSize() const;
+	//~ End FPrimitiveSceneProxy
+
+private:
+	void BuildArrowMesh();
+
+	void BuildEndPlaneMesh();
+
+	/** Gets the dynamic mesh element for the start arrow */
+	void DrawArrowElement(FMeshElementCollector& InCollector, int32 InViewIndex, const FMatrix& InLocalToWorld, FMaterialRenderProxy* InMaterialRenderProxy) const;
+
+	/** Gets the dynamic mesh element for the ending plane */
+	void DrawEndPlaneElement(FMeshElementCollector& InCollector, int32 InViewIndex, const FMatrix& InLocalToWorld, FMaterialRenderProxy* InMaterialRenderProxy) const;
+
+	FStaticMeshVertexBuffers ArrowVertexBuffers;
+	FDynamicMeshIndexBuffer32 ArrowIndexBuffer;
+	FLocalVertexFactory ArrowVertexFactory;
+
+	FStaticMeshVertexBuffers EndPlaneVertexBuffers;
+	FDynamicMeshIndexBuffer32 EndPlaneIndexBuffer;
+	FLocalVertexFactory EndPlaneVertexFactory;
+
+	/** The start location (in local space) where active elements start at */
+	FVector StartLocation;
+
+	/** The velocity (in local space) to move the active elements at */
+	FVector Velocity;
+
+	/** The distance the last part of the active element must cross before it gets destroyed */
+	double DestroyDistance;
+
+	/** Draw color for the arrow */
+	FColor DrawColor;
+
+	/** Transformation of StartLocation and Velocity in Local Space */
+	FMatrix LocalArrowMatrix;
+
+	/** Transformation of End Location and Inverse Velocity in Local Space */
+	FMatrix LocalEndPlaneMatrix;
+};
+
+} // UE::Avalanche
