MotionDesign : ActorModifier
- Fixed visibility modifier not respecting nested hierarchy

diff --git a/Engine/Plugins/VirtualProduction/Avalanche/Source/Avalanche/Private/AvaOutlinerUtils.cpp b/Engine/Plugins/VirtualProduction/Avalanche/Source/Avalanche/Private/AvaOutlinerUtils.cpp
--- a/Engine/Plugins/VirtualProduction/Avalanche/Source/Avalanche/Private/AvaOutlinerUtils.cpp
+++ b/Engine/Plugins/VirtualProduction/Avalanche/Source/Avalanche/Private/AvaOutlinerUtils.cpp
@@ -96,11 +96,4 @@ TArray<AActor*> FAvaOutlinerUtils::EditorOutlinerChildActors(TSharedPtr<IAvaOutl
 	return {};
 }
 
-bool FAvaOutlinerUtils::EditorActorIsolationInfo(TSharedPtr<IAvaOutliner> InOutliner, TArray<TWeakObjectPtr<const AActor>>& OutIsolatedActors)
-{
-	// todo
-	OutIsolatedActors.Empty();
-	return false;
-}
-
 #endif // WITH_EDITOR
diff --git a/Engine/Plugins/VirtualProduction/Avalanche/Source/Avalanche/Public/AvaOutlinerUtils.h b/Engine/Plugins/VirtualProduction/Avalanche/Source/Avalanche/Public/AvaOutlinerUtils.h
--- a/Engine/Plugins/VirtualProduction/Avalanche/Source/Avalanche/Public/AvaOutlinerUtils.h
+++ b/Engine/Plugins/VirtualProduction/Avalanche/Source/Avalanche/Public/AvaOutlinerUtils.h
@@ -21,9 +21,6 @@ struct AVALANCHE_API FAvaOutlinerUtils
 	/** Returns an array of direct child actors attached to a parent actor or all root level actors if the parent actor is null. */
 	static TArray<AActor*> EditorOutlinerChildActors(TSharedPtr<IAvaOutliner> InOutliner, AActor* const InParentActor = nullptr);
 
-	/** Returns true if the Motion Design editor is isolating actors and the array of actors currently being isolated. */
-	static bool EditorActorIsolationInfo(TSharedPtr<IAvaOutliner> InOutliner, TArray<TWeakObjectPtr<const AActor>>& OutIsolatedActors);
-
 	/**
 	 * Gets an array of actors from an array of outliner items while maintaining index ordering (for actor outliner items only).
 	 * Returns false and an empty OutActors if any outliner items are invalid or any actor objects they represent are invalid.
diff --git a/Engine/Plugins/VirtualProduction/Avalanche/Source/AvalancheModifiers/Private/AvaModifiersActorUtils.cpp b/Engine/Plugins/VirtualProduction/Avalanche/Source/AvalancheModifiers/Private/AvaModifiersActorUtils.cpp
--- a/Engine/Plugins/VirtualProduction/Avalanche/Source/AvalancheModifiers/Private/AvaModifiersActorUtils.cpp
+++ b/Engine/Plugins/VirtualProduction/Avalanche/Source/AvalancheModifiers/Private/AvaModifiersActorUtils.cpp
@@ -16,32 +16,6 @@
 #include "IAvaOutliner.h"
 #endif
 
-bool FAvaModifiersActorUtils::IsActorNotIsolated(const AActor* InActor)
-{
-	if (!InActor)
-	{
-		return false;
-	}
-
-#if WITH_EDITOR
-	bool bIsIsolatingActors = false;
-	TArray<TWeakObjectPtr<const AActor>> IsolatedActors;
-
-	TSharedPtr<IAvaOutliner> AvaOutliner = FAvaOutlinerUtils::EditorGetOutliner(InActor->GetWorld());
-	if (AvaOutliner.IsValid())
-	{
-		bIsIsolatingActors = FAvaOutlinerUtils::EditorActorIsolationInfo(AvaOutliner, IsolatedActors);
-	}
-
-	if (bIsIsolatingActors && !IsolatedActors.Contains(InActor))
-	{
-		return true;
-	}
-#endif
-
-	return false;
-}
-
 AActor* FAvaModifiersActorUtils::FindActorFromReferenceContainer(AActor* const InActor, const EActorModifierReferenceContainer InReferenceContainer, const bool bInIgnoreHiddenActors)
 {
 	if (InReferenceContainer == EActorModifierReferenceContainer::Other || !IsValid(InActor))
diff --git a/Engine/Plugins/VirtualProduction/Avalanche/Source/AvalancheModifiers/Private/AvaModifiersActorUtils.h b/Engine/Plugins/VirtualProduction/Avalanche/Source/AvalancheModifiers/Private/AvaModifiersActorUtils.h
--- a/Engine/Plugins/VirtualProduction/Avalanche/Source/AvalancheModifiers/Private/AvaModifiersActorUtils.h
+++ b/Engine/Plugins/VirtualProduction/Avalanche/Source/AvalancheModifiers/Private/AvaModifiersActorUtils.h
@@ -17,7 +17,6 @@ struct FActorModifierSceneTreeActor;
 struct FAvaModifiersActorUtils
 {
 	/** Begin Outliner */
-	static bool IsActorNotIsolated(const AActor* InActor);
 	static AActor* FindActorFromReferenceContainer(AActor* const InActor, const EActorModifierReferenceContainer InReferenceContainer, const bool bInIgnoreHiddenActors);
 	static TArray<AActor*> GetReferenceActors(const FActorModifierSceneTreeActor* InTrackedActor);
 	/** End Outliner */
diff --git a/Engine/Plugins/VirtualProduction/Avalanche/Source/AvalancheModifiers/Private/Modifiers/AvaVisibilityModifier.cpp b/Engine/Plugins/VirtualProduction/Avalanche/Source/AvalancheModifiers/Private/Modifiers/AvaVisibilityModifier.cpp
--- a/Engine/Plugins/VirtualProduction/Avalanche/Source/AvalancheModifiers/Private/Modifiers/AvaVisibilityModifier.cpp
+++ b/Engine/Plugins/VirtualProduction/Avalanche/Source/AvalancheModifiers/Private/Modifiers/AvaVisibilityModifier.cpp
@@ -11,15 +11,6 @@
 
 void UAvaVisibilityModifier::Apply()
 {
-	AActor* const ModifyActor = GetModifiedActor();
-
-	// Early exit if the modify actor is NOT being isolated. The outliner will manage the visibility for the actor and it's children.
-	if (FAvaModifiersActorUtils::IsActorNotIsolated(ModifyActor))
-	{
-		Next();
-		return;
-	}
-
 	const FActorModifierSceneTreeUpdateExtension* SceneExtension = GetExtension<FActorModifierSceneTreeUpdateExtension>();
 	if (!SceneExtension)
 	{
@@ -35,6 +26,7 @@ void UAvaVisibilityModifier::Apply()
 		return;
 	}
 
+	AActor* const ModifyActor = GetModifiedActor();
 	const TArray<TWeakObjectPtr<AActor>> AttachedActors = SceneExtension->GetDirectChildrenActor(ModifyActor);
 
 #if WITH_EDITOR
@@ -45,7 +37,7 @@ void UAvaVisibilityModifier::Apply()
 
 	// Top most modifier in tree has priority over this one if it is hiding current one
 	bool bIsNestedVisibilityModifier = false;
-	if (const UAvaVisibilityModifier* VisibilityModifier = GetFirstModifierAbove(ModifyActor))
+	if (const UAvaVisibilityModifier* VisibilityModifier = GetLastModifierAbove(ModifyActor))
 	{
 		// This actor is not visible and hidden by root visibility modifier do not proceed
 		if (VisibilityModifier->IsChildActorHidden(ModifyActor))
@@ -115,8 +107,11 @@ void UAvaVisibilityModifier::Apply()
 		}
 	}
 
-	// Untrack previous actors that are not attached anymore
-	VisibilityShared->RestoreActorsState(this, ChildrenActorsWeak.Difference(NewChildrenActorsWeak));
+	if (!bIsNestedVisibilityModifier)
+	{
+		// Untrack previous actors that are not attached anymore
+		VisibilityShared->RestoreActorsState(this, ChildrenActorsWeak.Difference(NewChildrenActorsWeak));
+	}
 
 	ChildrenActorsWeak = NewChildrenActorsWeak;
 
