MotionDesign : ActorModifier
- Fixed mask modifier on parent and child actor (Text3D) would causes infinite loop and crash in rundown

diff --git a/Engine/Plugins/Experimental/Text3D/Source/Text3D/Private/Text3DComponent.cpp b/Engine/Plugins/Experimental/Text3D/Source/Text3D/Private/Text3DComponent.cpp
--- a/Engine/Plugins/Experimental/Text3D/Source/Text3D/Private/Text3DComponent.cpp
+++ b/Engine/Plugins/Experimental/Text3D/Source/Text3D/Private/Text3DComponent.cpp
@@ -78,6 +78,11 @@ void UText3DComponent::OnComponentDestroyed(bool bInDestroyingHierarchy)
 		TextRenderer->Destroy();
 	}
 
+	if (CharacterExtension)
+	{
+		CharacterExtension->AllocateCharacters(0);
+	}
+
 	TransformUpdated.RemoveAll(this);
 }
 
@@ -100,7 +105,7 @@ void UText3DComponent::RequestUpdate(EText3DRendererFlags InFlags, bool bInImmed
 
 		if (bInImmediate)
 		{
-			Rebuild(InFlags, /*bCleanCache*/false);
+			Rebuild(InFlags);
 		}
 		else
 		{
@@ -257,12 +262,10 @@ void UText3DComponent::OnRegister()
 {
 	Super::OnRegister();
 
-	if (TextRenderer)
+	if (TextRenderer && !TextRenderer->IsInitialized())
 	{
-		TextRenderer->Create();
-		const EText3DRendererFlags Flags = bInitialized ? EText3DRendererFlags::None : EText3DRendererFlags::All;
 		constexpr bool bImmediate = true;
-		RequestUpdate(Flags, bImmediate);
+		RequestUpdate(EText3DRendererFlags::All, bImmediate);
 	}
 }
 
@@ -946,7 +949,7 @@ void UText3DComponent::QueueRebuild(EText3DRendererFlags InFlags)
 					const EText3DRendererFlags CurrentUpdateFlags = UpdateFlags;
 					UpdateFlags = EText3DRendererFlags::None;
 
-					Rebuild(CurrentUpdateFlags, /*bCleanCache*/false);
+					Rebuild(CurrentUpdateFlags);
 
 					// Executes only once
 					return false;
@@ -956,7 +959,7 @@ void UText3DComponent::QueueRebuild(EText3DRendererFlags InFlags)
 	}
 }
 
-void UText3DComponent::Rebuild(EText3DRendererFlags InUpdateFlags, bool bCleanCache)
+void UText3DComponent::Rebuild(EText3DRendererFlags InUpdateFlags)
 {
 	if (InUpdateFlags == EText3DRendererFlags::None)
 	{
@@ -965,7 +968,13 @@ void UText3DComponent::Rebuild(EText3DRendererFlags InUpdateFlags, bool bCleanCa
 
 	if (!TextRenderer)
 	{
-		UE_LOG(LogText3D, Error, TEXT("Text3D renderer is not valid, cannot proceed"));
+		UE_LOG(LogText3D, Error, TEXT("Text3D renderer is not valid or initialized, cannot proceed"));
+		return;
+	}
+
+	if (!IsRegistered())
+	{
+		UE_LOG(LogText3D, Warning, TEXT("Text3DComponent is not registered yet, cannot proceed"));
 		return;
 	}
 
@@ -985,14 +994,22 @@ void UText3DComponent::Rebuild(EText3DRendererFlags InUpdateFlags, bool bCleanCa
 
 	if (UE::IsSavingPackage(this))
 	{
-		UE_LOG(LogText3D, Warning, TEXT("[%s] Rebuild called with update flags (%u) while package is being saved. This new Rebuild will be deferred to the next tick."), *GetReadableName(), InUpdateFlags);
+		if (!EnumHasAllFlags(UpdateFlags, InUpdateFlags))
+		{
+			UE_LOG(LogText3D, Warning, TEXT("[%s] Rebuild called with update flags (%u) while package is being saved. This new Rebuild will be deferred to the next tick."), *GetReadableName(), InUpdateFlags);
+		}
+
 		QueueRebuild(InUpdateFlags);
 		return;
 	}
 
 	if (bIsUpdatingText)
 	{
-		UE_LOG(LogText3D, Warning, TEXT("[%s] Rebuild called with update flags (%u) while building text. This new Rebuild will be deferred to the next tick."), *GetReadableName(), InUpdateFlags);
+		if (!EnumHasAllFlags(UpdateFlags, InUpdateFlags))
+		{
+			UE_LOG(LogText3D, Warning, TEXT("[%s] Rebuild called with update flags (%u) while building text. This new Rebuild will be deferred to the next tick."), *GetReadableName(), InUpdateFlags);
+		}
+
 		QueueRebuild(InUpdateFlags);
 		return;
 	}
@@ -1004,11 +1021,10 @@ void UText3DComponent::Rebuild(EText3DRendererFlags InUpdateFlags, bool bCleanCa
 		bIsUpdatingText = false;
 	};
 
-	if (!bInitialized)
+	if (!TextRenderer->IsInitialized())
 	{
-		UE_LOG(LogText3D, Verbose, TEXT("[%s] Rebuild called with update flags (%u) while text is not initialized yet."), *GetReadableName(), InUpdateFlags);
+		TextRenderer->Create();
 		InUpdateFlags = EText3DRendererFlags::All;
-		bInitialized = true;
 	}
 
 	TextPreUpdateDelegate.Broadcast(this, InUpdateFlags);
@@ -1040,14 +1056,6 @@ void UText3DComponent::Rebuild(EText3DRendererFlags InUpdateFlags, bool bCleanCa
 		MarkRenderStateDirty();
 	}
 
-	if (bCleanCache)
-	{
-		if (UText3DEngineSubsystem* Subsystem = UText3DEngineSubsystem::Get())
-		{
-			Subsystem->Cleanup();
-		}
-	}
-
 	TextPostUpdateDelegate.Broadcast(this, InUpdateFlags);
 }
 
@@ -1087,11 +1095,8 @@ void UText3DComponent::OnTextRendererClassChanged()
 	const AActor* ContextActor = GetOwner();
 	UE_LOG(LogText3D, Log, TEXT("%s : Text3D renderer changed : old %s - new %s"), !!ContextActor ? *ContextActor->GetActorNameOrLabel() : TEXT("Invalid owner"), *OldRendererName.ToString(), *TextRenderer->GetFriendlyName().ToString());
 
-	bInitialized = false;
-
 	if (IsRegistered())
 	{
-		TextRenderer->Create();
 		constexpr bool bImmediate = true;
 		RequestUpdate(EText3DRendererFlags::All, bImmediate);
 	}
diff --git a/Engine/Plugins/Experimental/Text3D/Source/Text3D/Public/Renderers/Text3DRendererBase.h b/Engine/Plugins/Experimental/Text3D/Source/Text3D/Public/Renderers/Text3DRendererBase.h
--- a/Engine/Plugins/Experimental/Text3D/Source/Text3D/Public/Renderers/Text3DRendererBase.h
+++ b/Engine/Plugins/Experimental/Text3D/Source/Text3D/Public/Renderers/Text3DRendererBase.h
@@ -35,7 +35,13 @@ public:
 
 	/** Get the implementation name for debug purposes */
 	virtual FName GetFriendlyName() const PURE_VIRTUAL(UText3DRenderingImplementationBase::GetFriendlyName, return NAME_None; );
-	
+
+	/** Is the renderer initialized (render state created) */
+	bool IsInitialized() const
+	{
+		return bInitialized;
+	}
+
 protected:
 	/** Create and setup the implementation components, called on load or creation */
 	virtual void OnCreate() PURE_VIRTUAL(UText3DRenderingImplementationBase::OnCreate, );
diff --git a/Engine/Plugins/Experimental/Text3D/Source/Text3D/Public/Text3DComponent.h b/Engine/Plugins/Experimental/Text3D/Source/Text3D/Public/Text3DComponent.h
--- a/Engine/Plugins/Experimental/Text3D/Source/Text3D/Public/Text3DComponent.h
+++ b/Engine/Plugins/Experimental/Text3D/Source/Text3D/Public/Text3DComponent.h
@@ -497,7 +497,7 @@ private:
 	void QueueRebuild(EText3DRendererFlags InFlags);
 
 	/** Partially or fully rebuilds text depending on the given update flags */
-	void Rebuild(EText3DRendererFlags InUpdateFlags, bool bCleanCache);
+	void Rebuild(EText3DRendererFlags InUpdateFlags);
 
 	bool IsTypefaceAvailable(FName InTypeface) const;
 	TArrayView<const FTypefaceEntry> GetAvailableTypefaces() const;
@@ -509,9 +509,6 @@ private:
 
 	bool bZeroScaled = false;
 
-	/** Was the text initialized once */
-	bool bInitialized = false;
-
 #if WITH_EDITORONLY_DATA
 	UE_DEPRECATED(5.6, "Use Material Extension instead")
 	UPROPERTY()
diff --git a/Engine/Plugins/VirtualProduction/ActorModifierCore/Source/ActorModifierCore/Private/Modifiers/Utilities/ActorModifierCoreUtilities.cpp b/Engine/Plugins/VirtualProduction/ActorModifierCore/Source/ActorModifierCore/Private/Modifiers/Utilities/ActorModifierCoreUtilities.cpp
new file mode 100644
--- /dev/null
+++ b/Engine/Plugins/VirtualProduction/ActorModifierCore/Source/ActorModifierCore/Private/Modifiers/Utilities/ActorModifierCoreUtilities.cpp
@@ -0,0 +1,31 @@
+// Copyright Epic Games, Inc. All Rights Reserved.
+
+#include "Modifiers/Utilities/ActorModifierCoreUtilities.h"
+
+#include "Modifiers/ActorModifierCoreComponent.h"
+#include "Modifiers/ActorModifierCoreStack.h"
+
+UActorModifierCoreBase* UE::ActorModifierCore::Utilities::FindFirstActorModifierByClass(const AActor* InStartActor, const TSubclassOf<UActorModifierCoreBase>& InModifierClass)
+{
+	UActorModifierCoreBase* FoundModifier = nullptr;
+
+	if (!InStartActor || !InModifierClass)
+	{
+		return FoundModifier;
+	}
+
+	if (const UActorModifierCoreComponent* ModifierComponent = InStartActor->FindComponentByClass<UActorModifierCoreComponent>())
+	{
+		if (const UActorModifierCoreStack* ModifierStack = ModifierComponent->GetModifierStack())
+		{
+			FoundModifier = ModifierStack->FindModifier(InModifierClass.Get());
+		}
+	}
+
+	if (FoundModifier)
+	{
+		return FoundModifier;
+	}
+
+	return FindFirstActorModifierByClass(InStartActor->GetAttachParentActor(), InModifierClass);
+}
diff --git a/Engine/Plugins/VirtualProduction/ActorModifierCore/Source/ActorModifierCore/Public/Modifiers/Utilities/ActorModifierCoreUtilities.h b/Engine/Plugins/VirtualProduction/ActorModifierCore/Source/ActorModifierCore/Public/Modifiers/Utilities/ActorModifierCoreUtilities.h
new file mode 100644
--- /dev/null
+++ b/Engine/Plugins/VirtualProduction/ActorModifierCore/Source/ActorModifierCore/Public/Modifiers/Utilities/ActorModifierCoreUtilities.h
@@ -0,0 +1,11 @@
+// Copyright Epic Games, Inc. All Rights Reserved.
+
+#pragma once
+
+#include "Modifiers/ActorModifierCoreBase.h"
+
+namespace UE::ActorModifierCore::Utilities
+{
+	/** Climbs up the attachment tree to find a modifier with a specific class */
+	ACTORMODIFIERCORE_API UActorModifierCoreBase* FindFirstActorModifierByClass(const AActor* InStartActor, const TSubclassOf<UActorModifierCoreBase>& InModifierClass);
+}
diff --git a/Engine/Plugins/VirtualProduction/Avalanche/Source/AvalancheMask/Private/Handling/AvaMaskText3DActorMaterialCollectionHandle.cpp b/Engine/Plugins/VirtualProduction/Avalanche/Source/AvalancheMask/Private/Handling/AvaMaskText3DActorMaterialCollectionHandle.cpp
--- a/Engine/Plugins/VirtualProduction/Avalanche/Source/AvalancheMask/Private/Handling/AvaMaskText3DActorMaterialCollectionHandle.cpp
+++ b/Engine/Plugins/VirtualProduction/Avalanche/Source/AvalancheMask/Private/Handling/AvaMaskText3DActorMaterialCollectionHandle.cpp
@@ -72,9 +72,13 @@ TArray<TObjectPtr<UMaterialInterface>> FAvaMaskText3DActorMaterialCollectionHand
 
 TArray<TSharedPtr<IAvaMaskMaterialHandle>> FAvaMaskText3DActorMaterialCollectionHandle::GetMaterialHandles()
 {
-	if (MaterialHandles.IsEmpty() || !UE::Ava::Internal::IsHandleValid(MaterialHandles.Last()))
+	if (MaterialHandles.Num() != static_cast<int32>(EText3DGroupType::TypeCount)
+		|| MaterialHandles.ContainsByPredicate([](const TSharedPtr<IAvaMaskMaterialHandle>& InHandle)
+		{
+			return !UE::Ava::Internal::IsHandleValid(InHandle);
+		}))
 	{
-		if (UText3DComponent* TextComponent = GetComponent())
+		if (const UText3DComponent* TextComponent = GetComponent())
 		{
 			UAvaObjectHandleSubsystem* HandleSubsystem = UE::Ava::Internal::GetObjectHandleSubsystem();
 			if (!HandleSubsystem)
@@ -90,10 +94,10 @@ TArray<TSharedPtr<IAvaMaskMaterialHandle>> FAvaMaskText3DActorMaterialCollection
 			UMaterialInterface* BevelMaterial = MaterialExtension->GetMaterial(EText3DGroupType::Bevel);
 
 			MaterialHandles.Init(nullptr, static_cast<int32>(EText3DGroupType::TypeCount));
-			MaterialHandles[static_cast<int32>(EText3DGroupType::Front)] = HandleSubsystem->MakeHandle<IAvaMaskMaterialHandle>(FrontMaterial, UE::AvaMask::Internal::HandleTag);
-			MaterialHandles[static_cast<int32>(EText3DGroupType::Bevel)] = HandleSubsystem->MakeHandle<IAvaMaskMaterialHandle>(BevelMaterial, UE::AvaMask::Internal::HandleTag);
-			MaterialHandles[static_cast<int32>(EText3DGroupType::Extrude)] = HandleSubsystem->MakeHandle<IAvaMaskMaterialHandle>(ExtrudeMaterial, UE::AvaMask::Internal::HandleTag);
-			MaterialHandles[static_cast<int32>(EText3DGroupType::Back)] = HandleSubsystem->MakeHandle<IAvaMaskMaterialHandle>(BackMaterial, UE::AvaMask::Internal::HandleTag);
+			MaterialHandles[static_cast<int32>(EText3DGroupType::Front)] = FrontMaterial != nullptr ? HandleSubsystem->MakeHandle<IAvaMaskMaterialHandle>(FrontMaterial, UE::AvaMask::Internal::HandleTag) : nullptr;
+			MaterialHandles[static_cast<int32>(EText3DGroupType::Bevel)] = BevelMaterial != nullptr ? HandleSubsystem->MakeHandle<IAvaMaskMaterialHandle>(BevelMaterial, UE::AvaMask::Internal::HandleTag) : nullptr;
+			MaterialHandles[static_cast<int32>(EText3DGroupType::Extrude)] = ExtrudeMaterial != nullptr ? HandleSubsystem->MakeHandle<IAvaMaskMaterialHandle>(ExtrudeMaterial, UE::AvaMask::Internal::HandleTag) : nullptr;
+			MaterialHandles[static_cast<int32>(EText3DGroupType::Back)] = BackMaterial != nullptr ? HandleSubsystem->MakeHandle<IAvaMaskMaterialHandle>(BackMaterial, UE::AvaMask::Internal::HandleTag) : nullptr;
 		}
 	}
 
@@ -247,7 +251,12 @@ void FAvaMaskText3DActorMaterialCollectionHandle::ForEachMaterialHandle(TFunctio
 	{
 		// Effectively refreshes material handles if necessary
 		GetMaterialHandles();
-		
+
+		if (MaterialHandles.Num() != static_cast<int32>(EText3DGroupType::TypeCount))
+		{
+			return;
+		}
+
 		const FSoftComponentReference ComponentReference = UE::AvaMask::Internal::MakeComponentReference(TextComponent->GetOwner(), TextComponent);
 		TSharedPtr<IAvaMaskMaterialHandle> MaterialHandle = MaterialHandles[static_cast<int32>(EText3DGroupType::Front)];
 		if (!InFunction(ComponentReference, static_cast<int32>(EText3DGroupType::Front), MaterialHandle.IsValid(), MaterialHandle))
@@ -428,7 +437,7 @@ bool FAvaMaskText3DActorMaterialCollectionHandle::ApplyModifiedState(const FAvaM
 
 				// Refresh materials if needed immediately
 				constexpr bool bImmediateUpdate = true;
-				TextComponent->RequestUpdate(EText3DRendererFlags::None, bImmediateUpdate);
+				TextComponent->RequestUpdate(EText3DRendererFlags::Material, bImmediateUpdate);
 			}
 		}
 
@@ -463,9 +472,9 @@ FStructView FAvaMaskText3DActorMaterialCollectionHandle::GetMaterialHandleData(
 {
 	if (InParentHandleData)
 	{
-		if (ensure(InParentHandleData->GroupMaterialData.Contains(InSlotIdx)))
+		if (FInstancedStruct* Value = InParentHandleData->GroupMaterialData.Find(InSlotIdx))
 		{
-			return InParentHandleData->GroupMaterialData[InSlotIdx];
+			return *Value;
 		}
 	}
 
@@ -513,26 +522,33 @@ void FAvaMaskText3DActorMaterialCollectionHandle::OnTextPostUpdate(UText3DCompon
 		constexpr int32 BackIndex = static_cast<int32>(EText3DGroupType::Back);
 		constexpr int32 ExtrudeIndex = static_cast<int32>(EText3DGroupType::Extrude);
 		constexpr int32 BevelIndex = static_cast<int32>(EText3DGroupType::Bevel);
-		
-		bool bChanged = false;
-		if (MaterialHandles.IsEmpty() || !UE::Ava::Internal::IsHandleValid(MaterialHandles.Last()))
+
+		auto HasInvalidHandles = [this]()
 		{
-			GetMaterialHandles();
-			bChanged = true;
-		}
-		else if (MaterialHandles[FrontIndex]->GetMaterial() != InComponent->GetFrontMaterial()
+			return MaterialHandles.ContainsByPredicate([](const TSharedPtr<IAvaMaskMaterialHandle>& InHandle)
+			{
+				return !UE::Ava::Internal::IsHandleValid(InHandle);
+			});
+		};
+
+		const bool bHasInvalidHandles = HasInvalidHandles();
+	
+		if (MaterialHandles.IsEmpty()
+			|| bHasInvalidHandles
+			|| MaterialHandles[FrontIndex]->GetMaterial() != InComponent->GetFrontMaterial()
 			|| MaterialHandles[BevelIndex]->GetMaterial() != InComponent->GetBevelMaterial()
 			|| MaterialHandles[ExtrudeIndex]->GetMaterial() != InComponent->GetExtrudeMaterial()
 			|| MaterialHandles[BackIndex]->GetMaterial() != InComponent->GetBackMaterial())
 		{
 			MaterialHandles.Empty();
+
+			// refresh handles
 			GetMaterialHandles();
-			bChanged = true;
-		}
 
-		if (bChanged)
-		{
-			OnSourceMaterialsChanged().ExecuteIfBound(InComponent, MaterialHandles);
+			if (!HasInvalidHandles())
+			{
+				OnSourceMaterialsChanged().ExecuteIfBound(InComponent, MaterialHandles);
+			}
 		}
 	}
 }
diff --git a/Engine/Plugins/VirtualProduction/Avalanche/Source/AvalancheMask/Private/Handling/IAvaMaskMaterialCollectionHandle.h b/Engine/Plugins/VirtualProduction/Avalanche/Source/AvalancheMask/Private/Handling/IAvaMaskMaterialCollectionHandle.h
--- a/Engine/Plugins/VirtualProduction/Avalanche/Source/AvalancheMask/Private/Handling/IAvaMaskMaterialCollectionHandle.h
+++ b/Engine/Plugins/VirtualProduction/Avalanche/Source/AvalancheMask/Private/Handling/IAvaMaskMaterialCollectionHandle.h
@@ -229,9 +229,10 @@ bool TAvaMaskMaterialCollectionHandle<HandleDataType>::ApplyOriginalState(const
 			{
 				return true;
 			}
+
+			const FStructView MaterialHandleDataView = GetMaterialHandleData(HandleData, InComponent, InSlotIdx);
 			
-			if (const FStructView MaterialHandleDataView = GetMaterialHandleData(HandleData, InComponent, InSlotIdx);
-				MaterialHandleDataView.IsValid())
+			if (MaterialHandleDataView.IsValid())
 			{
 				return InMaterialHandle->ApplyOriginalState(
 					MaterialHandleDataView
@@ -240,7 +241,9 @@ bool TAvaMaskMaterialCollectionHandle<HandleDataType>::ApplyOriginalState(const
 						SetMaterial(InComponent, InSlotIdx, InMaterial);
 					});
 			}
+
 			UE_LOG(LogAvaMask, Display, TEXT("MaterialData not found for slot %u"), InSlotIdx);
+
 			return false;
 		});
 
@@ -269,9 +272,10 @@ bool TAvaMaskMaterialCollectionHandle<HandleDataType>::ApplyModifiedState(
 			{
 				return true;
 			}
-			
-			if (const FStructView MaterialHandleDataView = GetMaterialHandleData(HandleData, InComponent, InSlotIdx);
-				MaterialHandleDataView.IsValid())
+
+			const FStructView MaterialHandleDataView = GetMaterialHandleData(HandleData, InComponent, InSlotIdx);
+
+			if (MaterialHandleDataView.IsValid())
 			{
 				return InMaterialHandle->ApplyModifiedState(
 					InModifiedParameters
@@ -281,7 +285,9 @@ bool TAvaMaskMaterialCollectionHandle<HandleDataType>::ApplyModifiedState(
 						SetMaterial(InComponent, InSlotIdx, InMaterial);
 					});
 			}
+
 			UE_LOG(LogAvaMask, Display, TEXT("MaterialData not found for slot %u"), InSlotIdx);
+
 			return false;
 		});
 
diff --git a/Engine/Plugins/VirtualProduction/Avalanche/Source/AvalancheMask/Private/Mask2D/AvaMask2DBaseModifier.cpp b/Engine/Plugins/VirtualProduction/Avalanche/Source/AvalancheMask/Private/Mask2D/AvaMask2DBaseModifier.cpp
--- a/Engine/Plugins/VirtualProduction/Avalanche/Source/AvalancheMask/Private/Mask2D/AvaMask2DBaseModifier.cpp
+++ b/Engine/Plugins/VirtualProduction/Avalanche/Source/AvalancheMask/Private/Mask2D/AvaMask2DBaseModifier.cpp
@@ -24,6 +24,7 @@
 #include "Materials/Material.h"
 #include "Materials/MaterialInstanceDynamic.h"
 #include "Modifiers/ActorModifierCoreStack.h"
+#include "Modifiers/Utilities/ActorModifierCoreUtilities.h"
 #include "Particles/ParticleSystemComponent.h"
 #include "Subsystems/ActorModifierCoreSubsystem.h"
 #include "Text3DComponent.h"
@@ -352,7 +353,14 @@ void UAvaMask2DBaseModifier::SavePreState()
 	ForEachActor<AActor>(
 		[this, &NewActorDatas](AActor* InActor)
 		{
-			SaveActorPreState(InActor, NewActorDatas.Emplace(InActor, ActorData.FindOrAdd(InActor)));
+			const UActorModifierCoreBase* MaskModifier = UE::ActorModifierCore::Utilities::FindFirstActorModifierByClass(InActor, GetClass());
+
+			// If first modifier in the hierarchy is not this modifier, then skip actor
+			if (!MaskModifier || MaskModifier == this)
+			{
+				SaveActorPreState(InActor, NewActorDatas.Emplace(InActor, ActorData.FindOrAdd(InActor)));
+			}
+
 			return true;
 		}
 		, EActorModifierCoreLookup::SelfAndAllChildren);
@@ -858,7 +866,19 @@ UGeometryMaskCanvas* UAvaMask2DBaseModifier::GetCurrentCanvas()
 
 void UAvaMask2DBaseModifier::OnMaterialsChanged(UObject* InMaterialOwner, const TArray<TSharedPtr<IAvaMaskMaterialHandle>>& InMaterialHandles)
 {
-	MarkModifierDirty();
+	if (!InMaterialHandles.IsEmpty())
+	{
+		if (const AActor* OwningActor = InMaterialOwner->GetTypedOuter<AActor>())
+		{
+			const UActorModifierCoreBase* MaskModifier = UE::ActorModifierCore::Utilities::FindFirstActorModifierByClass(OwningActor, GetClass());
+
+			// Only trigger update if first modifier found in the hierarchy above us is this modifier
+			if (MaskModifier == this)
+			{
+				MarkModifierDirty();
+			}
+		}
+	}
 }
 
 #undef LOCTEXT_NAMESPACE
