Motion Design: added queueing limits to ticker

diff --git a/Engine/Plugins/VirtualProduction/Avalanche/Source/Avalanche/Internal/Framework/Ticker/AvaTickerComponent.h b/Engine/Plugins/VirtualProduction/Avalanche/Source/Avalanche/Internal/Framework/Ticker/AvaTickerComponent.h
--- a/Engine/Plugins/VirtualProduction/Avalanche/Source/Avalanche/Internal/Framework/Ticker/AvaTickerComponent.h
+++ b/Engine/Plugins/VirtualProduction/Avalanche/Source/Avalanche/Internal/Framework/Ticker/AvaTickerComponent.h
@@ -7,6 +7,16 @@
 
 #define UE_API AVALANCHE_API
 
+UENUM()
+enum class EAvaTickerQueueLimitType : uint8
+{
+	/** No new elements will be queued if queue is full*/
+	DisableQueueing,
+
+	/** New elements will be queued by discarding the oldest (front of queue) */
+	DiscardOldest,
+};
+
 USTRUCT()
 struct FAvaTickerElement
 {
@@ -18,6 +28,9 @@ struct FAvaTickerElement
 
 	UE_API bool operator==(const FAvaTickerElement& InOther) const;
 
+	/** Whether the element is a valid element */
+	bool IsValid() const;
+
 	/** Sets the game and editor visibility of the element */
 	void SetVisibility(bool bInVisible);
 
@@ -49,9 +62,16 @@ class UAvaTickerComponent : public UPrimitiveComponent
 public:
 	UE_API UAvaTickerComponent();
 
-	/** Adds the given actor to the queue */
+	/** Returns whether the queue can be added to */
+	UFUNCTION(BlueprintPure, Category="Ticker")
+	bool CanQueueElements() const;
+
+	/**
+	 * Adds the given actor to the queue
+	 * @return true if the actor was successfully queued
+	 */
 	UFUNCTION(BlueprintCallable, Category="Ticker")
-	UE_API void QueueActor(AActor* InActor);
+	UE_API bool QueueActor(AActor* InActor);
 
 	UFUNCTION(BlueprintPure, Category="Ticker")
 	FVector GetStartLocation() const
@@ -89,6 +109,33 @@ public:
 	UFUNCTION(BlueprintCallable, Category="Ticker")
 	UE_API void SetPadding(double InPadding);
 
+	UFUNCTION(BlueprintPure, Category="Ticker")
+	bool ShouldLimitQueue() const
+	{
+		return bLimitQueue;
+	}
+
+	UFUNCTION(BlueprintCallable, Category="Ticker")
+	UE_API void SetLimitQueue(bool bInLimitQueue);
+
+	UFUNCTION(BlueprintPure, Category="Ticker")
+	int32 GetQueueLimitCount() const
+	{
+		return QueueLimitCount;
+	}
+
+	UFUNCTION(BlueprintCallable, Category="Ticker")
+	UE_API void SetQueueLimitCount(int32 InQueueLimitCount);
+
+	UFUNCTION(BlueprintPure, Category="Ticker")
+	EAvaTickerQueueLimitType GetQueueLimitType() const
+	{
+		return QueueLimitType;
+	}
+
+	UFUNCTION(BlueprintCallable, Category="Ticker")
+	UE_API void SetQueueLimitType(EAvaTickerQueueLimitType InQueueLimitType);
+
 	//~ Begin UPrimitiveComponent
 	UE_API virtual FPrimitiveSceneProxy* CreateSceneProxy() override;
 	//~ End UPrimitiveComponent
@@ -102,6 +149,9 @@ public:
 	//~ End UActorComponent
 
 protected:
+	/** Adds the given element to the queue */
+	bool QueueElement(FAvaTickerElement&& InElement);
+
 	/** Retrieves the start location transformed to World Space */
 	FVector GetWorldStartLocation() const;
 
@@ -140,6 +190,18 @@ protected:
 	/** Padding between the active ticker elements */
 	UPROPERTY(EditAnywhere, Category="Ticker", meta=(AllowPrivateAccess="true"))
 	double Padding = 30.0;
+
+	/** Indicates whether to limit the queue when it gets past a certain number of elements */
+	UPROPERTY(EditAnywhere, Category="Ticker|Queue", meta=(AllowPrivateAccess="true"))
+	bool bLimitQueue = true;
+
+	/** The number of items allowed in the queue */
+	UPROPERTY(EditAnywhere, Category="Ticker|Queue", meta=(EditCondition="bLimitQueue", EditConditionHides, ClampMin=0, AllowPrivateAccess="true"))
+	int32 QueueLimitCount = 2;
+
+	/** The handling method when the queue is full */
+	UPROPERTY(EditAnywhere, Category="Ticker|Queue", meta=(EditCondition="bLimitQueue", EditConditionHides, AllowPrivateAccess="true"))
+	EAvaTickerQueueLimitType QueueLimitType = EAvaTickerQueueLimitType::DisableQueueing;
 };
 
 #undef UE_API
diff --git a/Engine/Plugins/VirtualProduction/Avalanche/Source/Avalanche/Private/Framework/Ticker/AvaTickerComponent.cpp b/Engine/Plugins/VirtualProduction/Avalanche/Source/Avalanche/Private/Framework/Ticker/AvaTickerComponent.cpp
--- a/Engine/Plugins/VirtualProduction/Avalanche/Source/Avalanche/Private/Framework/Ticker/AvaTickerComponent.cpp
+++ b/Engine/Plugins/VirtualProduction/Avalanche/Source/Avalanche/Private/Framework/Ticker/AvaTickerComponent.cpp
@@ -24,6 +24,11 @@ bool FAvaTickerElement::operator==(const FAvaTickerElement& InOther) const
 	return Actor == InOther.Actor;
 }
 
+bool FAvaTickerElement::IsValid() const
+{
+	return !!Actor;
+}
+
 void FAvaTickerElement::SetVisibility(bool bInVisible)
 {
 	if (Actor)
@@ -87,19 +92,27 @@ UAvaTickerComponent::UAvaTickerComponent()
 {
 	PrimaryComponentTick.bCanEverTick = true;
 	PrimaryComponentTick.bStartWithTickEnabled = true;
+	bTickInEditor = true;
 	bHiddenInGame = true;
 }
 
-void UAvaTickerComponent::QueueActor(AActor* InActor)
+bool UAvaTickerComponent::CanQueueElements() const
 {
-	if (!InActor)
+	if (bLimitQueue && QueuedElements.Num() >= QueueLimitCount)
 	{
-		return;
+		switch (QueueLimitType)
+		{
+		case EAvaTickerQueueLimitType::DisableQueueing:
+			// Queue is full, return
+			return false;
+		}
 	}
+	return true;
+}
 
-	FAvaTickerElement Element(InActor);
-	Element.SetVisibility(/*bVisible*/false);
-	QueuedElements.AddUnique(MoveTemp(Element));
+bool UAvaTickerComponent::QueueActor(AActor* InActor)
+{
+	return InActor && QueueElement(FAvaTickerElement(InActor));
 }
 
 void UAvaTickerComponent::SetStartLocation(const FVector& InStartLocation)
@@ -122,6 +135,21 @@ void UAvaTickerComponent::SetPadding(double InPadding)
 	Padding = InPadding;
 }
 
+void UAvaTickerComponent::SetLimitQueue(bool bInLimitQueue)
+{
+	bLimitQueue = bInLimitQueue;
+}
+
+void UAvaTickerComponent::SetQueueLimitCount(int32 InQueueLimitCount)
+{
+	QueueLimitCount = InQueueLimitCount;
+}
+
+void UAvaTickerComponent::SetQueueLimitType(EAvaTickerQueueLimitType InQueueLimitType)
+{
+	QueueLimitType = InQueueLimitType;
+}
+
 FPrimitiveSceneProxy* UAvaTickerComponent::CreateSceneProxy()
 {
 	return new UE::Avalanche::FTickerSceneProxy(this);
@@ -140,6 +168,31 @@ void UAvaTickerComponent::TickComponent(float InDeltaTime, ELevelTick InTickType
 	UpdateActiveElements(InDeltaTime);
 }
 
+bool UAvaTickerComponent::QueueElement(FAvaTickerElement&& InElement)
+{
+	if (bLimitQueue && QueuedElements.Num() >= QueueLimitCount)
+	{
+		switch (QueueLimitType)
+		{
+		case EAvaTickerQueueLimitType::DisableQueueing:
+			// Queue is full, return
+			return false;
+
+		case EAvaTickerQueueLimitType::DiscardOldest:
+			if (QueuedElements.Num() > 0)
+			{
+				QueuedElements[0].Destroy();
+				QueuedElements.RemoveAt(0);
+			}
+			break;
+		}
+	}
+
+	InElement.SetVisibility(/*bVisible*/false);
+	QueuedElements.AddUnique(MoveTemp(InElement));
+	return true;
+}
+
 FVector UAvaTickerComponent::GetWorldStartLocation() const
 {
 	return GetComponentTransform().TransformPosition(StartLocation);
@@ -233,6 +286,12 @@ void UAvaTickerComponent::UpdateActiveElements(float InDeltaTime)
 	// Find the last active element that is more than the given distance away from 
 	for (TEnumerateRef<FAvaTickerElement> ActiveElement : EnumerateRange(ActiveElements))
 	{
+		if (!ActiveElement->IsValid())
+		{
+			IndexRangeToRemove = ActiveElement.GetIndex();
+			continue;
+		}
+
 		FVector ElementOrigin;
 		FVector ElementExtents;
 		ActiveElement->GetBounds(ElementOrigin, ElementExtents);
