Motion Design Ticker: fixed issue with wrongly calculated StartExtentsToOrigin + added option to destroy actors that couldn't be queued

diff --git a/Engine/Plugins/VirtualProduction/Avalanche/Source/Avalanche/Internal/Framework/Ticker/AvaTickerComponent.h b/Engine/Plugins/VirtualProduction/Avalanche/Source/Avalanche/Internal/Framework/Ticker/AvaTickerComponent.h
--- a/Engine/Plugins/VirtualProduction/Avalanche/Source/Avalanche/Internal/Framework/Ticker/AvaTickerComponent.h
+++ b/Engine/Plugins/VirtualProduction/Avalanche/Source/Avalanche/Internal/Framework/Ticker/AvaTickerComponent.h
@@ -68,10 +68,12 @@ public:
 
 	/**
 	 * Adds the given actor to the queue
+	 * @param InActor the actor to queue
+	 * @param bInDestroyOnFailure whether to destroy the actor if it failed to be queued
 	 * @return true if the actor was successfully queued
 	 */
 	UFUNCTION(BlueprintCallable, Category="Ticker")
-	UE_API bool QueueActor(AActor* InActor);
+	UE_API bool QueueActor(AActor* InActor, bool bInDestroyOnFailure = true);
 
 	UFUNCTION(BlueprintPure, Category="Ticker")
 	FVector GetStartLocation() const
@@ -150,7 +152,7 @@ public:
 
 protected:
 	/** Adds the given element to the queue */
-	bool QueueElement(FAvaTickerElement&& InElement);
+	bool QueueElement(FAvaTickerElement&& InElement, bool bInDestroyOnFailure);
 
 	/** Retrieves the start location transformed to World Space */
 	FVector GetWorldStartLocation() const;
diff --git a/Engine/Plugins/VirtualProduction/Avalanche/Source/Avalanche/Private/Framework/Ticker/AvaTickerComponent.cpp b/Engine/Plugins/VirtualProduction/Avalanche/Source/Avalanche/Private/Framework/Ticker/AvaTickerComponent.cpp
--- a/Engine/Plugins/VirtualProduction/Avalanche/Source/Avalanche/Private/Framework/Ticker/AvaTickerComponent.cpp
+++ b/Engine/Plugins/VirtualProduction/Avalanche/Source/Avalanche/Private/Framework/Ticker/AvaTickerComponent.cpp
@@ -110,9 +110,9 @@ bool UAvaTickerComponent::CanQueueElements() const
 	return true;
 }
 
-bool UAvaTickerComponent::QueueActor(AActor* InActor)
+bool UAvaTickerComponent::QueueActor(AActor* InActor, bool bInDestroyOnFailure)
 {
-	return InActor && QueueElement(FAvaTickerElement(InActor));
+	return InActor && QueueElement(FAvaTickerElement(InActor), bInDestroyOnFailure);
 }
 
 void UAvaTickerComponent::SetStartLocation(const FVector& InStartLocation)
@@ -168,13 +168,17 @@ void UAvaTickerComponent::TickComponent(float InDeltaTime, ELevelTick InTickType
 	UpdateActiveElements(InDeltaTime);
 }
 
-bool UAvaTickerComponent::QueueElement(FAvaTickerElement&& InElement)
+bool UAvaTickerComponent::QueueElement(FAvaTickerElement&& InElement, bool bInDestroyOnFailure)
 {
 	if (bLimitQueue && QueuedElements.Num() >= QueueLimitCount)
 	{
 		switch (QueueLimitType)
 		{
 		case EAvaTickerQueueLimitType::DisableQueueing:
+			if (bInDestroyOnFailure)
+			{
+				InElement.Destroy();
+			}
 			// Queue is full, return
 			return false;
 
@@ -229,7 +233,7 @@ void UAvaTickerComponent::TryDequeueElement()
 
 	// Switch the sign extents so that it can be projected onto the velocity direction
 	const FVector SignedExtents = ElementExtents * Private::GetSignVector(NormalizedVelocity);
-	const FVector StartExtentsToOrigin = -1.0 * (ElementOrigin + NormalizedVelocity * SignedExtents.ProjectOnTo(NormalizedVelocity).Length());
+	const FVector StartExtentsToOrigin = -1.0 * NormalizedVelocity * SignedExtents.ProjectOnTo(NormalizedVelocity).Length();
 	const FVector OriginToLocation = Element.GetElementLocation() - ElementOrigin;
 
 	// Check if the last active element is in the way
