Motion Design Ticker: added safety for when velocity is set to 0

diff --git a/Engine/Plugins/VirtualProduction/Avalanche/Source/Avalanche/Internal/Framework/Ticker/AvaTickerComponent.h b/Engine/Plugins/VirtualProduction/Avalanche/Source/Avalanche/Internal/Framework/Ticker/AvaTickerComponent.h
--- a/Engine/Plugins/VirtualProduction/Avalanche/Source/Avalanche/Internal/Framework/Ticker/AvaTickerComponent.h
+++ b/Engine/Plugins/VirtualProduction/Avalanche/Source/Avalanche/Internal/Framework/Ticker/AvaTickerComponent.h
@@ -10,6 +10,8 @@
 UENUM()
 enum class EAvaTickerQueueLimitType : uint8
 {
+	None UMETA(Hidden),
+
 	/** No new elements will be queued if queue is full*/
 	DisableQueueing,
 
@@ -151,6 +153,8 @@ public:
 	//~ End UActorComponent
 
 protected:
+	bool IsQueueFull() const;
+
 	/** Adds the given element to the queue */
 	bool QueueElement(FAvaTickerElement&& InElement, bool bInDestroyOnFailure);
 
diff --git a/Engine/Plugins/VirtualProduction/Avalanche/Source/Avalanche/Private/Framework/Ticker/AvaTickerComponent.cpp b/Engine/Plugins/VirtualProduction/Avalanche/Source/Avalanche/Private/Framework/Ticker/AvaTickerComponent.cpp
--- a/Engine/Plugins/VirtualProduction/Avalanche/Source/Avalanche/Private/Framework/Ticker/AvaTickerComponent.cpp
+++ b/Engine/Plugins/VirtualProduction/Avalanche/Source/Avalanche/Private/Framework/Ticker/AvaTickerComponent.cpp
@@ -98,15 +98,16 @@ UAvaTickerComponent::UAvaTickerComponent()
 
 bool UAvaTickerComponent::CanQueueElements() const
 {
-	if (bLimitQueue && QueuedElements.Num() >= QueueLimitCount)
+	if (IsQueueFull())
 	{
-		switch (QueueLimitType)
-		{
-		case EAvaTickerQueueLimitType::DisableQueueing:
-			// Queue is full, return
-			return false;
-		}
+		return QueueLimitType != EAvaTickerQueueLimitType::DisableQueueing;
+	}
+
+	if (GetWorldVelocity().IsNearlyZero())
+	{
+		return false;
 	}
+
 	return true;
 }
 
@@ -168,9 +169,24 @@ void UAvaTickerComponent::TickComponent(float InDeltaTime, ELevelTick InTickType
 	UpdateActiveElements(InDeltaTime);
 }
 
+bool UAvaTickerComponent::IsQueueFull() const
+{
+	return bLimitQueue && QueuedElements.Num() >= QueueLimitCount;
+}
+
 bool UAvaTickerComponent::QueueElement(FAvaTickerElement&& InElement, bool bInDestroyOnFailure)
 {
-	if (bLimitQueue && QueuedElements.Num() >= QueueLimitCount)
+	// Can't queue if velocity is zero
+	if (GetWorldVelocity().IsNearlyZero())
+	{
+		if (bInDestroyOnFailure)
+		{
+			InElement.Destroy();
+		}
+		return false;
+	}
+
+	if (IsQueueFull())
 	{
 		switch (QueueLimitType)
 		{
@@ -179,7 +195,7 @@ bool UAvaTickerComponent::QueueElement(FAvaTickerElement&& InElement, bool bInDe
 			{
 				InElement.Destroy();
 			}
-			// Queue is full, return
+			// Queue is disabled on full, return
 			return false;
 
 		case EAvaTickerQueueLimitType::DiscardOldest:
@@ -222,6 +238,12 @@ void UAvaTickerComponent::TryDequeueElement()
 		return;
 	}
 
+	const FVector NormalizedVelocity = GetWorldVelocity().GetSafeNormal();
+	if (NormalizedVelocity.IsNearlyZero())
+	{
+		return;
+	}
+
 	const FAvaTickerElement& Element = QueuedElements[0];
 
 	FVector ElementOrigin;
@@ -229,7 +251,6 @@ void UAvaTickerComponent::TryDequeueElement()
 	Element.GetBounds(ElementOrigin, ElementExtents);
 
 	const FVector WorldStartLocation = GetWorldStartLocation();
-	const FVector NormalizedVelocity = GetWorldVelocity().GetSafeNormal();
 
 	// Switch the sign extents so that it can be projected onto the velocity direction
 	const FVector SignedExtents = ElementExtents * Private::GetSignVector(NormalizedVelocity);
