Motion Design Ticker: added  missing default break

diff --git a/Engine/Plugins/VirtualProduction/Avalanche/Source/Avalanche/Private/Framework/Ticker/AvaTickerComponent.cpp b/Engine/Plugins/VirtualProduction/Avalanche/Source/Avalanche/Private/Framework/Ticker/AvaTickerComponent.cpp
--- a/Engine/Plugins/VirtualProduction/Avalanche/Source/Avalanche/Private/Framework/Ticker/AvaTickerComponent.cpp
+++ b/Engine/Plugins/VirtualProduction/Avalanche/Source/Avalanche/Private/Framework/Ticker/AvaTickerComponent.cpp
@@ -205,6 +205,9 @@ bool UAvaTickerComponent::QueueElement(FAvaTickerElement&& InElement, bool bInDe
 				QueuedElements.RemoveAt(0);
 			}
 			break;
+
+		default:
+			break;
 		}
 	}
 
