MotionDesign : EditorInteractiveTools
- Fixed preview actor spawned infinitely once deleted in outliner when tool is active
- Simplified tool parameters creation

diff --git a/Engine/Plugins/Experimental/AvalancheDataLink/Source/AvalancheDataLinkEditor/Private/AvaDataLinkActorTool.cpp b/Engine/Plugins/Experimental/AvalancheDataLink/Source/AvalancheDataLinkEditor/Private/AvaDataLinkActorTool.cpp
--- a/Engine/Plugins/Experimental/AvalancheDataLink/Source/AvalancheDataLinkEditor/Private/AvaDataLinkActorTool.cpp
+++ b/Engine/Plugins/Experimental/AvalancheDataLink/Source/AvalancheDataLinkEditor/Private/AvaDataLinkActorTool.cpp
@@ -33,15 +33,8 @@ void UAvaDataLinkActorTool::OnRegisterTool(IAvalancheInteractiveToolsModule* InA
 {
 	Super::OnRegisterTool(InAITModule);
 
-	FAvaInteractiveToolsToolParameters ToolParameters;
+	FAvaInteractiveToolsToolParameters ToolParameters = CreateDefaultToolParameters();
 	ToolParameters.UICommand = UE::AvaDataLink::FEditorCommands::Get().DataLinkActorTool;
-	ToolParameters.ToolIdentifier = TEXT("Motion Design Data Link Actor Tool");
-	ToolParameters.FactoryClass = ActorClass;
-	ToolParameters.CreateBuilder.BindLambda(
-		[](UEdMode* InEdMode)
-		{
-			return UAvaInteractiveToolsToolBuilder::CreateToolBuilder<UAvaDataLinkActorTool>(InEdMode);
-		});
 
 	InAITModule->RegisterTool(IAvalancheInteractiveToolsModule::CategoryNameActor, MoveTemp(ToolParameters));
 }
diff --git a/Engine/Plugins/VirtualProduction/Avalanche/Source/AvalancheEditor/Private/Ticker/AvaTickerActorTool.cpp b/Engine/Plugins/VirtualProduction/Avalanche/Source/AvalancheEditor/Private/Ticker/AvaTickerActorTool.cpp
--- a/Engine/Plugins/VirtualProduction/Avalanche/Source/AvalancheEditor/Private/Ticker/AvaTickerActorTool.cpp
+++ b/Engine/Plugins/VirtualProduction/Avalanche/Source/AvalancheEditor/Private/Ticker/AvaTickerActorTool.cpp
@@ -14,15 +14,8 @@ void UAvaTickerActorTool::OnRegisterTool(IAvalancheInteractiveToolsModule* InAIT
 {
 	Super::OnRegisterTool(InAITModule);
 
-	FAvaInteractiveToolsToolParameters ToolParameters;
+	FAvaInteractiveToolsToolParameters ToolParameters = CreateDefaultToolParameters();
 	ToolParameters.UICommand = FAvaEditorCommands::Get().TickerTool;
-	ToolParameters.ToolIdentifier = TEXT("Motion Design Ticker Tool");
-	ToolParameters.FactoryClass = ActorClass;
-	ToolParameters.CreateBuilder.BindLambda(
-		[](UEdMode* InEdMode)
-		{
-			return UAvaInteractiveToolsToolBuilder::CreateToolBuilder<UAvaTickerActorTool>(InEdMode);
-		});
 
 	InAITModule->RegisterTool(IAvalancheInteractiveToolsModule::CategoryNameActor, MoveTemp(ToolParameters));
 }
diff --git a/Engine/Plugins/VirtualProduction/Avalanche/Source/AvalancheEffectorsEditor/Private/AvaEffectorsEditorCommands.cpp b/Engine/Plugins/VirtualProduction/Avalanche/Source/AvalancheEffectorsEditor/Private/AvaEffectorsEditorCommands.cpp
--- a/Engine/Plugins/VirtualProduction/Avalanche/Source/AvalancheEffectorsEditor/Private/AvaEffectorsEditorCommands.cpp
+++ b/Engine/Plugins/VirtualProduction/Avalanche/Source/AvalancheEffectorsEditor/Private/AvaEffectorsEditorCommands.cpp
@@ -31,7 +31,7 @@ void FAvaEffectorsEditorCommands::RegisterCommands()
 			AsShared(),
 			TypeCommand,
 			TypeName,
-			FText::FromName(TypeName),
+			FText::FromString(FName::NameToDisplayString(TypeName.ToString(), /** IsBool */false)),
 			FText::Format(LOCTEXT("EffectorCommandTooltip", "Create a {0} Effector Actor in the viewport."), FText::FromName(TypeName)),
 			FSlateIconFinder::FindIcon(TEXT("AvaEffectorsEditor.Tool_Actor_Effector")),
 			EUserInterfaceActionType::ToggleButton,
@@ -50,7 +50,7 @@ void FAvaEffectorsEditorCommands::RegisterCommands()
 			AsShared(),
 			LayoutCommand,
 			LayoutName,
-			FText::FromName(LayoutName),
+			FText::FromString(FName::NameToDisplayString(LayoutName.ToString(), /** IsBool */false)),
 			FText::Format(LOCTEXT("ClonerCommandTooltip", "Create a {0} Cloner Actor in the viewport."), FText::FromName(LayoutName)),
 			FSlateIconFinder::FindIcon(TEXT("AvaEffectorsEditor.Tool_Actor_Cloner")),
 			EUserInterfaceActionType::ToggleButton,
diff --git a/Engine/Plugins/VirtualProduction/Avalanche/Source/AvalancheEffectorsEditor/Private/Cloner/AvaClonerActorTool.cpp b/Engine/Plugins/VirtualProduction/Avalanche/Source/AvalancheEffectorsEditor/Private/Cloner/AvaClonerActorTool.cpp
--- a/Engine/Plugins/VirtualProduction/Avalanche/Source/AvalancheEffectorsEditor/Private/Cloner/AvaClonerActorTool.cpp
+++ b/Engine/Plugins/VirtualProduction/Avalanche/Source/AvalancheEffectorsEditor/Private/Cloner/AvaClonerActorTool.cpp
@@ -24,20 +24,11 @@ void UAvaClonerActorTool::OnRegisterTool(IAvalancheInteractiveToolsModule* InAIT
 		UCEClonerActorFactory* ClonerActorFactory = CreateActorFactory<UCEClonerActorFactory>();
 		ClonerActorFactory->SetClonerLayout(ClonerCommandPair.Key);
 		
-		FAvaInteractiveToolsToolParameters ToolParameters =
-		{
-			ClonerCommandPair.Value,
-			ToolIdentifier + ClonerCommandPair.Key.ToString(),
-			3000,
-			FAvalancheInteractiveToolsCreateBuilder::CreateLambda(
-				[](UEdMode* InEdMode)
-				{
-					return UAvaInteractiveToolsToolBuilder::CreateToolBuilder<UAvaClonerActorTool>(InEdMode);
-				}),
-			ActorClass,
-			ClonerActorFactory,
-			FText::FromString(ClonerCommandPair.Key.ToString()),
-		};
+		FAvaInteractiveToolsToolParameters ToolParameters = CreateDefaultToolParameters();
+		ToolParameters.UICommand = ClonerCommandPair.Value;
+		ToolParameters.ToolIdentifier = ToolIdentifier + ClonerCommandPair.Key.ToString();
+		ToolParameters.Priority = 3000;
+		ToolParameters.ActorFactory.Set<TObjectPtr<UActorFactory>>(ClonerActorFactory);
 
 		InAITModule->RegisterTool(IAvalancheInteractiveToolsModule::CategoryNameCloner, MoveTemp(ToolParameters));
 	}
diff --git a/Engine/Plugins/VirtualProduction/Avalanche/Source/AvalancheEffectorsEditor/Private/Effector/AvaEffectorActorTool.cpp b/Engine/Plugins/VirtualProduction/Avalanche/Source/AvalancheEffectorsEditor/Private/Effector/AvaEffectorActorTool.cpp
--- a/Engine/Plugins/VirtualProduction/Avalanche/Source/AvalancheEffectorsEditor/Private/Effector/AvaEffectorActorTool.cpp
+++ b/Engine/Plugins/VirtualProduction/Avalanche/Source/AvalancheEffectorsEditor/Private/Effector/AvaEffectorActorTool.cpp
@@ -25,20 +25,11 @@ void UAvaEffectorActorTool::OnRegisterTool(IAvalancheInteractiveToolsModule* InA
 		UCEEffectorActorFactory* EffectorActorFactory = CreateActorFactory<UCEEffectorActorFactory>();
 		EffectorActorFactory->SetEffectorTypeName(EffectorCommandPair.Key);
 
-		FAvaInteractiveToolsToolParameters ToolParameters =
-		{
-			EffectorCommandPair.Value,
-			ToolIdentifier + EffectorCommandPair.Key.ToString(),
-			4000,
-			FAvalancheInteractiveToolsCreateBuilder::CreateLambda(
-				[](UEdMode* InEdMode)
-				{
-					return UAvaInteractiveToolsToolBuilder::CreateToolBuilder<UAvaEffectorActorTool>(InEdMode);
-				}),
-			ActorClass,
-			EffectorActorFactory,
-			FText::FromString(EffectorCommandPair.Key.ToString()),
-		};
+		FAvaInteractiveToolsToolParameters ToolParameters = CreateDefaultToolParameters();
+		ToolParameters.ToolIdentifier = ToolIdentifier + EffectorCommandPair.Key.ToString();
+		ToolParameters.UICommand = EffectorCommandPair.Value;
+		ToolParameters.Priority = 4000;
+		ToolParameters.ActorFactory.Set<TObjectPtr<UActorFactory>>(EffectorActorFactory);
 
 		InAITModule->RegisterTool(IAvalancheInteractiveToolsModule::CategoryNameEffector, MoveTemp(ToolParameters));
 	}
diff --git a/Engine/Plugins/VirtualProduction/Avalanche/Source/AvalancheInteractiveTools/Private/AvalancheInteractiveToolsModule.cpp b/Engine/Plugins/VirtualProduction/Avalanche/Source/AvalancheInteractiveTools/Private/AvalancheInteractiveToolsModule.cpp
--- a/Engine/Plugins/VirtualProduction/Avalanche/Source/AvalancheInteractiveTools/Private/AvalancheInteractiveToolsModule.cpp
+++ b/Engine/Plugins/VirtualProduction/Avalanche/Source/AvalancheInteractiveTools/Private/AvalancheInteractiveToolsModule.cpp
@@ -167,9 +167,9 @@ void FAvalancheInteractiveToolsModule::AddReferencedObjects(FReferenceCollector&
 	{
 		for (FAvaInteractiveToolsToolParameters& ToolParams : CategoryPair.Value)
 		{
-			if (ToolParams.Factory)
+			if (ToolParams.ActorFactory.IsType<TObjectPtr<UActorFactory>>())
 			{
-				InCollector.AddReferencedObject<UActorFactory>(ToolParams.Factory);
+				InCollector.AddReferencedObject<UActorFactory>(ToolParams.ActorFactory.Get<TObjectPtr<UActorFactory>>());
 			}
 		}
 	}
@@ -295,7 +295,14 @@ void FAvalancheInteractiveToolsModule::OnPlacementCategoryRefreshed(FName InCate
 
 	for (const FAvaInteractiveToolsToolParameters& Tool : Tools[InCategory])
 	{
-		if (!(Tool.Factory && Tool.Factory->NewActorClass) && !Tool.FactoryClass)
+		const bool bHasActorClass = Tool.ActorFactory.IsType<TSubclassOf<AActor>>()
+			&& Tool.ActorFactory.Get<TSubclassOf<AActor>>().Get() != nullptr;
+
+		const bool bHasActorFactory = Tool.ActorFactory.IsType<TObjectPtr<UActorFactory>>()
+			&& Tool.ActorFactory.Get<TObjectPtr<UActorFactory>>() != nullptr
+			&& Tool.ActorFactory.Get<TObjectPtr<UActorFactory>>()->NewActorClass.Get() != nullptr;
+
+		if (!bHasActorFactory && !bHasActorClass)
 		{
 			continue;
 		}
@@ -314,19 +321,23 @@ void FAvalancheInteractiveToolsModule::OnPlacementCategoryRefreshed(FName InCate
 
 		TSharedPtr<FPlaceableItem> PlaceableItem;
 
-		if (Tool.Factory)
+		if (bHasActorFactory)
 		{
+			UActorFactory* ActorFactory = Tool.ActorFactory.Get<TObjectPtr<UActorFactory>>();
+
 			PlaceableItem = MakeShared<FPlaceableItem>(
-				Tool.Factory.Get(),
-				FAssetData(Tool.Factory->NewActorClass.GetDefaultObject()),
+				ActorFactory,
+				FAssetData(ActorFactory->NewActorClass.GetDefaultObject()),
 				Tool.Priority
 			);
 		}
 		else
 		{
+			UClass* ActorClass = Tool.ActorFactory.Get<TSubclassOf<AActor>>();
+
 			PlaceableItem = MakeShared<FPlaceableItem>(
-				*Tool.FactoryClass.Get(),
-				FAssetData(Tool.FactoryClass.Get()->GetDefaultObject()),
+				*ActorClass,
+				FAssetData(ActorClass->GetDefaultObject()),
 				NAME_None,
 				NAME_None,
 				TOptional<FLinearColor>(),
diff --git a/Engine/Plugins/VirtualProduction/Avalanche/Source/AvalancheInteractiveTools/Private/Tools/AvaInteractiveToolsActorToolBase.cpp b/Engine/Plugins/VirtualProduction/Avalanche/Source/AvalancheInteractiveTools/Private/Tools/AvaInteractiveToolsActorToolBase.cpp
--- a/Engine/Plugins/VirtualProduction/Avalanche/Source/AvalancheInteractiveTools/Private/Tools/AvaInteractiveToolsActorToolBase.cpp
+++ b/Engine/Plugins/VirtualProduction/Avalanche/Source/AvalancheInteractiveTools/Private/Tools/AvaInteractiveToolsActorToolBase.cpp
@@ -1,6 +1,8 @@
 // Copyright Epic Games, Inc. All Rights Reserved.
 
 #include "Tools/AvaInteractiveToolsActorToolBase.h"
+
+#include "Builders/AvaInteractiveToolsToolBuilder.h"
 #include "GameFramework/Actor.h"
 
 bool UAvaInteractiveToolsActorToolBase::OnBegin()
@@ -29,3 +31,16 @@ bool UAvaInteractiveToolsActorToolBase::UseIdentityRotation() const
 {
 	return ConditionalIdentityRotation();
 }
+
+FAvaInteractiveToolsToolParameters UAvaInteractiveToolsActorToolBase::CreateDefaultToolParameters() const
+{
+	FAvaInteractiveToolsToolParameters Parameters;
+	Parameters.ToolIdentifier = GetClass()->GetName();
+	Parameters.CreateBuilder = FAvalancheInteractiveToolsCreateBuilder::CreateLambda(
+		[this](UEdMode* InEdMode)
+		{
+			return UAvaInteractiveToolsToolBuilder::CreateToolBuilder(InEdMode, GetClass());
+		});
+	Parameters.ActorFactory.Set<TSubclassOf<AActor>>(ActorClass);
+	return Parameters;
+}
diff --git a/Engine/Plugins/VirtualProduction/Avalanche/Source/AvalancheInteractiveTools/Private/Tools/AvaInteractiveToolsActorToolNull.cpp b/Engine/Plugins/VirtualProduction/Avalanche/Source/AvalancheInteractiveTools/Private/Tools/AvaInteractiveToolsActorToolNull.cpp
--- a/Engine/Plugins/VirtualProduction/Avalanche/Source/AvalancheInteractiveTools/Private/Tools/AvaInteractiveToolsActorToolNull.cpp
+++ b/Engine/Plugins/VirtualProduction/Avalanche/Source/AvalancheInteractiveTools/Private/Tools/AvaInteractiveToolsActorToolNull.cpp
@@ -15,18 +15,9 @@ void UAvaInteractiveToolsActorToolNull::OnRegisterTool(IAvalancheInteractiveTool
 {
 	Super::OnRegisterTool(InAITModule);
 
-	FAvaInteractiveToolsToolParameters ToolParameters =
-	{
-		FAvaInteractiveToolsCommands::Get().Tool_Actor_Null,
-		TEXT("Null Actor Tool"),
-		2000,
-		FAvalancheInteractiveToolsCreateBuilder::CreateLambda(
-			[](UEdMode* InEdMode)
-			{
-				return UAvaInteractiveToolsToolBuilder::CreateToolBuilder<UAvaInteractiveToolsActorToolNull>(InEdMode);
-			}),
-		ActorClass
-	};
+	FAvaInteractiveToolsToolParameters ToolParameters = CreateDefaultToolParameters();
+	ToolParameters.UICommand = FAvaInteractiveToolsCommands::Get().Tool_Actor_Null;
+	ToolParameters.Priority = 2000;
 
 	InAITModule->RegisterTool(IAvalancheInteractiveToolsModule::CategoryNameActor, MoveTemp(ToolParameters));
 }
diff --git a/Engine/Plugins/VirtualProduction/Avalanche/Source/AvalancheInteractiveTools/Private/Tools/AvaInteractiveToolsActorToolSpline.cpp b/Engine/Plugins/VirtualProduction/Avalanche/Source/AvalancheInteractiveTools/Private/Tools/AvaInteractiveToolsActorToolSpline.cpp
--- a/Engine/Plugins/VirtualProduction/Avalanche/Source/AvalancheInteractiveTools/Private/Tools/AvaInteractiveToolsActorToolSpline.cpp
+++ b/Engine/Plugins/VirtualProduction/Avalanche/Source/AvalancheInteractiveTools/Private/Tools/AvaInteractiveToolsActorToolSpline.cpp
@@ -15,18 +15,9 @@ void UAvaInteractiveToolsActorToolSpline::OnRegisterTool(IAvalancheInteractiveTo
 {
 	Super::OnRegisterTool(InAITModule);
 
-	FAvaInteractiveToolsToolParameters ToolParameters =
-	{
-		FAvaInteractiveToolsCommands::Get().Tool_Actor_Spline,
-		TEXT("Spline Actor Tool"),
-		5000,
-		FAvalancheInteractiveToolsCreateBuilder::CreateLambda(
-			[](UEdMode* InEdMode)
-			{
-				return UAvaInteractiveToolsToolBuilder::CreateToolBuilder<UAvaInteractiveToolsActorToolSpline>(InEdMode);
-			}),
-		ActorClass
-	};
+	FAvaInteractiveToolsToolParameters ToolParameters = CreateDefaultToolParameters();
+	ToolParameters.UICommand = FAvaInteractiveToolsCommands::Get().Tool_Actor_Spline;
+	ToolParameters.Priority = 5000;
 
 	InAITModule->RegisterTool(IAvalancheInteractiveToolsModule::CategoryNameActor, MoveTemp(ToolParameters));
 }
diff --git a/Engine/Plugins/VirtualProduction/Avalanche/Source/AvalancheInteractiveTools/Private/Tools/AvaInteractiveToolsToolBase.cpp b/Engine/Plugins/VirtualProduction/Avalanche/Source/AvalancheInteractiveTools/Private/Tools/AvaInteractiveToolsToolBase.cpp
--- a/Engine/Plugins/VirtualProduction/Avalanche/Source/AvalancheInteractiveTools/Private/Tools/AvaInteractiveToolsToolBase.cpp
+++ b/Engine/Plugins/VirtualProduction/Avalanche/Source/AvalancheInteractiveTools/Private/Tools/AvaInteractiveToolsToolBase.cpp
@@ -25,17 +25,11 @@
 #include "Planners/AvaInteractiveToolsToolViewportPointPlanner.h"
 #include "ToolContextInterfaces.h"
 #include "Toolkits/BaseToolkit.h"
-#include "ToolMenus.h"
 #include "UnrealClient.h"
 #include "ViewportClient/IAvaViewportClient.h"
 
 #define LOCTEXT_NAMESPACE "AvaInteractiveToolsToolBase"
 
-namespace UE::AvaInteractiveTools::Private
-{
-	static const FName AvaITFToolPresetMenuName = "AvaITFToolPresetMenu";
-}
-
 UAvaInteractiveToolsRightClickBehavior::UAvaInteractiveToolsRightClickBehavior()
 {
 	HitTestOnRelease = false;
@@ -88,13 +82,11 @@ void UAvaInteractiveToolsToolBase::Setup()
 {
 	Super::Setup();
 
-	using namespace UE::AvaInteractiveTools::Private;
-
 	UInteractiveToolManager* ToolManager = GetToolManager();
 	bool bReactivated = false;
 	UAvaInteractiveToolsEdMode* AvaInteractiveToolsEdMode = nullptr;
 
-	if (UContextObjectStore* ContextStore = ToolManager->GetContextObjectStore())
+	if (const UContextObjectStore* ContextStore = ToolManager->GetContextObjectStore())
 	{
 		AvaInteractiveToolsEdMode = Cast<UAvaInteractiveToolsEdMode>(ContextStore->FindContextByClass(UAvaInteractiveToolsEdMode::StaticClass()));
 
@@ -205,13 +197,13 @@ bool UAvaInteractiveToolsToolBase::ConditionalIdentityRotation() const
 
 AActor* UAvaInteractiveToolsToolBase::SpawnActor(TSubclassOf<AActor> InActorClass, bool bInPreview, FString* InActorLabelOverride)
 {
-	if (IToolsContextQueriesAPI* ContextAPI = GetToolManager()->GetContextQueriesAPI())
+	if (const IToolsContextQueriesAPI* ContextAPI = GetToolManager()->GetContextQueriesAPI())
 	{
 		if (FViewport* Viewport = ContextAPI->GetFocusedViewport())
 		{
 			if (FEditorViewportClient* ViewportClient = FAvaViewportUtils::GetAsEditorViewportClient(Viewport))
 			{
-				if (TSharedPtr<IAvaViewportClient> AvaViewportClient = FAvaViewportUtils::GetAsAvaViewportClient(ViewportClient))
+				if (const TSharedPtr<IAvaViewportClient> AvaViewportClient = FAvaViewportUtils::GetAsAvaViewportClient(ViewportClient))
 				{
 					const FVector2f ViewportSize = AvaViewportClient->GetViewportSize();
 					return SpawnActor(InActorClass, EAvaViewportStatus::Focused, ViewportSize * 0.5, bInPreview, InActorLabelOverride);
@@ -312,17 +304,15 @@ AActor* UAvaInteractiveToolsToolBase::SpawnActor(TSubclassOf<AActor> InActorClas
 	}
 
 	UActorFactory* ActorFactory = nullptr;
-	const FAvaInteractiveToolsToolParameters* ActiveToolParameters = GetActiveToolParameters();
-
-	if (ActiveToolParameters)
+	if (const FAvaInteractiveToolsToolParameters* ActiveToolParameters = GetActiveToolParameters())
 	{
-		if (ActiveToolParameters->Factory)
+		if (ActiveToolParameters->ActorFactory.IsType<TObjectPtr<UActorFactory>>())
 		{
-			ActorFactory = ActiveToolParameters->Factory;
+			ActorFactory = ActiveToolParameters->ActorFactory.Get<TObjectPtr<UActorFactory>>();
 		}
-		else if (ActiveToolParameters->FactoryClass)
+		else if (ActiveToolParameters->ActorFactory.IsType<TSubclassOf<AActor>>())
 		{
-			ActorFactory = GEditor->FindActorFactoryByClass(ActiveToolParameters->FactoryClass);
+			ActorFactory = GEditor->FindActorFactoryByClass(ActiveToolParameters->ActorFactory.Get<TSubclassOf<AActor>>());
 		}
 	}
 
@@ -347,6 +337,12 @@ AActor* UAvaInteractiveToolsToolBase::SpawnActor(TSubclassOf<AActor> InActorClas
 	checkf(!!NewActor, TEXT("Motion Design Interactive Tool : Failed to spawn actor of class %s"), *InActorClass->GetName())
 	NewActor->bIsEditorPreviewActor = bInPreview;
 
+	if (bInPreview)
+	{
+		// Cancel tool when preview actor is deleted in outliner to avoid spawning new one
+		NewActor->OnDestroyed.AddUniqueDynamic(this, &UAvaInteractiveToolsToolBase::OnPreviewActorDestroyed);
+	}
+
 	if (bUseOverride)
 	{
 		FActorLabelUtilities::SetActorLabelUnique(NewActor, *InActorLabelOverride);
@@ -369,14 +365,14 @@ AActor* UAvaInteractiveToolsToolBase::SpawnActor(TSubclassOf<AActor> InActorClas
 
 void UAvaInteractiveToolsToolBase::SetToolkitSettingsObject(UObject* InObject) const
 {
-	if (UInteractiveToolManager* ToolManager = GetToolManager())
+	if (const UInteractiveToolManager* ToolManager = GetToolManager())
 	{
-		if (UContextObjectStore* ContextStore = ToolManager->GetContextObjectStore())
+		if (const UContextObjectStore* ContextStore = ToolManager->GetContextObjectStore())
 		{
 			if (UAvaInteractiveToolsEdMode* AvaInteractiveToolsEdMode =
 				Cast<UAvaInteractiveToolsEdMode>(ContextStore->FindContextByClass(UAvaInteractiveToolsEdMode::StaticClass())))
 			{
-				if (TSharedPtr<FModeToolkit> Toolkit = AvaInteractiveToolsEdMode->GetToolkit().Pin())
+				if (const TSharedPtr<FModeToolkit> Toolkit = AvaInteractiveToolsEdMode->GetToolkit().Pin())
 				{
 					Toolkit->SetModeSettingsObject(InObject);
 				}
@@ -396,16 +392,15 @@ bool UAvaInteractiveToolsToolBase::ViewportPositionToWorldPositionAndOrientation
 bool UAvaInteractiveToolsToolBase::ViewportPositionToWorldPositionAndOrientation(EAvaViewportStatus InViewportStatus,
 	const FVector2f& InViewportPosition, float InDistance, UWorld*& OutWorld, FVector& OutPosition, FRotator& OutRotation) const
 {
-	if (IToolsContextQueriesAPI* ContextAPI = GetToolManager()->GetContextQueriesAPI())
+	if (const IToolsContextQueriesAPI* ContextAPI = GetToolManager()->GetContextQueriesAPI())
 	{
 		if (UWorld* World = ContextAPI->GetCurrentEditingWorld())
 		{
-			if (TSharedPtr<IAvaViewportClient> AvaViewportClient = FAvaViewportUtils::GetAvaViewportClient(GetViewport(InViewportStatus)))
+			if (const TSharedPtr<IAvaViewportClient> AvaViewportClient = FAvaViewportUtils::GetAvaViewportClient(GetViewport(InViewportStatus)))
 			{
 				if (const FEditorViewportClient* EditorViewportClient = AvaViewportClient->AsEditorViewportClient())
 				{
 					const FViewportCameraTransform ViewTransform = EditorViewportClient->GetViewTransform();
-					const FVector2f ViewportSize = AvaViewportClient->GetViewportSize();
 
 					OutWorld = World;
 					OutRotation = ViewTransform.GetRotation();
@@ -424,9 +419,9 @@ FViewport* UAvaInteractiveToolsToolBase::GetViewport(EAvaViewportStatus InViewpo
 {
 	if (UInteractiveToolManager* ToolManager = GetToolManager())
 	{
-		if (IToolsContextQueriesAPI* ContextAPI = ToolManager->GetContextQueriesAPI())
+		if (const IToolsContextQueriesAPI* ContextAPI = ToolManager->GetContextQueriesAPI())
 		{
-			FViewport* Viewport = nullptr;
+			FViewport* Viewport;
 
 			switch (InViewportStatus)
 			{
@@ -455,10 +450,7 @@ FViewport* UAvaInteractiveToolsToolBase::GetViewport(EAvaViewportStatus InViewpo
 
 void UAvaInteractiveToolsToolBase::OnViewportPlannerComplete()
 {
-	if (PreviewActor)
-	{
-		PreviewActor->Destroy();
-	}
+	DestroyPreviewActor();
 
 	RequestShutdown(EToolShutdownType::Completed);
 }
@@ -552,6 +544,7 @@ void UAvaInteractiveToolsToolBase::OnUpdateModifierState(int InModifierID, bool
 			}
 
 			break;
+		default: ;
 	}
 }
 
@@ -645,11 +638,7 @@ void UAvaInteractiveToolsToolBase::OnCancel()
 {
 	CancelTransaction();
 
-	if (PreviewActor)
-	{
-		PreviewActor->Destroy();
-		PreviewActor = nullptr;
-	}
+	DestroyPreviewActor();
 
 	SetToolkitSettingsObject(nullptr);
 }
@@ -658,11 +647,7 @@ void UAvaInteractiveToolsToolBase::OnComplete()
 {
 	EndTransaction();
 
-	if (PreviewActor)
-	{
-		PreviewActor->Destroy();
-		PreviewActor = nullptr;
-	}
+	DestroyPreviewActor();
 
 	if (FEngineAnalytics::IsAvailable())
 	{
@@ -696,7 +681,7 @@ void UAvaInteractiveToolsToolBase::CancelTransaction()
 
 void UAvaInteractiveToolsToolBase::RequestShutdown(EToolShutdownType InShutdownType)
 {
-	if (UContextObjectStore* ContextStore = GetToolManager()->GetContextObjectStore())
+	if (const UContextObjectStore* ContextStore = GetToolManager()->GetContextObjectStore())
 	{
 		if (UAvaInteractiveToolsEdMode* AvaInteractiveToolsEdMode = Cast<UAvaInteractiveToolsEdMode>(ContextStore->FindContextByClass(UAvaInteractiveToolsEdMode::StaticClass())))
 		{
@@ -718,11 +703,11 @@ void UAvaInteractiveToolsToolBase::RequestShutdown(EToolShutdownType InShutdownT
 
 bool UAvaInteractiveToolsToolBase::IsMotionDesignViewport() const
 {
-	if (IToolsContextQueriesAPI* ContextAPI = GetToolManager()->GetContextQueriesAPI())
+	if (const IToolsContextQueriesAPI* ContextAPI = GetToolManager()->GetContextQueriesAPI())
 	{
 		if (FViewport* Viewport = ContextAPI->GetFocusedViewport())
 		{
-			if (TSharedPtr<IAvaViewportClient> AvaViewportClient = FAvaViewportUtils::GetAvaViewportClient(Viewport))
+			if (const TSharedPtr<IAvaViewportClient> AvaViewportClient = FAvaViewportUtils::GetAvaViewportClient(Viewport))
 			{
 				return AvaViewportClient->IsMotionDesignViewport();
 			}
@@ -737,4 +722,19 @@ bool UAvaInteractiveToolsToolBase::ShouldForceDefaultAction() const
 	return FSlateApplication::Get().GetModifierKeys().IsAltDown();
 }
 
+void UAvaInteractiveToolsToolBase::DestroyPreviewActor()
+{
+	if (PreviewActor)
+	{
+		PreviewActor->OnDestroyed.RemoveAll(this);
+		PreviewActor->Destroy();
+		PreviewActor = nullptr;
+	}
+}
+
+void UAvaInteractiveToolsToolBase::OnPreviewActorDestroyed(AActor* InDestroyedActor)
+{
+	RequestShutdown(EToolShutdownType::Cancel);
+}
+
 #undef LOCTEXT_NAMESPACE
diff --git a/Engine/Plugins/VirtualProduction/Avalanche/Source/AvalancheInteractiveTools/Private/Widgets/SAvaInteractiveToolsToolbarCategoryButton.cpp b/Engine/Plugins/VirtualProduction/Avalanche/Source/AvalancheInteractiveTools/Private/Widgets/SAvaInteractiveToolsToolbarCategoryButton.cpp
--- a/Engine/Plugins/VirtualProduction/Avalanche/Source/AvalancheInteractiveTools/Private/Widgets/SAvaInteractiveToolsToolbarCategoryButton.cpp
+++ b/Engine/Plugins/VirtualProduction/Avalanche/Source/AvalancheInteractiveTools/Private/Widgets/SAvaInteractiveToolsToolbarCategoryButton.cpp
@@ -75,10 +75,7 @@ void SAvaInteractiveToolsToolbarCategoryButton::CreateActiveCommandWidget()
 
 		if (Tool && Tool->UICommand.IsValid())
 		{
-			TAttribute<FText> LabelAttribute = Tool->Label.IsSet() ? Tool->Label.GetValue() : TAttribute<FText>();
-			TAttribute<FText> TooltipAttribute = Tool->Tooltip.IsSet() ? Tool->Tooltip.GetValue() : TAttribute<FText>();
-			TAttribute<FSlateIcon> IconAttribute = Tool->Icon.IsSet() ? Tool->Icon.GetValue() : TAttribute<FSlateIcon>();
-			CategoryToolbarBuilder.AddToolBarButton(Tool->UICommand, NAME_None, LabelAttribute, TooltipAttribute, IconAttribute);
+			CategoryToolbarBuilder.AddToolBarButton(Tool->UICommand);
 		}
 	}
 	else
@@ -152,11 +149,7 @@ void SAvaInteractiveToolsToolbarCategoryButton::ShowCommandsContextMenu()
 	{
 		if (Tool.UICommand.IsValid())
 		{
-			TAttribute<FText> LabelAttribute = Tool.Label.IsSet() ? Tool.Label.GetValue() : TAttribute<FText>();
-			TAttribute<FText> TooltipAttribute = Tool.Tooltip.IsSet() ? Tool.Tooltip.GetValue() : TAttribute<FText>();
-			TAttribute<FSlateIcon> IconAttribute = Tool.Icon.IsSet() ? Tool.Icon.GetValue() : TAttribute<FSlateIcon>();
-
-			CategoryToolbarBuilder.AddMenuEntry(Tool.UICommand, NAME_None, LabelAttribute, TooltipAttribute, IconAttribute.Get());
+			CategoryToolbarBuilder.AddMenuEntry(Tool.UICommand);
 		}
 	}
 
diff --git a/Engine/Plugins/VirtualProduction/Avalanche/Source/AvalancheInteractiveTools/Public/IAvalancheInteractiveToolsModule.h b/Engine/Plugins/VirtualProduction/Avalanche/Source/AvalancheInteractiveTools/Public/IAvalancheInteractiveToolsModule.h
--- a/Engine/Plugins/VirtualProduction/Avalanche/Source/AvalancheInteractiveTools/Public/IAvalancheInteractiveToolsModule.h
+++ b/Engine/Plugins/VirtualProduction/Avalanche/Source/AvalancheInteractiveTools/Public/IAvalancheInteractiveToolsModule.h
@@ -6,6 +6,7 @@
 #include "Containers/UnrealString.h"
 #include "Editor.h"
 #include "Misc/Optional.h"
+#include "Misc/TVariant.h"
 #include "Modules/ModuleInterface.h"
 #include "Modules/ModuleManager.h"
 #include "Templates/SharedPointer.h"
@@ -27,11 +28,7 @@ struct FAvaInteractiveToolsToolParameters
 	FString ToolIdentifier = "";
 	int32 Priority = 0;
 	FAvalancheInteractiveToolsCreateBuilder CreateBuilder;
-	TSubclassOf<AActor> FactoryClass = nullptr;
-	TObjectPtr<UActorFactory> Factory = nullptr;
-	TOptional<FText> Label;
-	TOptional<FText> Tooltip;
-	TOptional<FSlateIcon> Icon;
+	TVariant<TSubclassOf<AActor>, TObjectPtr<UActorFactory>> ActorFactory;
 };
 
 namespace UE::AvaInteractiveTools
diff --git a/Engine/Plugins/VirtualProduction/Avalanche/Source/AvalancheInteractiveTools/Public/Tools/AvaInteractiveToolsActorToolBase.h b/Engine/Plugins/VirtualProduction/Avalanche/Source/AvalancheInteractiveTools/Public/Tools/AvaInteractiveToolsActorToolBase.h
--- a/Engine/Plugins/VirtualProduction/Avalanche/Source/AvalancheInteractiveTools/Public/Tools/AvaInteractiveToolsActorToolBase.h
+++ b/Engine/Plugins/VirtualProduction/Avalanche/Source/AvalancheInteractiveTools/Public/Tools/AvaInteractiveToolsActorToolBase.h
@@ -23,4 +23,6 @@ protected:
 	virtual bool UseIdentityLocation() const override { return bPerformingDefaultAction; }
 	virtual bool UseIdentityRotation() const override;
 	//~ End UAvaInteractiveToolsToolBase
+
+	FAvaInteractiveToolsToolParameters CreateDefaultToolParameters() const;
 };
diff --git a/Engine/Plugins/VirtualProduction/Avalanche/Source/AvalancheInteractiveTools/Public/Tools/AvaInteractiveToolsToolBase.h b/Engine/Plugins/VirtualProduction/Avalanche/Source/AvalancheInteractiveTools/Public/Tools/AvaInteractiveToolsToolBase.h
--- a/Engine/Plugins/VirtualProduction/Avalanche/Source/AvalancheInteractiveTools/Public/Tools/AvaInteractiveToolsToolBase.h
+++ b/Engine/Plugins/VirtualProduction/Avalanche/Source/AvalancheInteractiveTools/Public/Tools/AvaInteractiveToolsToolBase.h
@@ -193,4 +193,9 @@ protected:
 	bool IsMotionDesignViewport() const;
 
 	virtual bool ShouldForceDefaultAction() const;
+
+	void DestroyPreviewActor();
+
+	UFUNCTION()
+	void OnPreviewActorDestroyed(AActor* InDestroyedActor);
 };
diff --git a/Engine/Plugins/VirtualProduction/Avalanche/Source/AvalancheSVGEditor/Private/Tool/AvaSVGActorTool.cpp b/Engine/Plugins/VirtualProduction/Avalanche/Source/AvalancheSVGEditor/Private/Tool/AvaSVGActorTool.cpp
--- a/Engine/Plugins/VirtualProduction/Avalanche/Source/AvalancheSVGEditor/Private/Tool/AvaSVGActorTool.cpp
+++ b/Engine/Plugins/VirtualProduction/Avalanche/Source/AvalancheSVGEditor/Private/Tool/AvaSVGActorTool.cpp
@@ -20,19 +20,10 @@ void UAvaSVGActorTool::OnRegisterTool(IAvalancheInteractiveToolsModule* InAITMod
 	// Load svg module
 	FSVGImporterModule::Get();
 
-	FAvaInteractiveToolsToolParameters ToolParameters =
-	{
-		FSVGImporterEditorCommands::GetExternal().SpawnSVGActor,
-		TEXT("SVG Actor Tool"),
-		6000,
-		FAvalancheInteractiveToolsCreateBuilder::CreateLambda(
-			[](UEdMode* InEdMode)
-			{
-				return UAvaInteractiveToolsToolBuilder::CreateToolBuilder<UAvaSVGActorTool>(InEdMode);
-			}),
-		ActorClass,
-		CreateActorFactory<USVGActorFactory>()
-	};
+	FAvaInteractiveToolsToolParameters ToolParameters = CreateDefaultToolParameters();
+	ToolParameters.UICommand = FSVGImporterEditorCommands::GetExternal().SpawnSVGActor;
+	ToolParameters.Priority = 6000;
+	ToolParameters.ActorFactory.Set<TObjectPtr<UActorFactory>>(CreateActorFactory<USVGActorFactory>());
 
 	InAITModule->RegisterTool(IAvalancheInteractiveToolsModule::CategoryNameActor, MoveTemp(ToolParameters));
 }
diff --git a/Engine/Plugins/VirtualProduction/Avalanche/Source/AvalancheShapesEditor/Private/Tools/AvaShapesEditorShapeTool2DArrow.cpp b/Engine/Plugins/VirtualProduction/Avalanche/Source/AvalancheShapesEditor/Private/Tools/AvaShapesEditorShapeTool2DArrow.cpp
--- a/Engine/Plugins/VirtualProduction/Avalanche/Source/AvalancheShapesEditor/Private/Tools/AvaShapesEditorShapeTool2DArrow.cpp
+++ b/Engine/Plugins/VirtualProduction/Avalanche/Source/AvalancheShapesEditor/Private/Tools/AvaShapesEditorShapeTool2DArrow.cpp
@@ -15,19 +15,10 @@ void UAvaShapesEditorShapeTool2DArrow::OnRegisterTool(IAvalancheInteractiveTools
 {
 	Super::OnRegisterTool(InAITModule);
 
-	FAvaInteractiveToolsToolParameters ToolParameters = 
-	{
-		FAvaShapesEditorCommands::Get().Tool_Shape_2DArrow,
-		TEXT("Parametric 2D Arrow Tool"),
-		7000,
-		FAvalancheInteractiveToolsCreateBuilder::CreateLambda(
-			[](UEdMode* InEdMode)
-			{
-				return UAvaInteractiveToolsToolBuilder::CreateToolBuilder<UAvaShapesEditorShapeTool2DArrow>(InEdMode);
-			}),
-		nullptr,
-		CreateFactory<UAvaShape2DArrowDynamicMesh>()
-	};
+	FAvaInteractiveToolsToolParameters ToolParameters = CreateDefaultToolParameters();
+	ToolParameters.UICommand = FAvaShapesEditorCommands::Get().Tool_Shape_2DArrow;
+	ToolParameters.Priority = 7000;
+	ToolParameters.ActorFactory.Set<TObjectPtr<UActorFactory>>(CreateFactory<UAvaShape2DArrowDynamicMesh>());
 
 	InAITModule->RegisterTool(IAvalancheInteractiveToolsModule::CategoryName2D, MoveTemp(ToolParameters));
 }
diff --git a/Engine/Plugins/VirtualProduction/Avalanche/Source/AvalancheShapesEditor/Private/Tools/AvaShapesEditorShapeToolChevron.cpp b/Engine/Plugins/VirtualProduction/Avalanche/Source/AvalancheShapesEditor/Private/Tools/AvaShapesEditorShapeToolChevron.cpp
--- a/Engine/Plugins/VirtualProduction/Avalanche/Source/AvalancheShapesEditor/Private/Tools/AvaShapesEditorShapeToolChevron.cpp
+++ b/Engine/Plugins/VirtualProduction/Avalanche/Source/AvalancheShapesEditor/Private/Tools/AvaShapesEditorShapeToolChevron.cpp
@@ -15,19 +15,10 @@ void UAvaShapesEditorShapeToolChevron::OnRegisterTool(IAvalancheInteractiveTools
 {
 	Super::OnRegisterTool(InAITModule);
 
-	FAvaInteractiveToolsToolParameters ToolParameters = 
-	{
-		FAvaShapesEditorCommands::Get().Tool_Shape_Chevron,
-		TEXT("Parametric Chevron Tool"),
-		8000,
-		FAvalancheInteractiveToolsCreateBuilder::CreateLambda(
-			[](UEdMode* InEdMode)
-			{
-				return UAvaInteractiveToolsToolBuilder::CreateToolBuilder<UAvaShapesEditorShapeToolChevron>(InEdMode);
-			}),
-		nullptr,
-		CreateFactory<UAvaShapeChevronDynamicMesh>()
-	};
+	FAvaInteractiveToolsToolParameters ToolParameters = CreateDefaultToolParameters();
+	ToolParameters.UICommand = FAvaShapesEditorCommands::Get().Tool_Shape_Chevron;
+	ToolParameters.Priority = 8000;
+	ToolParameters.ActorFactory.Set<TObjectPtr<UActorFactory>>(CreateFactory<UAvaShapeChevronDynamicMesh>());
 
 	InAITModule->RegisterTool(IAvalancheInteractiveToolsModule::CategoryName2D, MoveTemp(ToolParameters));
 }
diff --git a/Engine/Plugins/VirtualProduction/Avalanche/Source/AvalancheShapesEditor/Private/Tools/AvaShapesEditorShapeToolCone.cpp b/Engine/Plugins/VirtualProduction/Avalanche/Source/AvalancheShapesEditor/Private/Tools/AvaShapesEditorShapeToolCone.cpp
--- a/Engine/Plugins/VirtualProduction/Avalanche/Source/AvalancheShapesEditor/Private/Tools/AvaShapesEditorShapeToolCone.cpp
+++ b/Engine/Plugins/VirtualProduction/Avalanche/Source/AvalancheShapesEditor/Private/Tools/AvaShapesEditorShapeToolCone.cpp
@@ -15,19 +15,11 @@ void UAvaShapesEditorShapeToolCone::OnRegisterTool(IAvalancheInteractiveToolsMod
 {
 	Super::OnRegisterTool(InAITModule);
 
-	FAvaInteractiveToolsToolParameters ConeToolParameters = 
-	{
-		FAvaShapesEditorCommands::Get().Tool_Shape_Cone,
-		TEXT("Parametric Cone Tool"),
-		3000,
-		FAvalancheInteractiveToolsCreateBuilder::CreateLambda(
-			[](UEdMode* InEdMode)
-			{
-				return UAvaInteractiveToolsToolBuilder::CreateToolBuilder<UAvaShapesEditorShapeToolCone>(InEdMode);
-			}),
-		nullptr,
-		CreateFactory<UAvaShapeConeDynamicMesh>()
-	};
+	FAvaInteractiveToolsToolParameters ConeToolParameters = CreateDefaultToolParameters();
+	ConeToolParameters.UICommand = FAvaShapesEditorCommands::Get().Tool_Shape_Cone;
+	ConeToolParameters.Priority = 3000;
+	ConeToolParameters.ActorFactory.Set<TObjectPtr<UActorFactory>>(CreateFactory<UAvaShapeConeDynamicMesh>());
+	ConeToolParameters.ToolIdentifier = TEXT("Parametric Cone Tool");
 
 	InAITModule->RegisterTool(IAvalancheInteractiveToolsModule::CategoryName3D, MoveTemp(ConeToolParameters));
 
@@ -41,19 +33,11 @@ void UAvaShapesEditorShapeToolCone::OnRegisterTool(IAvalancheInteractiveToolsMod
 		.NameOverride = FString("Cylinder"),
 	};
 
-	FAvaInteractiveToolsToolParameters CylinderToolParameters = 
-	{
-		FAvaShapesEditorCommands::Get().Tool_Shape_Cylinder,
-		TEXT("Parametric Cylinder Tool"),
-		3000,
-		FAvalancheInteractiveToolsCreateBuilder::CreateLambda(
-			[](UEdMode* InEdMode)
-			{
-				return UAvaInteractiveToolsToolBuilder::CreateToolBuilder<UAvaShapesEditorShapeToolCone>(InEdMode);
-			}),
-		nullptr,
-		CreateFactory<UAvaShapeConeDynamicMesh>(CylinderFactoryParameters)
-	};
+	FAvaInteractiveToolsToolParameters CylinderToolParameters = CreateDefaultToolParameters();
+	CylinderToolParameters.UICommand = FAvaShapesEditorCommands::Get().Tool_Shape_Cylinder;
+	CylinderToolParameters.Priority = 3000;
+	CylinderToolParameters.ActorFactory.Set<TObjectPtr<UActorFactory>>(CreateFactory<UAvaShapeConeDynamicMesh>(CylinderFactoryParameters));
+	CylinderToolParameters.ToolIdentifier = TEXT("Parametric Cylinder Tool");
 
 	InAITModule->RegisterTool(IAvalancheInteractiveToolsModule::CategoryName3D, MoveTemp(CylinderToolParameters));
 }
diff --git a/Engine/Plugins/VirtualProduction/Avalanche/Source/AvalancheShapesEditor/Private/Tools/AvaShapesEditorShapeToolCube.cpp b/Engine/Plugins/VirtualProduction/Avalanche/Source/AvalancheShapesEditor/Private/Tools/AvaShapesEditorShapeToolCube.cpp
--- a/Engine/Plugins/VirtualProduction/Avalanche/Source/AvalancheShapesEditor/Private/Tools/AvaShapesEditorShapeToolCube.cpp
+++ b/Engine/Plugins/VirtualProduction/Avalanche/Source/AvalancheShapesEditor/Private/Tools/AvaShapesEditorShapeToolCube.cpp
@@ -15,19 +15,11 @@ void UAvaShapesEditorShapeToolCube::OnRegisterTool(IAvalancheInteractiveToolsMod
 {
 	Super::OnRegisterTool(InAITModule);
 
-	FAvaInteractiveToolsToolParameters ToolParameters = 
-	{
-		FAvaShapesEditorCommands::Get().Tool_Shape_Cube,
-		TEXT("Parametric Cube Tool"),
-		1000,
-		FAvalancheInteractiveToolsCreateBuilder::CreateLambda(
-			[](UEdMode* InEdMode)
-			{
-				return UAvaInteractiveToolsToolBuilder::CreateToolBuilder<UAvaShapesEditorShapeToolCube>(InEdMode);
-			}),
-		nullptr,
-		CreateFactory<UAvaShapeCubeDynamicMesh>()
-	};
+	FAvaInteractiveToolsToolParameters ToolParameters = CreateDefaultToolParameters();
+	ToolParameters.UICommand = FAvaShapesEditorCommands::Get().Tool_Shape_Cube;
+	ToolParameters.Priority = 1000;
+	ToolParameters.ActorFactory.Set<TObjectPtr<UActorFactory>>(CreateFactory<UAvaShapeCubeDynamicMesh>());
+	ToolParameters.ToolIdentifier = TEXT("Parametric Cube Tool");
 
 	InAITModule->RegisterTool(IAvalancheInteractiveToolsModule::CategoryName3D, MoveTemp(ToolParameters));
 }
diff --git a/Engine/Plugins/VirtualProduction/Avalanche/Source/AvalancheShapesEditor/Private/Tools/AvaShapesEditorShapeToolEllipse.cpp b/Engine/Plugins/VirtualProduction/Avalanche/Source/AvalancheShapesEditor/Private/Tools/AvaShapesEditorShapeToolEllipse.cpp
--- a/Engine/Plugins/VirtualProduction/Avalanche/Source/AvalancheShapesEditor/Private/Tools/AvaShapesEditorShapeToolEllipse.cpp
+++ b/Engine/Plugins/VirtualProduction/Avalanche/Source/AvalancheShapesEditor/Private/Tools/AvaShapesEditorShapeToolEllipse.cpp
@@ -15,19 +15,10 @@ void UAvaShapesEditorShapeToolEllipse::OnRegisterTool(IAvalancheInteractiveTools
 {
 	Super::OnRegisterTool(InAITModule);
 
-	FAvaInteractiveToolsToolParameters ToolParameters = 
-	{
-		FAvaShapesEditorCommands::Get().Tool_Shape_Ellipse,
-		TEXT("Parametric Ellipse Tool"),
-		2000,
-		FAvalancheInteractiveToolsCreateBuilder::CreateLambda(
-			[](UEdMode* InEdMode)
-			{
-				return UAvaInteractiveToolsToolBuilder::CreateToolBuilder<UAvaShapesEditorShapeToolEllipse>(InEdMode);
-			}),
-		nullptr,
-		CreateFactory<UAvaShapeEllipseDynamicMesh>()
-	};
+	FAvaInteractiveToolsToolParameters ToolParameters = CreateDefaultToolParameters();
+	ToolParameters.UICommand = FAvaShapesEditorCommands::Get().Tool_Shape_Ellipse;
+	ToolParameters.Priority = 2000;
+	ToolParameters.ActorFactory.Set<TObjectPtr<UActorFactory>>(CreateFactory<UAvaShapeEllipseDynamicMesh>());
 	
 	InAITModule->RegisterTool(IAvalancheInteractiveToolsModule::CategoryName2D, MoveTemp(ToolParameters));
 }
diff --git a/Engine/Plugins/VirtualProduction/Avalanche/Source/AvalancheShapesEditor/Private/Tools/AvaShapesEditorShapeToolIrregularPoly.cpp b/Engine/Plugins/VirtualProduction/Avalanche/Source/AvalancheShapesEditor/Private/Tools/AvaShapesEditorShapeToolIrregularPoly.cpp
--- a/Engine/Plugins/VirtualProduction/Avalanche/Source/AvalancheShapesEditor/Private/Tools/AvaShapesEditorShapeToolIrregularPoly.cpp
+++ b/Engine/Plugins/VirtualProduction/Avalanche/Source/AvalancheShapesEditor/Private/Tools/AvaShapesEditorShapeToolIrregularPoly.cpp
@@ -25,19 +25,10 @@ void UAvaShapesEditorShapeToolIrregularPoly::OnRegisterTool(IAvalancheInteractiv
 {
 	Super::OnRegisterTool(InAITModule);
 
-	FAvaInteractiveToolsToolParameters ToolParameters =
-	{
-		FAvaShapesEditorCommands::Get().Tool_Shape_IrregularPoly,
-		TEXT("Parametric Freehand Polygon Tool"),
-		4000,
-		FAvalancheInteractiveToolsCreateBuilder::CreateLambda(
-			[](UEdMode* InEdMode)
-			{
-				return UAvaInteractiveToolsToolBuilder::CreateToolBuilder<UAvaShapesEditorShapeToolIrregularPoly>(InEdMode);
-			}),
-		nullptr,
-		CreateFactory<UAvaShapeIrregularPolygonDynamicMesh>()
-	};
+	FAvaInteractiveToolsToolParameters ToolParameters = CreateDefaultToolParameters();
+	ToolParameters.UICommand = FAvaShapesEditorCommands::Get().Tool_Shape_IrregularPoly;
+	ToolParameters.Priority = 4000;
+	ToolParameters.ActorFactory.Set<TObjectPtr<UActorFactory>>(CreateFactory<UAvaShapeIrregularPolygonDynamicMesh>());
 
 	InAITModule->RegisterTool(IAvalancheInteractiveToolsModule::CategoryName2D, MoveTemp(ToolParameters));
 }
diff --git a/Engine/Plugins/VirtualProduction/Avalanche/Source/AvalancheShapesEditor/Private/Tools/AvaShapesEditorShapeToolLine.cpp b/Engine/Plugins/VirtualProduction/Avalanche/Source/AvalancheShapesEditor/Private/Tools/AvaShapesEditorShapeToolLine.cpp
--- a/Engine/Plugins/VirtualProduction/Avalanche/Source/AvalancheShapesEditor/Private/Tools/AvaShapesEditorShapeToolLine.cpp
+++ b/Engine/Plugins/VirtualProduction/Avalanche/Source/AvalancheShapesEditor/Private/Tools/AvaShapesEditorShapeToolLine.cpp
@@ -72,19 +72,10 @@ void UAvaShapesEditorShapeToolLine::OnRegisterTool(IAvalancheInteractiveToolsMod
 		}
 	};
 
-	FAvaInteractiveToolsToolParameters ToolParameters =
-	{
-		FAvaShapesEditorCommands::Get().Tool_Shape_Line,
-		TEXT("Parametric Line Tool"),
-		9000,
-		FAvalancheInteractiveToolsCreateBuilder::CreateLambda(
-			[](UEdMode* InEdMode)
-			{
-				return UAvaInteractiveToolsToolBuilder::CreateToolBuilder<UAvaShapesEditorShapeToolLine>(InEdMode);
-			}),
-		nullptr,
-		CreateFactory<UAvaShapeLineDynamicMesh>(LineFactoryParameters)
-	};
+	FAvaInteractiveToolsToolParameters ToolParameters = CreateDefaultToolParameters();
+	ToolParameters.UICommand = FAvaShapesEditorCommands::Get().Tool_Shape_Line;
+	ToolParameters.Priority = 9000;
+	ToolParameters.ActorFactory.Set<TObjectPtr<UActorFactory>>(CreateFactory<UAvaShapeLineDynamicMesh>(LineFactoryParameters));
 
 	InAITModule->RegisterTool(IAvalancheInteractiveToolsModule::CategoryName2D, MoveTemp(ToolParameters));
 }
diff --git a/Engine/Plugins/VirtualProduction/Avalanche/Source/AvalancheShapesEditor/Private/Tools/AvaShapesEditorShapeToolNGon.cpp b/Engine/Plugins/VirtualProduction/Avalanche/Source/AvalancheShapesEditor/Private/Tools/AvaShapesEditorShapeToolNGon.cpp
--- a/Engine/Plugins/VirtualProduction/Avalanche/Source/AvalancheShapesEditor/Private/Tools/AvaShapesEditorShapeToolNGon.cpp
+++ b/Engine/Plugins/VirtualProduction/Avalanche/Source/AvalancheShapesEditor/Private/Tools/AvaShapesEditorShapeToolNGon.cpp
@@ -15,19 +15,10 @@ void UAvaShapesEditorShapeToolNGon::OnRegisterTool(IAvalancheInteractiveToolsMod
 {
 	Super::OnRegisterTool(InAITModule);
 
-	FAvaInteractiveToolsToolParameters ToolParameters =
-	{
-		FAvaShapesEditorCommands::Get().Tool_Shape_NGon,
-		TEXT("Parametric Regular Polygon Tool"),
-		3000,
-		FAvalancheInteractiveToolsCreateBuilder::CreateLambda(
-			[](UEdMode* InEdMode)
-			{
-				return UAvaInteractiveToolsToolBuilder::CreateToolBuilder<UAvaShapesEditorShapeToolNGon>(InEdMode);
-			}),
-		nullptr,
-		CreateFactory<UAvaShapeNGonDynamicMesh>()
-	};
+	FAvaInteractiveToolsToolParameters ToolParameters = CreateDefaultToolParameters();
+	ToolParameters.UICommand = FAvaShapesEditorCommands::Get().Tool_Shape_NGon;
+	ToolParameters.Priority = 3000;
+	ToolParameters.ActorFactory.Set<TObjectPtr<UActorFactory>>(CreateFactory<UAvaShapeNGonDynamicMesh>());
 	
 	InAITModule->RegisterTool(IAvalancheInteractiveToolsModule::CategoryName2D, MoveTemp(ToolParameters));
 }
diff --git a/Engine/Plugins/VirtualProduction/Avalanche/Source/AvalancheShapesEditor/Private/Tools/AvaShapesEditorShapeToolRectangle.cpp b/Engine/Plugins/VirtualProduction/Avalanche/Source/AvalancheShapesEditor/Private/Tools/AvaShapesEditorShapeToolRectangle.cpp
--- a/Engine/Plugins/VirtualProduction/Avalanche/Source/AvalancheShapesEditor/Private/Tools/AvaShapesEditorShapeToolRectangle.cpp
+++ b/Engine/Plugins/VirtualProduction/Avalanche/Source/AvalancheShapesEditor/Private/Tools/AvaShapesEditorShapeToolRectangle.cpp
@@ -26,19 +26,11 @@ void UAvaShapesEditorShapeToolRectangle::OnRegisterTool(IAvalancheInteractiveToo
 		}
 	};
 
-	FAvaInteractiveToolsToolParameters RectangleToolParameters = 
-	{
-		FAvaShapesEditorCommands::Get().Tool_Shape_Rectangle,
-		TEXT("Parametric Rectangle Tool"),
-		1000,
-		FAvalancheInteractiveToolsCreateBuilder::CreateLambda(
-			[](UEdMode* InEdMode)
-			{
-				return UAvaInteractiveToolsToolBuilder::CreateToolBuilder<UAvaShapesEditorShapeToolRectangle>(InEdMode);
-			}),
-		nullptr,
-		CreateFactory<UAvaShapeRectangleDynamicMesh>(RectangleFactoryParameters)
-	};
+	FAvaInteractiveToolsToolParameters RectangleToolParameters = CreateDefaultToolParameters();
+	RectangleToolParameters.UICommand = FAvaShapesEditorCommands::Get().Tool_Shape_Rectangle;
+	RectangleToolParameters.Priority = 1000;
+	RectangleToolParameters.ActorFactory.Set<TObjectPtr<UActorFactory>>(CreateFactory<UAvaShapeRectangleDynamicMesh>(RectangleFactoryParameters));
+	RectangleToolParameters.ToolIdentifier = TEXT("Parametric Rectangle Tool");
 
 	InAITModule->RegisterTool(IAvalancheInteractiveToolsModule::CategoryName2D, MoveTemp(RectangleToolParameters));
 
@@ -47,19 +39,11 @@ void UAvaShapesEditorShapeToolRectangle::OnRegisterTool(IAvalancheInteractiveToo
 		.NameOverride = FString("Square")
 	};
 
-	FAvaInteractiveToolsToolParameters SquareToolParameters = 
-	{
-		FAvaShapesEditorCommands::Get().Tool_Shape_Square,
-		TEXT("Parametric Square Tool"),
-		1001,
-		FAvalancheInteractiveToolsCreateBuilder::CreateLambda(
-			[](UEdMode* InEdMode)
-			{
-				return UAvaInteractiveToolsToolBuilder::CreateToolBuilder<UAvaShapesEditorShapeToolRectangle>(InEdMode);
-			}),
-		nullptr,
-		CreateFactory<UAvaShapeRectangleDynamicMesh>(SquareFactoryParameters)
-	};
+	FAvaInteractiveToolsToolParameters SquareToolParameters = CreateDefaultToolParameters();
+	SquareToolParameters.UICommand = FAvaShapesEditorCommands::Get().Tool_Shape_Square;
+	SquareToolParameters.Priority = 1100;
+	SquareToolParameters.ActorFactory.Set<TObjectPtr<UActorFactory>>(CreateFactory<UAvaShapeRectangleDynamicMesh>(SquareFactoryParameters));
+	SquareToolParameters.ToolIdentifier = TEXT("Parametric Square Tool");
 
 	InAITModule->RegisterTool(IAvalancheInteractiveToolsModule::CategoryName2D, MoveTemp(SquareToolParameters));
 }
diff --git a/Engine/Plugins/VirtualProduction/Avalanche/Source/AvalancheShapesEditor/Private/Tools/AvaShapesEditorShapeToolRing.cpp b/Engine/Plugins/VirtualProduction/Avalanche/Source/AvalancheShapesEditor/Private/Tools/AvaShapesEditorShapeToolRing.cpp
--- a/Engine/Plugins/VirtualProduction/Avalanche/Source/AvalancheShapesEditor/Private/Tools/AvaShapesEditorShapeToolRing.cpp
+++ b/Engine/Plugins/VirtualProduction/Avalanche/Source/AvalancheShapesEditor/Private/Tools/AvaShapesEditorShapeToolRing.cpp
@@ -15,19 +15,10 @@ void UAvaShapesEditorShapeToolRing::OnRegisterTool(IAvalancheInteractiveToolsMod
 {
 	Super::OnRegisterTool(InAITModule);
 
-	FAvaInteractiveToolsToolParameters ToolParameters = 
-	{
-		FAvaShapesEditorCommands::Get().Tool_Shape_Ring,
-		TEXT("Parametric Ring Tool"),
-		5000,
-		FAvalancheInteractiveToolsCreateBuilder::CreateLambda(
-			[](UEdMode* InEdMode)
-			{
-				return UAvaInteractiveToolsToolBuilder::CreateToolBuilder<UAvaShapesEditorShapeToolRing>(InEdMode);
-			}),
-		nullptr,
-		CreateFactory<UAvaShapeRingDynamicMesh>()
-	};
+	FAvaInteractiveToolsToolParameters ToolParameters = CreateDefaultToolParameters();
+	ToolParameters.UICommand = FAvaShapesEditorCommands::Get().Tool_Shape_Ring;
+	ToolParameters.Priority = 5000;
+	ToolParameters.ActorFactory.Set<TObjectPtr<UActorFactory>>(CreateFactory<UAvaShapeRingDynamicMesh>());
 
 	InAITModule->RegisterTool(IAvalancheInteractiveToolsModule::CategoryName2D, MoveTemp(ToolParameters));
 }
diff --git a/Engine/Plugins/VirtualProduction/Avalanche/Source/AvalancheShapesEditor/Private/Tools/AvaShapesEditorShapeToolSphere.cpp b/Engine/Plugins/VirtualProduction/Avalanche/Source/AvalancheShapesEditor/Private/Tools/AvaShapesEditorShapeToolSphere.cpp
--- a/Engine/Plugins/VirtualProduction/Avalanche/Source/AvalancheShapesEditor/Private/Tools/AvaShapesEditorShapeToolSphere.cpp
+++ b/Engine/Plugins/VirtualProduction/Avalanche/Source/AvalancheShapesEditor/Private/Tools/AvaShapesEditorShapeToolSphere.cpp
@@ -15,19 +15,10 @@ void UAvaShapesEditorShapeToolSphere::OnRegisterTool(IAvalancheInteractiveToolsM
 {
 	Super::OnRegisterTool(InAITModule);
 
-	FAvaInteractiveToolsToolParameters ToolParameters = 
-	{
-		FAvaShapesEditorCommands::Get().Tool_Shape_Sphere,
-		TEXT("Parametric Sphere Tool"),
-		2000,
-		FAvalancheInteractiveToolsCreateBuilder::CreateLambda(
-			[](UEdMode* InEdMode)
-			{
-				return UAvaInteractiveToolsToolBuilder::CreateToolBuilder<UAvaShapesEditorShapeToolSphere>(InEdMode);
-			}),
-		nullptr,
-		CreateFactory<UAvaShapeSphereDynamicMesh>()
-	};
+	FAvaInteractiveToolsToolParameters ToolParameters = CreateDefaultToolParameters();
+	ToolParameters.UICommand = FAvaShapesEditorCommands::Get().Tool_Shape_Sphere;
+	ToolParameters.Priority = 2000;
+	ToolParameters.ActorFactory.Set<TObjectPtr<UActorFactory>>(CreateFactory<UAvaShapeSphereDynamicMesh>());
 	
 	InAITModule->RegisterTool(IAvalancheInteractiveToolsModule::CategoryName3D, MoveTemp(ToolParameters));
 }
diff --git a/Engine/Plugins/VirtualProduction/Avalanche/Source/AvalancheShapesEditor/Private/Tools/AvaShapesEditorShapeToolStar.cpp b/Engine/Plugins/VirtualProduction/Avalanche/Source/AvalancheShapesEditor/Private/Tools/AvaShapesEditorShapeToolStar.cpp
--- a/Engine/Plugins/VirtualProduction/Avalanche/Source/AvalancheShapesEditor/Private/Tools/AvaShapesEditorShapeToolStar.cpp
+++ b/Engine/Plugins/VirtualProduction/Avalanche/Source/AvalancheShapesEditor/Private/Tools/AvaShapesEditorShapeToolStar.cpp
@@ -15,19 +15,10 @@ void UAvaShapesEditorShapeToolStar::OnRegisterTool(IAvalancheInteractiveToolsMod
 {
 	Super::OnRegisterTool(InAITModule);
 
-	FAvaInteractiveToolsToolParameters ToolParameters = 
-	{
-		FAvaShapesEditorCommands::Get().Tool_Shape_Star,
-		TEXT("Parametric Star Tool"),
-		6000,
-		FAvalancheInteractiveToolsCreateBuilder::CreateLambda(
-			[](UEdMode* InEdMode)
-			{
-				return UAvaInteractiveToolsToolBuilder::CreateToolBuilder<UAvaShapesEditorShapeToolStar>(InEdMode);
-			}),
-		nullptr,
-		CreateFactory<UAvaShapeStarDynamicMesh>()
-	};
+	FAvaInteractiveToolsToolParameters ToolParameters = CreateDefaultToolParameters();
+	ToolParameters.UICommand = FAvaShapesEditorCommands::Get().Tool_Shape_Star;
+	ToolParameters.Priority = 6000;
+	ToolParameters.ActorFactory.Set<TObjectPtr<UActorFactory>>(CreateFactory<UAvaShapeStarDynamicMesh>());
 	
 	InAITModule->RegisterTool(IAvalancheInteractiveToolsModule::CategoryName2D, MoveTemp(ToolParameters));
 }
diff --git a/Engine/Plugins/VirtualProduction/Avalanche/Source/AvalancheShapesEditor/Private/Tools/AvaShapesEditorShapeToolTorus.cpp b/Engine/Plugins/VirtualProduction/Avalanche/Source/AvalancheShapesEditor/Private/Tools/AvaShapesEditorShapeToolTorus.cpp
--- a/Engine/Plugins/VirtualProduction/Avalanche/Source/AvalancheShapesEditor/Private/Tools/AvaShapesEditorShapeToolTorus.cpp
+++ b/Engine/Plugins/VirtualProduction/Avalanche/Source/AvalancheShapesEditor/Private/Tools/AvaShapesEditorShapeToolTorus.cpp
@@ -15,19 +15,10 @@ void UAvaShapesEditorShapeToolTorus::OnRegisterTool(IAvalancheInteractiveToolsMo
 {
 	Super::OnRegisterTool(InAITModule);
 
-	FAvaInteractiveToolsToolParameters ToolParameters = 
-	{
-		FAvaShapesEditorCommands::Get().Tool_Shape_Torus,
-		TEXT("Parametric Torus Tool"),
-		4000,
-		FAvalancheInteractiveToolsCreateBuilder::CreateLambda(
-			[](UEdMode* InEdMode)
-			{
-				return UAvaInteractiveToolsToolBuilder::CreateToolBuilder<UAvaShapesEditorShapeToolTorus>(InEdMode);
-			}),
-		nullptr,
-		CreateFactory<UAvaShapeTorusDynamicMesh>()
-	};
-	
+	FAvaInteractiveToolsToolParameters ToolParameters = CreateDefaultToolParameters();
+	ToolParameters.UICommand = FAvaShapesEditorCommands::Get().Tool_Shape_Torus;
+	ToolParameters.Priority = 4000;
+	ToolParameters.ActorFactory.Set<TObjectPtr<UActorFactory>>(CreateFactory<UAvaShapeTorusDynamicMesh>());
+
 	InAITModule->RegisterTool(IAvalancheInteractiveToolsModule::CategoryName3D, MoveTemp(ToolParameters));
 }
diff --git a/Engine/Plugins/VirtualProduction/Avalanche/Source/AvalancheTextEditor/Private/Tool/AvaTextActorTool.cpp b/Engine/Plugins/VirtualProduction/Avalanche/Source/AvalancheTextEditor/Private/Tool/AvaTextActorTool.cpp
--- a/Engine/Plugins/VirtualProduction/Avalanche/Source/AvalancheTextEditor/Private/Tool/AvaTextActorTool.cpp
+++ b/Engine/Plugins/VirtualProduction/Avalanche/Source/AvalancheTextEditor/Private/Tool/AvaTextActorTool.cpp
@@ -17,19 +17,10 @@ void UAvaTextActorTool::OnRegisterTool(IAvalancheInteractiveToolsModule* InAITMo
 {
 	Super::OnRegisterTool(InAITModule);
 
-	FAvaInteractiveToolsToolParameters ToolParameters =
-	{
-		FAvaTextEditorCommands::Get().Tool_Actor_Text3D,
-		TEXT("Text Actor Tool"),
-		1000,
-		FAvalancheInteractiveToolsCreateBuilder::CreateLambda(
-			[](UEdMode* InEdMode)
-			{
-				return UAvaInteractiveToolsToolBuilder::CreateToolBuilder<UAvaTextActorTool>(InEdMode);
-			}),
-		ActorClass,
-		CreateActorFactory<UAvaTextActorFactory>()
-	};
+	FAvaInteractiveToolsToolParameters ToolParameters = CreateDefaultToolParameters();
+	ToolParameters.UICommand = FAvaTextEditorCommands::Get().Tool_Actor_Text3D;
+	ToolParameters.Priority = 1000;
+	ToolParameters.ActorFactory.Set<TObjectPtr<UActorFactory>>(CreateActorFactory<UAvaTextActorFactory>());
 
 	InAITModule->RegisterTool(IAvalancheInteractiveToolsModule::CategoryNameActor, MoveTemp(ToolParameters));
 }
