Motion Design: added display name for motion design tag structs and classes.

diff --git a/Engine/Plugins/VirtualProduction/Avalanche/Source/AvalancheTag/Public/AvaTag.h b/Engine/Plugins/VirtualProduction/Avalanche/Source/AvalancheTag/Public/AvaTag.h
--- a/Engine/Plugins/VirtualProduction/Avalanche/Source/AvalancheTag/Public/AvaTag.h
+++ b/Engine/Plugins/VirtualProduction/Avalanche/Source/AvalancheTag/Public/AvaTag.h
@@ -9,7 +9,7 @@
  * Struct for a Tag
  * NOTE: Do not use as storage type. Prefer using FAvaTagHandle to then retrieve the FAvaTag it references
  */
-USTRUCT(BlueprintType)
+USTRUCT(BlueprintType, DisplayName="Motion Design Tag")
 struct FAvaTag
 {
 	GENERATED_BODY()
diff --git a/Engine/Plugins/VirtualProduction/Avalanche/Source/AvalancheTag/Public/AvaTagAlias.h b/Engine/Plugins/VirtualProduction/Avalanche/Source/AvalancheTag/Public/AvaTagAlias.h
--- a/Engine/Plugins/VirtualProduction/Avalanche/Source/AvalancheTag/Public/AvaTagAlias.h
+++ b/Engine/Plugins/VirtualProduction/Avalanche/Source/AvalancheTag/Public/AvaTagAlias.h
@@ -13,7 +13,7 @@ class UAvaTagCollection;
  * Unlike Tag Containers which would need to be updated in every place it's used when the set of tags it needs to manipulate is added to or removed from,
  * Aliases are a layer of abstraction that allows the set of tags to be added to or removed from without affecting the places where the alias is used.
  */
-USTRUCT(BlueprintType)
+USTRUCT(BlueprintType, DisplayName="Motion Design Tag Alias")
 struct FAvaTagAlias
 {
 	GENERATED_BODY()
diff --git a/Engine/Plugins/VirtualProduction/Avalanche/Source/AvalancheTag/Public/AvaTagCollection.h b/Engine/Plugins/VirtualProduction/Avalanche/Source/AvalancheTag/Public/AvaTagCollection.h
--- a/Engine/Plugins/VirtualProduction/Avalanche/Source/AvalancheTag/Public/AvaTagCollection.h
+++ b/Engine/Plugins/VirtualProduction/Avalanche/Source/AvalancheTag/Public/AvaTagCollection.h
@@ -15,7 +15,7 @@ struct FAvaTagList;
  * Tag Collection that identifies a tag with an underlying Tag Id Guid
  * and provides Tag reference capabilities
  */
-UCLASS(MinimalAPI)
+UCLASS(MinimalAPI, DisplayName="Motion Design Tag Collection")
 class UAvaTagCollection : public UObject
 {
 	GENERATED_BODY()
diff --git a/Engine/Plugins/VirtualProduction/Avalanche/Source/AvalancheTag/Public/AvaTagHandle.h b/Engine/Plugins/VirtualProduction/Avalanche/Source/AvalancheTag/Public/AvaTagHandle.h
--- a/Engine/Plugins/VirtualProduction/Avalanche/Source/AvalancheTag/Public/AvaTagHandle.h
+++ b/Engine/Plugins/VirtualProduction/Avalanche/Source/AvalancheTag/Public/AvaTagHandle.h
@@ -14,7 +14,7 @@ struct FAvaTagList;
  * Handle to a Tag or Alias (set of Tags) in a particular Source.
  * This should be used by the UStructs/UObjects to properly reference a Tag or Alias (set of Tags).
  */
-USTRUCT(BlueprintType)
+USTRUCT(BlueprintType, DisplayName="Motion Design Tag Handle")
 struct FAvaTagHandle
 {
 	GENERATED_BODY()
diff --git a/Engine/Plugins/VirtualProduction/Avalanche/Source/AvalancheTag/Public/AvaTagHandleContainer.h b/Engine/Plugins/VirtualProduction/Avalanche/Source/AvalancheTag/Public/AvaTagHandleContainer.h
--- a/Engine/Plugins/VirtualProduction/Avalanche/Source/AvalancheTag/Public/AvaTagHandleContainer.h
+++ b/Engine/Plugins/VirtualProduction/Avalanche/Source/AvalancheTag/Public/AvaTagHandleContainer.h
@@ -14,35 +14,35 @@ struct FAvaTagHandle;
  * Handle to a multiple tags in a particular Source.
  * This should be used by the UStructs/UObjects to properly reference a multiple FAvaTags.
  */
-USTRUCT(BlueprintType)
-struct AVALANCHETAG_API FAvaTagHandleContainer
+USTRUCT(BlueprintType, DisplayName="Motion Design Tag Handle Container")
+struct FAvaTagHandleContainer
 {
 	GENERATED_BODY()
 
 	FAvaTagHandleContainer() = default;
 
-	explicit FAvaTagHandleContainer(const FAvaTagHandle& InTagHandle);
+	AVALANCHETAG_API explicit FAvaTagHandleContainer(const FAvaTagHandle& InTagHandle);
 
 	/** Returns true if the Tag Handles resolve to same valued FAvaTags, even if the Source or Tag Id is different */
-	bool ContainsTag(const FAvaTagHandle& InTagHandle) const;
+	AVALANCHETAG_API bool ContainsTag(const FAvaTagHandle& InTagHandle) const;
 
 	/** Returns true if the Tag Handles is the exact same as the other (Same Source and Tag Id) */
-	bool ContainsTagHandle(const FAvaTagHandle& InTagHandle) const;
+	AVALANCHETAG_API bool ContainsTagHandle(const FAvaTagHandle& InTagHandle) const;
 
-	FString ToString() const;
+	AVALANCHETAG_API FString ToString() const;
 
-	void PostSerialize(const FArchive& Ar);
+	AVALANCHETAG_API void PostSerialize(const FArchive& Ar);
 
-	bool SerializeFromMismatchedTag(const FPropertyTag& InPropertyTag, FStructuredArchive::FSlot InSlot);
+	AVALANCHETAG_API bool SerializeFromMismatchedTag(const FPropertyTag& InPropertyTag, FStructuredArchive::FSlot InSlot);
 
 	/** Adds the provided Tag Handle to TagIds. Returns true only if it was added as a new entry to TagIds */
-	bool AddTagHandle(const FAvaTagHandle& InTagHandle);
+	AVALANCHETAG_API bool AddTagHandle(const FAvaTagHandle& InTagHandle);
 
 	/** Removes the provided Tag Handle from TagIds. Returns true only if it existed and was removed from TagIds */
-	bool RemoveTagHandle(const FAvaTagHandle& InTagHandle);
+	AVALANCHETAG_API bool RemoveTagHandle(const FAvaTagHandle& InTagHandle);
 
 	/** Returns an array of resolved tags through this container's tag ids and source tag collection */
-	TArray<FAvaTag> ResolveTags() const;
+	AVALANCHETAG_API TArray<FAvaTag> ResolveTags() const;
 
 	TConstArrayView<FAvaTagId> GetTagIds() const
 	{
diff --git a/Engine/Plugins/VirtualProduction/Avalanche/Source/AvalancheTag/Public/AvaTagSoftHandle.h b/Engine/Plugins/VirtualProduction/Avalanche/Source/AvalancheTag/Public/AvaTagSoftHandle.h
--- a/Engine/Plugins/VirtualProduction/Avalanche/Source/AvalancheTag/Public/AvaTagSoftHandle.h
+++ b/Engine/Plugins/VirtualProduction/Avalanche/Source/AvalancheTag/Public/AvaTagSoftHandle.h
@@ -13,8 +13,8 @@ struct FAvaTagHandle;
  * Soft reference handle to a Tag in a particular Source Tag Collection
  * This should be used when needing to soft reference a particular FAvaTag.
  */
-USTRUCT(BlueprintType)
-struct AVALANCHETAG_API FAvaTagSoftHandle
+USTRUCT(BlueprintType, DisplayName="Motion Design Tag Soft Handle")
+struct FAvaTagSoftHandle
 {
 	GENERATED_BODY()
 
@@ -26,18 +26,18 @@ struct AVALANCHETAG_API FAvaTagSoftHandle
 	{
 	}
 
-	explicit FAvaTagSoftHandle(const FAvaTagHandle& InTagHandle);
+	AVALANCHETAG_API explicit FAvaTagSoftHandle(const FAvaTagHandle& InTagHandle);
 
 	/**
 	 * Creates a Tag Handle from this Soft Handle.
 	 * Warning: This loads in the Source Tag Collection if not already loaded.
 	 */
-	FAvaTagHandle MakeTagHandle() const;
+	AVALANCHETAG_API FAvaTagHandle MakeTagHandle() const;
 
 	/** Returns true if this Soft Tag Handle matches the provided Tag Handle (Same Source and Tag Id) */
-	bool MatchesExact(const FAvaTagHandle& InTagHandle) const;
+	AVALANCHETAG_API bool MatchesExact(const FAvaTagHandle& InTagHandle) const;
 
-	void PostSerialize(const FArchive& Ar);
+	AVALANCHETAG_API void PostSerialize(const FArchive& Ar);
 
 	bool IsValid() const
 	{
