[Motion Design] Rundown: Integration of RC Event Controllers

The stateless (event) controller calls RCP OnControllerModified()

The handler is implemented in FAvaRundownManagedInstance so that it can route the events to the rundown either for the rundown editor or rundown server.

The rundown has a new function to push the controller events. This is re-using the same RC update API, but a new flag is added to indicate (to the playable rc update logic) if is an event update. In that case, the RC Values are considered "event payload" and will execute the controller behaviors regardless of value controller changes.

Rundown Server: Updated to support combo pages with multiple managed instances per page (opportunistic change) along with adding the controller event routing.

Known limitation (unchanged):
Rundown server and editor can't be editing different pages at the same time, that will cause conflicts with the managed instances (there can only be one managed instance per asset). This remains an issue in this change as well. It is not fixed, but it is not worse either.

diff --git a/Engine/Plugins/VirtualProduction/Avalanche/Source/AvalancheMedia/Internal/Rundown/AvaRundownManagedInstance.h b/Engine/Plugins/VirtualProduction/Avalanche/Source/AvalancheMedia/Internal/Rundown/AvaRundownManagedInstance.h
--- a/Engine/Plugins/VirtualProduction/Avalanche/Source/AvalancheMedia/Internal/Rundown/AvaRundownManagedInstance.h
+++ b/Engine/Plugins/VirtualProduction/Avalanche/Source/AvalancheMedia/Internal/Rundown/AvaRundownManagedInstance.h
@@ -10,6 +10,7 @@
 
 class FAvaRundownManagedInstanceCache;
 class IAvaSceneInterface;
+class UAvaRundown;
 class URemoteControlPreset;
 struct FAvaPlayableRemoteControlPresetInfo;
 
@@ -41,14 +42,25 @@ public:
 	virtual IAvaSceneInterface* GetSceneInterface() const { return nullptr; }
 
 	const FSoftObjectPath& GetSourceAssetPath() const { return SourceAssetPath; }
-	
+
+	/** Set the current rundown page for routing controller events. */
+	void SetCurrentRundownPageForEvents(UAvaRundown* InRundown, int32 InPageId);
+
+	/** Returns the rundown currently set to receive the RC Controller events. */
+	UAvaRundown* GetRundownForEvents() const
+	{
+		return RundownForEvents.Get();
+	}
+
 protected:
 	FAvaPlayableRemoteControlPresetInfo& GetRemoteControlPresetInfoFromCache(URemoteControlPreset* InRemoteControlPreset) const;
 	void InvalidateFromCache();
 
 	void RegisterSourceRemoteControlPresetDelegates(URemoteControlPreset* InSourceRemoteControlPreset);
 	void UnregisterSourceRemoteControlPresetDelegates(URemoteControlPreset* InSourceRemoteControlPreset) const;
-	
+	void RegisterManagedRemoteControlPresetDelegates(URemoteControlPreset* InManagedRemoteControlPreset);
+	void UnregisterManagedRemoteControlPresetDelegates(URemoteControlPreset* InManagedRemoteControlPreset) const;
+
 	void OnSourceRemoteControlEntitiesExposed(URemoteControlPreset* InSourcePreset, const FGuid& InEntityId);
 	void OnSourceRemoteControlEntitiesUnexposed(URemoteControlPreset* InSourcePreset, const FGuid& InEntityId);
 	void OnSourceRemoteControlEntitiesUpdated(URemoteControlPreset* InSourcePreset, const TSet<FGuid>& InModifiedEntities);
@@ -58,9 +70,17 @@ protected:
 	void OnSourceRemoteControlControllerRenamed(URemoteControlPreset* InSourcePreset, const FName InOldLabel, const FName InNewLabel);
 	void OnSourceRemoteControlControllerModified(URemoteControlPreset* InSourcePreset, const TSet<FGuid>& InModifiedControllerIds);
 
+	void OnManagedRemoteControlControllerModified(URemoteControlPreset* InPreset, const TSet<FGuid>& InModifiedControllerIds);
+
 protected:
 	FAvaRundownManagedInstanceCache* ParentCache;
 	FSoftObjectPath SourceAssetPath;
 	FAvaPlayableRemoteControlValues DefaultRemoteControlValues;
 	mutable TSharedPtr<FAvaPlayableRemoteControlPresetInfo> RemoteControlPresetInfo;
+
+	/** Current rundown for routing controller events. */
+	TWeakObjectPtr<UAvaRundown> RundownForEvents;
+
+	/** Current active page for routing controller events. */
+	int32 ActivePageId;
 };
\ No newline at end of file
diff --git a/Engine/Plugins/VirtualProduction/Avalanche/Source/AvalancheMedia/Internal/Rundown/AvaRundownManagedInstanceCache.h b/Engine/Plugins/VirtualProduction/Avalanche/Source/AvalancheMedia/Internal/Rundown/AvaRundownManagedInstanceCache.h
--- a/Engine/Plugins/VirtualProduction/Avalanche/Source/AvalancheMedia/Internal/Rundown/AvaRundownManagedInstanceCache.h
+++ b/Engine/Plugins/VirtualProduction/Avalanche/Source/AvalancheMedia/Internal/Rundown/AvaRundownManagedInstanceCache.h
@@ -3,14 +3,22 @@
 #pragma once
 
 #include "Rundown/AvaRundownManagedInstance.h"
+#include "Rundown/AvaRundownManagedInstanceHandle.h"
 
 class IAvaMediaSyncProvider;
+class UAvaRundown;
 class UPackage;
+struct FAvaRundownPage;
 
 /**
  * Cache for the managed Motion Design instances.
  *
- * This is LRU cache. 
+ * This is LRU cache.
+ *
+ * Todo: This needs to be refactored as a pool instead.
+ * A managed asset instance is needed for each page that is edited. Since rundown editor and server
+ * can edit different pages (of the same template), it needs more than one instance and it needs to be
+ * tied to the rundown (for RC event routing).
  */
 class AVALANCHEMEDIA_API FAvaRundownManagedInstanceCache
 {
@@ -68,6 +76,25 @@ public:
 	 */
 	void FinishPendingActions();
 
+	/**
+	 * Utility function to get all the managed instances for a given page, including all sub-pages in case of a combo page. 
+	 * @param InRundown Rundown owning the current page.
+	 * @param InPage Page to get the assets from.
+	 * @return managed asset instances for the given page.
+	 */
+	TArray<TSharedPtr<FAvaRundownManagedInstance>> GetManagedInstancesForPage(const UAvaRundown* InRundown, const FAvaRundownPage& InPage);
+
+	/**
+	 * Utility function to acquire all the managed instances for a given page, including all sub-pages in case of a combo page.
+	 * This will register the given rundown page for controller events routing in the managed instances (potentially overriding
+	 * previous rundown or page).
+	 * 
+	 * @param InRundown Rundown owning the current page.
+	 * @param InPage Page to get the assets from.
+	 * @return managed asset instances for the given page.
+	 */
+	FAvaRundownManagedInstanceHandles GetManagedHandlesForPage(UAvaRundown* InRundown, const FAvaRundownPage& InPage);
+	
 private:
 	void RemovePendingInvalidatedPaths();
 	void RemoveEntry(const FSoftObjectPath& InAssetPath);
diff --git a/Engine/Plugins/VirtualProduction/Avalanche/Source/AvalancheMedia/Internal/Rundown/AvaRundownManagedInstanceHandle.h b/Engine/Plugins/VirtualProduction/Avalanche/Source/AvalancheMedia/Internal/Rundown/AvaRundownManagedInstanceHandle.h
new file mode 100644
--- /dev/null
+++ b/Engine/Plugins/VirtualProduction/Avalanche/Source/AvalancheMedia/Internal/Rundown/AvaRundownManagedInstanceHandle.h
@@ -0,0 +1,70 @@
+// Copyright Epic Games, Inc. All Rights Reserved.
+
+#pragma once
+
+#include "Containers/Array.h"
+#include "Templates/SharedPointer.h"
+#include "UObject/WeakObjectPtr.h"
+#include "UObject/WeakObjectPtrTemplates.h"
+
+#define UE_API AVALANCHEMEDIA_API
+
+class FAvaRundownManagedInstance;
+class UAvaRundown;
+
+/**
+ * Set of rundown managed asset instances tied to a given rundown for RC editing operations.
+ * It will hook the controller events routing to the provided rundown and page on construction
+ * and handles auto-unregister of Rundown on destruction.
+ *
+ * As this object represents the ownership of the "event routing" for a given rundown and page, it is not copiable.
+ */
+struct FAvaRundownManagedInstanceHandles
+{
+	FAvaRundownManagedInstanceHandles() = default;
+	
+	/**
+	 * The constructor will register the given rundown page for controller events routing in the managed instances (potentially overriding
+	 * previous rundown or page).
+	 */
+	UE_API FAvaRundownManagedInstanceHandles(const TArray<TSharedPtr<FAvaRundownManagedInstance>>& InInstances, UAvaRundown* InRundown, int32 InPageId);
+
+	FAvaRundownManagedInstanceHandles(const FAvaRundownManagedInstanceHandles& InOther) = delete;
+	FAvaRundownManagedInstanceHandles& operator=(const FAvaRundownManagedInstanceHandles& InOther) = delete;
+
+	FAvaRundownManagedInstanceHandles(FAvaRundownManagedInstanceHandles&& InOther)
+	{
+		*this = MoveTemp(InOther);
+	}
+
+	/** Copy Move operator handles unregistering instances if needed. */
+	UE_API FAvaRundownManagedInstanceHandles& operator=(FAvaRundownManagedInstanceHandles&& InOther);
+
+	~FAvaRundownManagedInstanceHandles()
+	{
+		Reset();
+	}
+
+	bool IsEmpty() const
+	{
+		return Instances.IsEmpty();
+	}
+
+	int32 Num() const
+	{
+		return Instances.Num();
+	}
+
+	/**
+	 * Unregisters the rundown from the instances (only if it hasn't been overriden already) and resets the array of instances. 
+	 */
+	UE_API void Reset();
+
+	/** Keep track of the registered Rundown to be able to unregister if needed. */
+	TWeakObjectPtr<UAvaRundown> RegisteredRundown;
+
+	/** Array of managed asset instances. */
+	TArray<TSharedPtr<FAvaRundownManagedInstance>> Instances;
+};
+
+#undef UE_API
\ No newline at end of file
diff --git a/Engine/Plugins/VirtualProduction/Avalanche/Source/AvalancheMedia/Internal/Rundown/AvaRundownMessages.h b/Engine/Plugins/VirtualProduction/Avalanche/Source/AvalancheMedia/Internal/Rundown/AvaRundownMessages.h
--- a/Engine/Plugins/VirtualProduction/Avalanche/Source/AvalancheMedia/Internal/Rundown/AvaRundownMessages.h
+++ b/Engine/Plugins/VirtualProduction/Avalanche/Source/AvalancheMedia/Internal/Rundown/AvaRundownMessages.h
@@ -772,13 +772,33 @@ public:
 	UPROPERTY()
 	FAvaPlayableRemoteControlValues RemoteControlValues;
 
-	/** Name of the remote control preset to resolve through WebRC API. */
+	/**
+	 * Name of the remote control preset to resolve through WebRC API.
+	 * In case of combo page, this is the first sub-page only. See RemoteControlPresetNames for all the names in that case.
+	 */
 	UPROPERTY()
 	FString RemoteControlPresetName;
 
-	/** Uuid of the remote control preset to resolve through WebRC API. */
+	/**
+	 * Uuid of the remote control preset to resolve through WebRC API.
+	 * In case of combo page, this is the first sub-page only. See RemoteControlPresetIds for all the uuids in that case.
+	 */
 	UPROPERTY()
 	FString RemoteControlPresetId;
+
+	/** 
+	 * This list will be populated with the names of the remote control preset of each sub-page in case of a combo page.
+	 * The provided Uuids can be used to with the WebRC API. 
+	 */
+	UPROPERTY()
+	TArray<FString> RemoteControlPresetNames;
+	
+	/** 
+	 * This list will be populated with the uuid of the remote control preset of each sub-page in case of a combo page.
+	 * The provided uuids can be used to with the WebRC API. 
+	 */
+	UPROPERTY()
+	TArray<FString> RemoteControlPresetIds;
 };
 
 /**
diff --git a/Engine/Plugins/VirtualProduction/Avalanche/Source/AvalancheMedia/Private/Playable/AvaPlayable.cpp b/Engine/Plugins/VirtualProduction/Avalanche/Source/AvalancheMedia/Private/Playable/AvaPlayable.cpp
--- a/Engine/Plugins/VirtualProduction/Avalanche/Source/AvalancheMedia/Private/Playable/AvaPlayable.cpp
+++ b/Engine/Plugins/VirtualProduction/Avalanche/Source/AvalancheMedia/Private/Playable/AvaPlayable.cpp
@@ -318,6 +318,7 @@ EAvaPlayableCommandResult UAvaPlayable::UpdateRemoteControlCommand(const TShared
 		const bool bExecuteSpecialControllerBehavior = CVarExecuteSpecialRCControllerBehavior.GetValueOnGameThread();
 		const bool bExecuteAllControllerBehaviors = EnumHasAnyFlags(InFlags, EAvaPlayableRCUpdateFlags::ExecuteControllerBehaviors);
 		const bool bExecuteNonOverlappingControllers = CVarExecuteNonOverlappingRCControllerBehavior.GetValueOnGameThread();
+		const bool bExecuteEventControllers = EnumHasAnyFlags(InFlags, EAvaPlayableRCUpdateFlags::ControllerValuesAsEvents);
 
 		TSet<FGuid> ModifiedControllers;
 
@@ -365,16 +366,33 @@ EAvaPlayableCommandResult UAvaPlayable::UpdateRemoteControlCommand(const TShared
 			
 			if (const FAvaPlayableRemoteControlValue* ControllerValue = InRemoteControlValues->GetControllerValue(Controller->Id))
 			{
+				const bool bIsEventController = FAvaPlayableRemoteControlValues::IsRuntimeEventController(Controller);
+
+				// Skip event controllers unless this is an event RC update.
+				if (bIsEventController && !bExecuteEventControllers)
+				{
+					continue;
+				}
+				
 				const bool bBehaviorsEnabled =
 					// Texture Controller Require the Bind Behavior to be executed to setup the texture in external mode.
 					(bExecuteSpecialControllerBehavior && GetCustomControllerTypeName(Controller) == CustomTextureControllerName)
 					// Modified controller's behavior are executed when updating RC values
 					|| (bExecuteAllControllerBehaviors && ModifiedControllers.Contains(Controller->Id))
 					// Non-overlapping controller's behavior can be enabled (experimental)
-					|| (bExecuteNonOverlappingControllers && RCPInfo.IsValid() && !RCPInfo->IsControllerOverlapping(Controller->Id)); 
+					|| (bExecuteNonOverlappingControllers && RCPInfo.IsValid() && !RCPInfo->IsControllerOverlapping(Controller->Id))
+					// Execute event controllers if this is an event update regardless of the controller value (i.e. modified or not)
+					|| (bIsEventController && bExecuteEventControllers);
 
 				SetValueOfController(Controller, ControllerValue->Value, bBehaviorsEnabled);
 
+				// (Tentative) Update the payload of event controllers.
+				// It is supposed to reflect the latest values as they are applied.
+				if (bIsEventController && bExecuteEventControllers && LatestRemoteControlValues)
+				{
+					LatestRemoteControlValues->SetControllerValue(Controller->Id, *ControllerValue);
+				}
+
 				if (bBehaviorsEnabled)
 				{
 					GetEntitiesControlledByController(RemoteControlPreset, Controller, AppliedOrIgnoredEntities);
@@ -389,10 +407,13 @@ EAvaPlayableCommandResult UAvaPlayable::UpdateRemoteControlCommand(const TShared
 	// We assume the controller actions are already executed in the rundown's managed RCP
 	// during page edition and the resulting entity values are already captured.
 
-	InRemoteControlValues->ApplyEntityValuesToRemoteControlPreset(RemoteControlPreset, AppliedOrIgnoredEntities);
-
-	LatestRemoteControlValues = InRemoteControlValues;
-	OnRemoteControlValuesApplied();
+	// Only apply entity values for non-event RC updates.
+	if (!EnumHasAnyFlags(InFlags, EAvaPlayableRCUpdateFlags::ControllerValuesAsEvents))
+	{
+		InRemoteControlValues->ApplyEntityValuesToRemoteControlPreset(RemoteControlPreset, AppliedOrIgnoredEntities);
+		LatestRemoteControlValues = InRemoteControlValues;
+		OnRemoteControlValuesApplied();
+	}
 
 	return EAvaPlayableCommandResult::Executed;
 }
diff --git a/Engine/Plugins/VirtualProduction/Avalanche/Source/AvalancheMedia/Private/Playable/AvaPlayableRemoteControlValues.cpp b/Engine/Plugins/VirtualProduction/Avalanche/Source/AvalancheMedia/Private/Playable/AvaPlayableRemoteControlValues.cpp
--- a/Engine/Plugins/VirtualProduction/Avalanche/Source/AvalancheMedia/Private/Playable/AvaPlayableRemoteControlValues.cpp
+++ b/Engine/Plugins/VirtualProduction/Avalanche/Source/AvalancheMedia/Private/Playable/AvaPlayableRemoteControlValues.cpp
@@ -3,13 +3,14 @@
 #include "Playable/AvaPlayableRemoteControlValues.h"
 
 #include "AvaMediaSerializationUtils.h"
-#include "IAvaMediaModule.h"
+#include "Controller/RCStatelessEventController.h"
 #include "Dom/JsonObject.h"
+#include "IAvaMediaModule.h"
 #include "Misc/PackageName.h"
 #include "Playable/AvaPlayableRemoteControl.h"
 #include "Playable/AvaPlayableRemoteControlValuesPrivate.h"
-#include "RCVirtualProperty.h"
 #include "Playable/AvaPlayableSettings.h"
+#include "RCVirtualProperty.h"
 #include "Serialization/CustomVersion.h"
 #include "Serialization/JsonReader.h"
 #include "Serialization/JsonSerializer.h"
@@ -240,8 +241,14 @@ void FAvaPlayableRemoteControlValues::CopyFrom(const URemoteControlPreset* InRem
 		TArray<URCVirtualPropertyBase*> Controllers = InRemoteControlPreset->GetControllers();
 		for (URCVirtualPropertyBase* Controller : Controllers)
 		{
+			// Skip ignored controllers
+			if (ShouldIgnoreController(Controller))
+			{
+				continue;
+			}
+
 			const EAvaPlayableRemoteControlResult Result = GetValueOfController(Controller, ValueAsString);
-	
+
 			if (Failed(Result))
 			{
 				UE_LOG(LogAvaPlayableRemoteControl, Error,
@@ -249,7 +256,7 @@ void FAvaPlayableRemoteControlValues::CopyFrom(const URemoteControlPreset* InRem
 					*Controller->DisplayName.ToString(), *Controller->Id.ToString(), *InRemoteControlPreset->GetName(), *EnumToString(Result));
 				continue;
 			}
-			
+
 			ControllerValues.Add(Controller->Id, FAvaPlayableRemoteControlValue(ValueAsString, bInIsDefault));
 		}
 	}
@@ -513,4 +520,11 @@ bool FAvaPlayableRemoteControlValues::ShouldIgnoreController(const URCVirtualPro
 		}
 	}
 	return false;
+}
+
+bool FAvaPlayableRemoteControlValues::IsRuntimeEventController(const URCVirtualPropertyBase* InController)
+{
+	// For now, we only consider the "stateless" event controllers.
+	// todo: event controllers with a payload.
+	return FRCStatelessEventController::IsStatelessEventController(InController);
 }
\ No newline at end of file
diff --git a/Engine/Plugins/VirtualProduction/Avalanche/Source/AvalancheMedia/Private/Rundown/AvaRundown.cpp b/Engine/Plugins/VirtualProduction/Avalanche/Source/AvalancheMedia/Private/Rundown/AvaRundown.cpp
--- a/Engine/Plugins/VirtualProduction/Avalanche/Source/AvalancheMedia/Private/Rundown/AvaRundown.cpp
+++ b/Engine/Plugins/VirtualProduction/Avalanche/Source/AvalancheMedia/Private/Rundown/AvaRundown.cpp
@@ -1950,6 +1950,94 @@ void UAvaRundown::OnParentWordBeginTearDown()
 	PagePlayers.Reset();
 }
 
+namespace UE::AvaMedia::Rundown::Private
+{
+	/**
+	 * Build an RC Values object with the given controller event ids.
+	 * This will be used to push an "controller events" RC update with only those controllers.
+	 * 
+	 * @param InRcItemIds List of controllers to be part of this update.
+	 * @param InPageValues Page Values, in case the controllers have values "payload" for this event.
+	 * @return Shared object with the controller values that can be pushed in the RC update command.
+	 */
+	TSharedRef<FAvaPlayableRemoteControlValues> MakeControllerEvents(TConstArrayView<FGuid> InRcItemIds, const FAvaPlayableRemoteControlValues& InPageValues)
+	{
+		// Build the RC values with the events
+		TSharedRef<FAvaPlayableRemoteControlValues> ControllerEvents = MakeShared<FAvaPlayableRemoteControlValues>();
+
+		for (const FGuid& ItemId : InRcItemIds)
+		{
+			// If the controller has payload in the page, we can push it.
+			if (const FAvaPlayableRemoteControlValue* ControllerValue = InPageValues.GetControllerValue(ItemId))
+			{
+				ControllerEvents->SetControllerValue(ItemId, *ControllerValue);
+			}
+			else
+			{
+				// Push empty value, in case of stateless controller.
+				ControllerEvents->SetControllerValue(ItemId, FAvaPlayableRemoteControlValue()); 
+			}
+		}
+
+		return ControllerEvents;
+	}
+}
+
+bool UAvaRundown::PushRuntimeRemoteControlEvents(TConstArrayView<FGuid> InRcItemIds, int32 InPageId, bool bInPushToProgram, bool bInPushToPreview, FName InPreviewChannelName)
+{
+	const FAvaRundownPage& Page = GetPage(InPageId);
+	if (!Page.IsValidPage())
+	{
+		return false;
+	}
+
+	bool bValuesPushed = false;
+	TSharedPtr<FAvaPlayableRemoteControlValues> ControllerEvents;
+
+	for (const UAvaRundownPagePlayer* PagePlayer : PagePlayers)
+	{
+		if (PagePlayer->PageId != InPageId)
+		{
+			continue;
+		}
+
+		if (PagePlayer->bIsPreview)
+		{
+			// Filter on preview channel if provided.
+			if (!bInPushToPreview || (!InPreviewChannelName.IsNone() && PagePlayer->ChannelFName != InPreviewChannelName))
+			{
+				continue;
+			}
+		}
+		else
+		{
+			if (!bInPushToProgram)
+			{
+				continue;
+			}
+		}
+
+		// Build the controller events update for the page (only once for all players).
+		if (!ControllerEvents.IsValid())
+		{
+			using namespace UE::AvaMedia::Rundown::Private;
+			ControllerEvents = MakeControllerEvents(InRcItemIds, Page.GetRemoteControlValues());
+		}
+
+		for (int32 InstanceIndex = 0; InstanceIndex < PagePlayer->GetNumInstancePlayers(); ++InstanceIndex)
+		{
+			if (UAvaPlaybackGraph* Playback = PagePlayer->GetPlayback(InstanceIndex))
+			{
+				constexpr EAvaPlayableRCUpdateFlags UpdateFlags = EAvaPlayableRCUpdateFlags::ExecuteControllerBehaviors | EAvaPlayableRCUpdateFlags::ControllerValuesAsEvents;
+				Playback->PushRemoteControlValues(PagePlayer->GetSourceAssetPath(InstanceIndex), PagePlayer->ChannelName, ControllerEvents.ToSharedRef(), UpdateFlags);
+			}
+		}
+		bValuesPushed = true;
+	}
+	return bValuesPushed;
+}
+
+
 bool UAvaRundown::PushRuntimeRemoteControlValues(int32 InPageId, bool bInIsPreview, const FName& InPreviewChannelName) const
 {
 	const FAvaRundownPage& Page = GetPage(InPageId);
@@ -1999,7 +2087,7 @@ void UAvaRundown::NotifyPageRemoteControlValueChanged(int32 InPageId, EAvaPlayab
 		// For now, potentially pushing all values multiple time (per frame) is mitigated by the
 		// optimization in FAvaRemoteControlUtils::SetValueOfEntity that
 		// will only set the value of the entity if it changed.
-		PushRuntimeRemoteControlValues(InPageId, true);
+		PushRuntimeRemoteControlValues(InPageId, /*bInIsPreview*/ true);
 	}
 	OnPagesChanged.Broadcast(this, GetPage(InPageId), EAvaRundownPageChanges::RemoteControlValues);
 }
diff --git a/Engine/Plugins/VirtualProduction/Avalanche/Source/AvalancheMedia/Private/Rundown/AvaRundownManagedInstance.cpp b/Engine/Plugins/VirtualProduction/Avalanche/Source/AvalancheMedia/Private/Rundown/AvaRundownManagedInstance.cpp
--- a/Engine/Plugins/VirtualProduction/Avalanche/Source/AvalancheMedia/Private/Rundown/AvaRundownManagedInstance.cpp
+++ b/Engine/Plugins/VirtualProduction/Avalanche/Source/AvalancheMedia/Private/Rundown/AvaRundownManagedInstance.cpp
@@ -2,18 +2,39 @@
 
 #include "Rundown/AvaRundownManagedInstance.h"
 
+#include "AvaMediaSettings.h"
 #include "Playable/AvaPlayableRemoteControlPresetInfoCache.h"
 #include "RemoteControlPreset.h"
+#include "Rundown/AvaRundown.h"
 #include "Rundown/AvaRundownManagedInstanceCache.h"
 
+namespace UE::AvaMedia::RundownManagedInstance::Private
+{
+	bool IsEditorTransacting()
+	{
+#if WITH_EDITOR
+		return GIsTransacting;
+#else
+		return false;
+#endif
+	}
+}
+
 FAvaRundownManagedInstance::FAvaRundownManagedInstance(FAvaRundownManagedInstanceCache* InParentCache, const FSoftObjectPath& InAssetPath)
 	: ParentCache(InParentCache)
 	, SourceAssetPath(InAssetPath)
+	, ActivePageId(FAvaRundownPage::InvalidPageId)
 {}
 
 FAvaRundownManagedInstance::~FAvaRundownManagedInstance()
 {}
 
+void FAvaRundownManagedInstance::SetCurrentRundownPageForEvents(UAvaRundown* InRundown, int32 InPageId)
+{
+	RundownForEvents = InRundown;
+	ActivePageId = InPageId;
+}
+
 void FAvaRundownManagedInstance::InvalidateFromCache()
 {
 	if (ParentCache)
@@ -79,6 +100,23 @@ void FAvaRundownManagedInstance::UnregisterSourceRemoteControlPresetDelegates(UR
 	InSourceRemoteControlPreset->OnControllerModified().RemoveAll(this);
 }
 
+void FAvaRundownManagedInstance::RegisterManagedRemoteControlPresetDelegates(URemoteControlPreset* InManagedRemoteControlPreset)
+{
+	if (::IsValid(InManagedRemoteControlPreset))
+	{
+		InManagedRemoteControlPreset->OnControllerModified().AddRaw(this, &FAvaRundownManagedInstance::OnManagedRemoteControlControllerModified);
+	}
+	
+}
+
+void FAvaRundownManagedInstance::UnregisterManagedRemoteControlPresetDelegates(URemoteControlPreset* InManagedRemoteControlPreset) const
+{
+	if (::IsValid(InManagedRemoteControlPreset))
+	{
+		InManagedRemoteControlPreset->OnControllerModified().RemoveAll(this);
+	}
+}
+
 void FAvaRundownManagedInstance::OnSourceRemoteControlEntitiesExposed(URemoteControlPreset* InSourcePreset, const FGuid& InEntityId)
 {
 	InvalidateFromCache();
@@ -155,4 +193,41 @@ void FAvaRundownManagedInstance::OnSourceRemoteControlControllerModified(URemote
 			DefaultRemoteControlValues.SetControllerValue(ControllerId, InSourcePreset, true);
 		}
 	}
-}
\ No newline at end of file
+}
+
+void FAvaRundownManagedInstance::OnManagedRemoteControlControllerModified(URemoteControlPreset* InPreset, const TSet<FGuid>& InModifiedControllerIds)
+{
+	using namespace UE::AvaMedia::RundownManagedInstance::Private;
+	// Note: Ignore changes from the RCP Transaction listener.
+	if (!::IsValid(InPreset) || GetRemoteControlPreset() != InPreset || IsEditorTransacting())
+	{
+		return;
+	}
+
+	UAvaRundown* Rundown = RundownForEvents.Get();
+
+	if (!Rundown || !Rundown->GetPage(ActivePageId).IsValidPage())
+	{
+		return;
+	}
+	
+	TArray<FGuid, TInlineAllocator<1>> EventControllerIds;
+
+	for (const FGuid& Id : InModifiedControllerIds)
+	{
+		if (URCVirtualPropertyBase* Controller = InPreset->GetController(Id))
+		{
+			if (FAvaPlayableRemoteControlValues::IsRuntimeEventController(Controller))
+			{
+				EventControllerIds.Add(Id);
+			}
+		}
+	}
+
+	if (!EventControllerIds.IsEmpty())
+	{
+		const bool bPushToProgram = UAvaMediaSettings::Get().bRundownPushControllerEventsToProgram;
+		const bool bPushToPreview = UAvaMediaSettings::Get().bRundownPushControllerEventsToPreview;
+		Rundown->PushRuntimeRemoteControlEvents(EventControllerIds, ActivePageId, bPushToProgram, bPushToPreview);
+	}
+}
diff --git a/Engine/Plugins/VirtualProduction/Avalanche/Source/AvalancheMedia/Private/Rundown/AvaRundownManagedInstanceCache.cpp b/Engine/Plugins/VirtualProduction/Avalanche/Source/AvalancheMedia/Private/Rundown/AvaRundownManagedInstanceCache.cpp
--- a/Engine/Plugins/VirtualProduction/Avalanche/Source/AvalancheMedia/Private/Rundown/AvaRundownManagedInstanceCache.cpp
+++ b/Engine/Plugins/VirtualProduction/Avalanche/Source/AvalancheMedia/Private/Rundown/AvaRundownManagedInstanceCache.cpp
@@ -5,10 +5,9 @@
 #include "AssetRegistry/AssetData.h"
 #include "AssetRegistry/IAssetRegistry.h"
 #include "AvaMediaSettings.h"
-#include "Engine/Engine.h"
-#include "Engine/World.h"
 #include "IAvaMediaModule.h"
 #include "Playback/AvaPlaybackUtils.h"
+#include "Rundown/AvaRundown.h"
 #include "Rundown/AvaRundownManagedInstanceLevel.h"
 #include "UObject/ObjectSaveContext.h"
 #include "UObject/Package.h"
@@ -122,6 +121,24 @@ void FAvaRundownManagedInstanceCache::Flush()
 	{
 		return (InManagedInstance.GetSharedReferenceCount() <= 1) ? true : false;
 	}, false);
+	
+	// Dump post-flush remaining assets to help diagnose which editor keeps the asset alive.
+	// Since the cache is shared with all editors (currently), assets can still remain if used by other playlist editors/servers.
+	if (UE_LOG_ACTIVE(LogAvaRundown, Verbose))
+	{
+		if (Instances.Num() > 0)
+		{
+			UE_LOG(LogAvaRundown, Verbose, TEXT("Managed Instance Cache: Assets remaining in cache after flush:"));
+			for (TPair<FSoftObjectPath, TSharedPtr<FAvaRundownManagedInstance>>& Entry : Instances)
+			{
+				UE_LOG(LogAvaRundown, Verbose, TEXT("- Asset: \"%s\" RefCount: %d"), *Entry.Key.ToString(), Entry.Value.GetSharedReferenceCount());
+			}
+		}
+		else
+		{
+			UE_LOG(LogAvaRundown, Verbose, TEXT("Managed Instance Cache: No Assets remaining in cache after flush."));
+		}
+	}
 }
 
 void FAvaRundownManagedInstanceCache::TrimCache()
@@ -143,6 +160,29 @@ void FAvaRundownManagedInstanceCache::FinishPendingActions()
 	RemovePendingInvalidatedPaths();
 }
 
+TArray<TSharedPtr<FAvaRundownManagedInstance>> FAvaRundownManagedInstanceCache::GetManagedInstancesForPage(const UAvaRundown* InRundown, const FAvaRundownPage& InPage)
+{
+	const TArray<FSoftObjectPath> AssetPaths = InPage.GetAssetPaths(InRundown);
+	
+	TArray<TSharedPtr<FAvaRundownManagedInstance>> OutInstances;
+	OutInstances.Reserve(AssetPaths.Num());
+	
+	for (const FSoftObjectPath& AssetPath : AssetPaths)
+	{		
+		if (TSharedPtr<FAvaRundownManagedInstance> Instance = GetOrLoadInstance(AssetPath))
+		{
+			OutInstances.Add(Instance);
+		}
+	}
+
+	return OutInstances;
+}
+
+FAvaRundownManagedInstanceHandles FAvaRundownManagedInstanceCache::GetManagedHandlesForPage(UAvaRundown* InRundown, const FAvaRundownPage& InPage)
+{
+	return FAvaRundownManagedInstanceHandles(GetManagedInstancesForPage(InRundown, InPage), InRundown, InPage.GetPageId());
+}
+
 void FAvaRundownManagedInstanceCache::RemovePendingInvalidatedPaths()
 {
 	for (const FSoftObjectPath& InvalidatedPath : PendingInvalidatedPaths)
diff --git a/Engine/Plugins/VirtualProduction/Avalanche/Source/AvalancheMedia/Private/Rundown/AvaRundownManagedInstanceHandle.cpp b/Engine/Plugins/VirtualProduction/Avalanche/Source/AvalancheMedia/Private/Rundown/AvaRundownManagedInstanceHandle.cpp
new file mode 100644
--- /dev/null
+++ b/Engine/Plugins/VirtualProduction/Avalanche/Source/AvalancheMedia/Private/Rundown/AvaRundownManagedInstanceHandle.cpp
@@ -0,0 +1,65 @@
+// Copyright Epic Games, Inc. All Rights Reserved.
+
+#include "Rundown/AvaRundownManagedInstanceHandle.h"
+
+#include "Rundown/AvaRundown.h"
+#include "Rundown/AvaRundownManagedInstance.h"
+#include "Rundown/AvaRundownPage.h"
+
+namespace UE::AvaMedia::Rundown::Private
+{
+	/** Utility function to unregister the current rundown from controller events routing. */
+	void UnregisterRundownForEvents(const TSharedPtr<FAvaRundownManagedInstance>& InInstance, const UAvaRundown* InRundown)
+	{
+		// Don't reset the rundown in case this instance is already used by another rundown.
+		if (InInstance && InInstance->GetRundownForEvents() == InRundown)
+		{
+			InInstance->SetCurrentRundownPageForEvents(nullptr, FAvaRundownPage::InvalidPageId);
+		}
+	}
+	
+	/** Utility function to unregister the current rundown from controller events routing. */
+	void UnregisterRundownForEvents(TConstArrayView<TSharedPtr<FAvaRundownManagedInstance>> InInstances, const UAvaRundown* InRundown)
+	{
+		for (const TSharedPtr<FAvaRundownManagedInstance>& Instance : InInstances)
+		{
+			UnregisterRundownForEvents(Instance, InRundown);
+		}
+	}
+}
+
+FAvaRundownManagedInstanceHandles::FAvaRundownManagedInstanceHandles(const TArray<TSharedPtr<FAvaRundownManagedInstance>>& InInstances, UAvaRundown* InRundown, int32 InPageId)
+	: RegisteredRundown(InRundown)
+	, Instances(InInstances)
+{
+	for (const TSharedPtr<FAvaRundownManagedInstance>& Instance : Instances)
+	{
+		// RC Controller events will be routed to the given rundown.
+		Instance->SetCurrentRundownPageForEvents(InRundown, InPageId);
+	}
+}
+
+void FAvaRundownManagedInstanceHandles::Reset()
+{
+	UE::AvaMedia::Rundown::Private::UnregisterRundownForEvents(Instances, RegisteredRundown.Get());
+	Instances.Reset();
+}
+
+FAvaRundownManagedInstanceHandles& FAvaRundownManagedInstanceHandles::operator=(FAvaRundownManagedInstanceHandles&& InOther)
+{
+	if (&InOther == this)
+	{
+		return *this;
+	}
+	
+	// Need to unregister previous instances from another rundown before copying.
+	if (!Instances.IsEmpty() && RegisteredRundown != InOther.RegisteredRundown)
+	{
+		UE::AvaMedia::Rundown::Private::UnregisterRundownForEvents(Instances, RegisteredRundown.Get());
+		RegisteredRundown = nullptr;
+	}
+
+	RegisteredRundown = MoveTemp(InOther.RegisteredRundown);
+	Instances = MoveTemp(InOther.Instances);
+	return *this;
+}
diff --git a/Engine/Plugins/VirtualProduction/Avalanche/Source/AvalancheMedia/Private/Rundown/AvaRundownManagedInstanceLevel.cpp b/Engine/Plugins/VirtualProduction/Avalanche/Source/AvalancheMedia/Private/Rundown/AvaRundownManagedInstanceLevel.cpp
--- a/Engine/Plugins/VirtualProduction/Avalanche/Source/AvalancheMedia/Private/Rundown/AvaRundownManagedInstanceLevel.cpp
+++ b/Engine/Plugins/VirtualProduction/Avalanche/Source/AvalancheMedia/Private/Rundown/AvaRundownManagedInstanceLevel.cpp
@@ -6,7 +6,6 @@
 #include "AvaRemoteControlRebind.h"
 #include "AvaRemoteControlUtils.h"
 #include "AvaScene.h"
-#include "Engine/Engine.h"
 #include "Engine/Level.h"
 #include "Engine/World.h"
 #include "IAvaMediaModule.h"
@@ -110,6 +109,7 @@ FAvaRundownManagedInstanceLevel::FAvaRundownManagedInstanceLevel(FAvaRundownMana
 	
 	ManagedRemoteControlPreset = FindRemoteControlPreset(ManagedLevel->PersistentLevel);
 	FAvaRemoteControlRebind::RebindUnboundEntities(ManagedRemoteControlPreset, ManagedLevel->PersistentLevel);
+	RegisterManagedRemoteControlPresetDelegates(ManagedRemoteControlPreset);
 	
 	// Backup the remote control values from the source asset.
 	constexpr bool bIsDefault = true;	// Flag the values as "default".
@@ -127,6 +127,7 @@ FAvaRundownManagedInstanceLevel::~FAvaRundownManagedInstanceLevel()
 	if (ManagedRemoteControlPreset)
 	{
 		FAvaRemoteControlUtils::UnregisterRemoteControlPreset(ManagedRemoteControlPreset);
+		UnregisterManagedRemoteControlPresetDelegates(ManagedRemoteControlPreset);
 	}
 
 	if (ManagedLevelPackage)
diff --git a/Engine/Plugins/VirtualProduction/Avalanche/Source/AvalancheMedia/Private/Rundown/AvaRundownServer.cpp b/Engine/Plugins/VirtualProduction/Avalanche/Source/AvalancheMedia/Private/Rundown/AvaRundownServer.cpp
--- a/Engine/Plugins/VirtualProduction/Avalanche/Source/AvalancheMedia/Private/Rundown/AvaRundownServer.cpp
+++ b/Engine/Plugins/VirtualProduction/Avalanche/Source/AvalancheMedia/Private/Rundown/AvaRundownServer.cpp
@@ -1373,29 +1373,34 @@ void FAvaRundownServer::HandleGetPageDetails(const FAvaRundownGetPageDetails& In
 	
 	if (InMessage.bLoadRemoteControlPreset)
 	{
-		FAvaRundownManagedInstanceCache& ManagedInstanceCache = IAvaMediaModule::Get().GetManagedInstanceCache();
-		const TSharedPtr<FAvaRundownManagedInstance> ManagedInstance
-			= ManagedInstanceCache.GetOrLoadInstance(Page.GetAssetPath(Rundown));
+		const TArray<FSoftObjectPath> AssetPaths = Page.GetAssetPaths(Rundown);
 		
-		if (ManagedInstance.IsValid())
+		FAvaRundownManagedInstanceCache& ManagedInstanceCache = IAvaMediaModule::Get().GetManagedInstanceCache();
+
+		FAvaRundownManagedInstanceHandles ManagedHandles = ManagedInstanceCache.GetManagedHandlesForPage(Rundown, Page);
+
+		if (!ManagedHandles.IsEmpty())
 		{
 			EditCommandContext.SaveCurrentRemoteControlPresetToPage(true);
 
-			// Applying the controller values can break the WYSIWYG of the editor,
-			// in case multiple controllers set the same exposed entity with different values.
-			// There is no guaranty that the controller actions are self consistent.
-			// To avoid this issue, we apply the controllers first, and then
-			// restore the entity values in a second pass.
+			for (const TSharedPtr<FAvaRundownManagedInstance>& ManagedInstance : ManagedHandles.Instances)
+			{
+				// Applying the controller values can break the WYSIWYG of the editor,
+				// in case multiple controllers set the same exposed entity with different values.
+				// There is no guaranty that the controller actions are self consistent.
+				// To avoid this issue, we apply the controllers first, and then
+				// restore the entity values in a second pass.
 			
-			Page.GetRemoteControlValues().ApplyControllerValuesToRemoteControlPreset(ManagedInstance->GetRemoteControlPreset(), true);
-			Page.GetRemoteControlValues().ApplyEntityValuesToRemoteControlPreset(ManagedInstance->GetRemoteControlPreset());
+				Page.GetRemoteControlValues().ApplyControllerValuesToRemoteControlPreset(ManagedInstance->GetRemoteControlPreset(), true);
+				Page.GetRemoteControlValues().ApplyEntityValuesToRemoteControlPreset(ManagedInstance->GetRemoteControlPreset());
 
-			// Register the RC Preset to Remote Control module to make it available through WebRC.
-			FAvaRemoteControlUtils::RegisterRemoteControlPreset(ManagedInstance->GetRemoteControlPreset(), /*bInEnsureUniqueId*/ true);
+				// Register the RC Preset to Remote Control module to make it available through WebRC.
+				FAvaRemoteControlUtils::RegisterRemoteControlPreset(ManagedInstance->GetRemoteControlPreset(), /*bInEnsureUniqueId*/ true);
+			}
 
 			// Keep track of what is currently registered.
 			EditCommandContext.ManagedPageId = InMessage.PageId;
-			EditCommandContext.ManagedInstance = ManagedInstance;
+			EditCommandContext.ManagedHandles = MoveTemp(ManagedHandles);
 		}
 	}
 	
@@ -1404,10 +1409,23 @@ void FAvaRundownServer::HandleGetPageDetails(const FAvaRundownGetPageDetails& In
 	ReplyMessage->Rundown = InMessage.Rundown;
 	ReplyMessage->PageInfo = UE::AvaRundownServer::Private::GetPageInfo(Rundown, Page);
 	ReplyMessage->RemoteControlValues = Page.GetRemoteControlValues();
-	if (InMessage.bLoadRemoteControlPreset && EditCommandContext.ManagedInstance.IsValid() && EditCommandContext.ManagedInstance->GetRemoteControlPreset())
+	
+	if (InMessage.bLoadRemoteControlPreset)
 	{
-		ReplyMessage->RemoteControlPresetName = EditCommandContext.ManagedInstance->GetRemoteControlPreset()->GetPresetName().ToString();
-		ReplyMessage->RemoteControlPresetId = EditCommandContext.ManagedInstance->GetRemoteControlPreset()->GetPresetId().ToString();
+		ReplyMessage->RemoteControlPresetNames.Reserve(EditCommandContext.ManagedHandles.Num());
+		ReplyMessage->RemoteControlPresetIds.Reserve(EditCommandContext.ManagedHandles.Num());
+
+		for (const TSharedPtr<FAvaRundownManagedInstance>& ManagedInstance : EditCommandContext.ManagedHandles.Instances)
+		{
+			if (ManagedInstance->GetRemoteControlPreset())
+			{
+				ReplyMessage->RemoteControlPresetNames.Add(ManagedInstance->GetRemoteControlPreset()->GetPresetName().ToString());
+				ReplyMessage->RemoteControlPresetIds.Add(ManagedInstance->GetRemoteControlPreset()->GetPresetId().ToString());
+			}
+		}
+
+		ReplyMessage->RemoteControlPresetName = !ReplyMessage->RemoteControlPresetNames.IsEmpty() ? ReplyMessage->RemoteControlPresetNames[0] : FString();
+		ReplyMessage->RemoteControlPresetId = !ReplyMessage->RemoteControlPresetIds.IsEmpty() ? ReplyMessage->RemoteControlPresetIds[0] : FString();
 	}
 	SendResponse(ReplyMessage, InContext->GetSender());
 }
@@ -1491,7 +1509,7 @@ void FAvaRundownServer::HandleUpdatePageFromRCP(const FAvaRundownUpdatePageFromR
 	if (InMessage.bUnregister)
 	{
 		EditCommandContext.ManagedPageId = FAvaRundownPage::InvalidPageId;
-		EditCommandContext.ManagedInstance.Reset();
+		EditCommandContext.ManagedHandles.Reset();
 	}
 	SendMessage(InContext->GetSender(), InMessage.RequestId, ELogVerbosity::Log, TEXT("\"UpdatePageFromRCP\" Ok."));
 }
@@ -3045,24 +3063,17 @@ void FAvaRundownServer::FEditCommandContext::SetCurrentRundown(const TSharedPtr<
 
 void FAvaRundownServer::FEditCommandContext::SaveCurrentRemoteControlPresetToPage(bool bInUnregister)
 {
-	if (!ManagedInstance.IsValid() || !ManagedInstance->GetRemoteControlPreset())
-	{
-		return;
-	}
-
-	// Check if the RCP was registered.
-	IRemoteControlModule& RemoteControlModule = IRemoteControlModule::Get();
-	const FName CurrentPresetName = ManagedInstance->GetRemoteControlPreset()->GetPresetName();
-	const URemoteControlPreset* ResolvedPreset = RemoteControlModule.ResolvePreset(CurrentPresetName);
-	if (ResolvedPreset != ManagedInstance->GetRemoteControlPreset())
+	if (ManagedHandles.IsEmpty())
 	{
 		return;
 	}
 
 	if (bInUnregister)
 	{
-		// Unregister from RC module.
-		RemoteControlModule.UnregisterEmbeddedPreset(CurrentPresetName);
+		for (const TSharedPtr<FAvaRundownManagedInstance>& ManagedInstance : ManagedHandles.Instances)
+		{
+			FAvaRemoteControlUtils::UnregisterRemoteControlPreset(ManagedInstance->GetRemoteControlPreset());
+		}
 	}
 
 	UAvaRundown* CurrentRundown = GetCurrentRundown();
@@ -3081,7 +3092,16 @@ void FAvaRundownServer::FEditCommandContext::SaveCurrentRemoteControlPresetToPag
 
 	constexpr bool bIsDefault = false;
 	FAvaPlayableRemoteControlValues NewValues;
-	NewValues.CopyFrom(ManagedInstance->GetRemoteControlPreset(), bIsDefault);
+	
+	for (const TSharedPtr<FAvaRundownManagedInstance>& ManagedInstance : ManagedHandles.Instances)
+	{
+		if (ManagedInstance->GetRemoteControlPreset())
+		{
+			FAvaPlayableRemoteControlValues Values;
+			Values.CopyFrom(ManagedInstance->GetRemoteControlPreset(), bIsDefault);
+			NewValues.Merge(Values);
+		}
+	}
 
 	// UpdateRemoteControlValues does half the job by ensuring that missing values are added and
 	// extra values are removed. But it doesn't change existing values.
diff --git a/Engine/Plugins/VirtualProduction/Avalanche/Source/AvalancheMedia/Private/Rundown/AvaRundownServer.h b/Engine/Plugins/VirtualProduction/Avalanche/Source/AvalancheMedia/Private/Rundown/AvaRundownServer.h
--- a/Engine/Plugins/VirtualProduction/Avalanche/Source/AvalancheMedia/Private/Rundown/AvaRundownServer.h
+++ b/Engine/Plugins/VirtualProduction/Avalanche/Source/AvalancheMedia/Private/Rundown/AvaRundownServer.h
@@ -6,6 +6,7 @@
 #include "MessageEndpoint.h"
 #include "Playback/AvaPlaybackManager.h"
 #include "Rundown/AvaRundown.h"
+#include "Rundown/AvaRundownManagedInstanceHandle.h"
 #include "Rundown/AvaRundownMessages.h"
 #include "Rundown/IAvaRundownServer.h"
 #include "UObject/SoftObjectPath.h"
@@ -327,7 +328,7 @@ private:
 	{
 		/** PageId of the current managed ava asset. */
 		int32 ManagedPageId = FAvaRundownPage::InvalidPageId;
-		TSharedPtr<FAvaRundownManagedInstance> ManagedInstance;
+		FAvaRundownManagedInstanceHandles ManagedHandles;
 
 		virtual ~FEditCommandContext() override;
 
diff --git a/Engine/Plugins/VirtualProduction/Avalanche/Source/AvalancheMedia/Public/AvaMediaDefines.h b/Engine/Plugins/VirtualProduction/Avalanche/Source/AvalancheMedia/Public/AvaMediaDefines.h
--- a/Engine/Plugins/VirtualProduction/Avalanche/Source/AvalancheMedia/Public/AvaMediaDefines.h
+++ b/Engine/Plugins/VirtualProduction/Avalanche/Source/AvalancheMedia/Public/AvaMediaDefines.h
@@ -178,7 +178,12 @@ enum class EAvaPlayableRCUpdateFlags : uint8
 {
 	None = 0,
 	/** Execute Controller Behaviors */
-	ExecuteControllerBehaviors = 1 << 0
+	ExecuteControllerBehaviors = 1 << 0,
+	/**
+	 *  The controller values are to be considered as events and will only apply to the event controllers.
+	 *  Non-event controllers and entity values are ignored during the execution of this update.
+	 */
+	ControllerValuesAsEvents = 1 << 1,
 };
 ENUM_CLASS_FLAGS(EAvaPlayableRCUpdateFlags);
 
diff --git a/Engine/Plugins/VirtualProduction/Avalanche/Source/AvalancheMedia/Public/AvaMediaSettings.h b/Engine/Plugins/VirtualProduction/Avalanche/Source/AvalancheMedia/Public/AvaMediaSettings.h
--- a/Engine/Plugins/VirtualProduction/Avalanche/Source/AvalancheMedia/Public/AvaMediaSettings.h
+++ b/Engine/Plugins/VirtualProduction/Avalanche/Source/AvalancheMedia/Public/AvaMediaSettings.h
@@ -146,6 +146,18 @@ public:
 	 */
 	UPROPERTY(config, EditAnywhere, Category = "Rundown")
 	bool bEnableSingleTemplateSpecialLogic = false;
+
+	/**
+	 * Enables the propagation of controller events to pages playing in program channels. 
+	 */
+	UPROPERTY(config, EditAnywhere, Category = "Rundown")
+	bool bRundownPushControllerEventsToProgram = true;
+
+	/**
+	 * Enables the propagation of controller events to pages playing in preview channels. 
+	 */
+	UPROPERTY(config, EditAnywhere, Category = "Rundown")
+	bool bRundownPushControllerEventsToPreview = true;
 	
 	/** Whether playback client is started automatically. */
 	UPROPERTY(config, EditAnywhere, Category = "Playback Client")
diff --git a/Engine/Plugins/VirtualProduction/Avalanche/Source/AvalancheMedia/Public/Playable/AvaPlayableRemoteControlValues.h b/Engine/Plugins/VirtualProduction/Avalanche/Source/AvalancheMedia/Public/Playable/AvaPlayableRemoteControlValues.h
--- a/Engine/Plugins/VirtualProduction/Avalanche/Source/AvalancheMedia/Public/Playable/AvaPlayableRemoteControlValues.h
+++ b/Engine/Plugins/VirtualProduction/Avalanche/Source/AvalancheMedia/Public/Playable/AvaPlayableRemoteControlValues.h
@@ -199,4 +199,10 @@ struct AVALANCHEMEDIA_API FAvaPlayableRemoteControlValues
 	 * @return true if the controller should be ignored.
 	 */
 	static bool ShouldIgnoreController(const URCVirtualPropertyBase* InController);
+
+	/**
+	 * Utility function to determine if a controller is considered an "event" controller, i.e.
+	 * it doesn't carry a value to be used to set or restore state, but rather is runtime event based.
+	 */
+	static bool IsRuntimeEventController(const URCVirtualPropertyBase* InController);
 };
\ No newline at end of file
diff --git a/Engine/Plugins/VirtualProduction/Avalanche/Source/AvalancheMedia/Public/Rundown/AvaRundown.h b/Engine/Plugins/VirtualProduction/Avalanche/Source/AvalancheMedia/Public/Rundown/AvaRundown.h
--- a/Engine/Plugins/VirtualProduction/Avalanche/Source/AvalancheMedia/Public/Rundown/AvaRundown.h
+++ b/Engine/Plugins/VirtualProduction/Avalanche/Source/AvalancheMedia/Public/Rundown/AvaRundown.h
@@ -633,6 +633,19 @@ public:
 	/** Clean up playing status on a system tear down. */
 	void OnParentWordBeginTearDown();
 
+	/**
+	 * Push an RC event on the given RC item(s) (controllers or entities) for the given page.
+	 * The propagation to channels in determined by rundown settings.
+	 * 
+	 * @param InRcItemIds List of rc item ids that should be triggered.
+	 * @param InPageId Page id these events apply to.
+	 * @param bInPushToProgram if true, the command will be pushed to page's program channel (if currently playing).
+	 * @param bInPushToPreview if true, the command will be pushed to preview channel(s).
+	 * @param InPreviewChannelName Allow to specify a specific preview channel to push to. If left none, will push to all preview channel this page is playing in.
+	 * @return true if the event command was pushed to the corresponding playback instance.
+	 */
+	bool PushRuntimeRemoteControlEvents(TConstArrayView<FGuid> InRcItemIds, int32 InPageId, bool bInPushToProgram, bool bInPushToPreview, FName InPreviewChannelName = NAME_None);
+	
 	/**
 	 * Push the page's RC values to the runtime playback instances.
 	 * @remark If updating preview, will update all preview channels by default. 
diff --git a/Engine/Plugins/VirtualProduction/Avalanche/Source/AvalancheMediaEditor/Private/Rundown/AvaRundownEditorUtils.cpp b/Engine/Plugins/VirtualProduction/Avalanche/Source/AvalancheMediaEditor/Private/Rundown/AvaRundownEditorUtils.cpp
--- a/Engine/Plugins/VirtualProduction/Avalanche/Source/AvalancheMediaEditor/Private/Rundown/AvaRundownEditorUtils.cpp
+++ b/Engine/Plugins/VirtualProduction/Avalanche/Source/AvalancheMediaEditor/Private/Rundown/AvaRundownEditorUtils.cpp
@@ -562,25 +562,6 @@ UAvaRundown* UE::AvaRundownEditor::Utils::SaveDuplicateRundown(UAvaRundown* InSo
 	return Cast<UAvaRundown>(IAssetTools::Get().DuplicateAsset(InAssetName, InPackagePath, InSourceRundown));
 }
 
-TArray<TSharedPtr<FAvaRundownManagedInstance>> UE::AvaRundownEditor::Utils::GetManagedInstancesForPage(const UAvaRundown* InRundown, const FAvaRundownPage& InPage)
-{
-	const TArray<FSoftObjectPath> AssetPaths = InPage.GetAssetPaths(InRundown);
-	
-	TArray<TSharedPtr<FAvaRundownManagedInstance>> ManagedInstances;
-	ManagedInstances.Reserve(AssetPaths.Num());
-	
-	FAvaRundownManagedInstanceCache& ManagedInstanceCache = IAvaMediaModule::Get().GetManagedInstanceCache();
-	
-	for (const FSoftObjectPath& AssetPath : AssetPaths)
-	{		
-		if (TSharedPtr<FAvaRundownManagedInstance> ManagedInstance = ManagedInstanceCache.GetOrLoadInstance(AssetPath))
-		{
-			ManagedInstances.Add(ManagedInstance);
-		}
-	}
-	return ManagedInstances;
-}
-
 bool UE::AvaRundownEditor::Utils::MergeDefaultRemoteControlValues(const TArray<TSharedPtr<FAvaRundownManagedInstance>>& InManagedInstances, FAvaPlayableRemoteControlValues& OutMergedValues)
 {
 	bool bAllUniqueIds = true;
@@ -604,14 +585,16 @@ EAvaPlayableRemoteControlChanges UE::AvaRundownEditor::Utils::UpdateDefaultRemot
 	{
 		return Changes;
 	}
-	
+
+	FAvaRundownManagedInstanceCache& ManagedInstanceCache = IAvaMediaModule::Get().GetManagedInstanceCache();
+
 	for (const int32 PageId : InSelectedPageIds)
 	{
 		FAvaRundownPage& Page = InRundown->GetPage(PageId);
 
 		if (Page.IsValidPage())
 		{
-			TArray<TSharedPtr<FAvaRundownManagedInstance>> ManagedInstances = GetManagedInstancesForPage(InRundown, Page);
+			TArray<TSharedPtr<FAvaRundownManagedInstance>> ManagedInstances = ManagedInstanceCache.GetManagedInstancesForPage(InRundown, Page);
 					
 			if (!ManagedInstances.IsEmpty())
 			{
@@ -646,8 +629,9 @@ bool UE::AvaRundownEditor::Utils::ForceUpdatePageSummary(int32 InPageId, UAvaRun
 
 	// Important: the managed instance cache is used by the page details panel, which assumes the instances it holds
 	// for the current page are not modified by something else.
-	// This is why we need to backup the controllers and restore them at the end of the function. 
-	const TArray<TSharedPtr<FAvaRundownManagedInstance>> ManagedInstances = GetManagedInstancesForPage(InRundown, Page);
+	// This is why we need to back up the controllers and restore them at the end of the function.
+	FAvaRundownManagedInstanceCache& ManagedInstanceCache = IAvaMediaModule::Get().GetManagedInstanceCache();
+	TArray<TSharedPtr<FAvaRundownManagedInstance>> ManagedInstances = ManagedInstanceCache.GetManagedInstancesForPage(InRundown, Page);
 	TArray<const URemoteControlPreset*> Presets;
 	Presets.Reserve(ManagedInstances.Num());
 
diff --git a/Engine/Plugins/VirtualProduction/Avalanche/Source/AvalancheMediaEditor/Private/Rundown/AvaRundownEditorUtils.h b/Engine/Plugins/VirtualProduction/Avalanche/Source/AvalancheMediaEditor/Private/Rundown/AvaRundownEditorUtils.h
--- a/Engine/Plugins/VirtualProduction/Avalanche/Source/AvalancheMediaEditor/Private/Rundown/AvaRundownEditorUtils.h
+++ b/Engine/Plugins/VirtualProduction/Avalanche/Source/AvalancheMediaEditor/Private/Rundown/AvaRundownEditorUtils.h
@@ -104,7 +104,6 @@ namespace UE::AvaRundownEditor::Utils
 	
 	UAvaRundown* SaveDuplicateRundown(UAvaRundown* InSourceRundown, const FString& InAssetName, const FString& InPackagePath);
 
-	TArray<TSharedPtr<FAvaRundownManagedInstance>> GetManagedInstancesForPage(const UAvaRundown* InRundown, const FAvaRundownPage& InPage);
 	bool MergeDefaultRemoteControlValues(const TArray<TSharedPtr<FAvaRundownManagedInstance>>& InManagedInstances, FAvaPlayableRemoteControlValues& OutMergedValues);
 	EAvaPlayableRemoteControlChanges UpdateDefaultRemoteControlValues(UAvaRundown* InRundown, const TArray<int32>& InSelectedPageIds);
 	
diff --git a/Engine/Plugins/VirtualProduction/Avalanche/Source/AvalancheMediaEditor/Private/Rundown/DetailsView/RemoteControl/Controllers/SAvaRundownRCControllerPanel.cpp b/Engine/Plugins/VirtualProduction/Avalanche/Source/AvalancheMediaEditor/Private/Rundown/DetailsView/RemoteControl/Controllers/SAvaRundownRCControllerPanel.cpp
--- a/Engine/Plugins/VirtualProduction/Avalanche/Source/AvalancheMediaEditor/Private/Rundown/DetailsView/RemoteControl/Controllers/SAvaRundownRCControllerPanel.cpp
+++ b/Engine/Plugins/VirtualProduction/Avalanche/Source/AvalancheMediaEditor/Private/Rundown/DetailsView/RemoteControl/Controllers/SAvaRundownRCControllerPanel.cpp
@@ -4,11 +4,11 @@
 #include "AvaRundownRCControllerItem.h"
 #include "Behaviour/Builtin/Path/RCSetAssetByPathBehaviour.h"
 #include "Controller/RCController.h"
+#include "IAvaMediaModule.h"
 #include "IDetailTreeNode.h"
 #include "IPropertyRowGenerator.h"
 #include "Modules/ModuleManager.h"
 #include "Playable/AvaPlayableRemoteControl.h"
-#include "Playable/AvaPlayableRemoteControlPresetInfo.h"
 #include "RCVirtualProperty.h"
 #include "Rundown/AvaRundownEditor.h"
 #include "Rundown/AvaRundownEditorUtils.h"
@@ -69,7 +69,7 @@ void SAvaRundownRCControllerPanel::Construct(const FArguments& InArgs, const TSh
 
 bool SAvaRundownRCControllerPanel::HasRemoteControlPreset(const URemoteControlPreset* InPreset) const
 {
-	for (const TSharedPtr<FAvaRundownManagedInstance>& ManagedInstance : ManagedInstances)
+	for (const TSharedPtr<FAvaRundownManagedInstance>& ManagedInstance : ManagedHandles.Instances)
 	{
 		if (ManagedInstance && ManagedInstance->GetRemoteControlPreset() == InPreset)
 		{
@@ -232,18 +232,18 @@ void SAvaRundownRCControllerPanel::Refresh(const TArray<int32>& InSelectedPageId
 
 	UAvaRundown* Rundown = GetRundown();
 	const FAvaRundownPage& Page = GetActivePage(Rundown);
-	
+
 	if (!Page.IsValidPage())
 	{
 		ControllerItems.Empty();
-       	ManagedInstances.Reset();
+		ManagedHandles.Reset();
 		return;
 	}
 
-	ManagedInstances = UE::AvaRundownEditor::Utils::GetManagedInstancesForPage(Rundown, Page);
+	ManagedHandles = IAvaMediaModule::Get().GetManagedInstanceCache().GetManagedHandlesForPage(Rundown, Page);
 
 	int32 NumItems = 0;
-	for (const TSharedPtr<FAvaRundownManagedInstance>& ManagedInstance : ManagedInstances)
+	for (const TSharedPtr<FAvaRundownManagedInstance>& ManagedInstance : ManagedHandles.Instances)
 	{
 		if (const URemoteControlPreset* Preset = ManagedInstance ? ManagedInstance->GetRemoteControlPreset() : nullptr)
 		{
@@ -260,12 +260,12 @@ void SAvaRundownRCControllerPanel::Refresh(const TArray<int32>& InSelectedPageId
 
 	ControllerItems.Empty(NumItems);
 	
-	UpdatePropertyRowGenerators(ManagedInstances.Num());
+	UpdatePropertyRowGenerators(ManagedHandles.Num());
 	UpdatePageSummary(/*bInForceUpdate*/ false); // Generate only if missing.
 
-	for (int32 Index = 0; Index < ManagedInstances.Num(); ++Index)
+	for (int32 Index = 0; Index < ManagedHandles.Num(); ++Index)
 	{
-		RefreshForManagedInstance(Index, *ManagedInstances[Index], Page);
+		RefreshForManagedInstance(Index, *ManagedHandles.Instances[Index], Page);
 	}
 
 	ControllerItems.Sort([](const FAvaRundownRCControllerItemPtr& A, const FAvaRundownRCControllerItemPtr& B)
@@ -411,14 +411,14 @@ void SAvaRundownRCControllerPanel::UpdatePageSummary(bool bInForceUpdate)
 	if (Page.IsValidPage())
 	{
 		TArray<const URemoteControlPreset*> Presets;
-		Presets.Reserve(ManagedInstances.Num());
-		for (const TSharedPtr<FAvaRundownManagedInstance>& ManagedInstance : ManagedInstances)
+		Presets.Reserve(ManagedHandles.Num());
+		for (const TSharedPtr<FAvaRundownManagedInstance>& ManagedInstance : ManagedHandles.Instances)
 		{
 			if (ManagedInstance && ManagedInstance->GetRemoteControlPreset())
 			{
 				Presets.Add(ManagedInstance->GetRemoteControlPreset());
 			}
-		}					
+		}
 		Page.UpdatePageSummary(GetRundown(), Presets, bInForceUpdate);
 	}
 }
diff --git a/Engine/Plugins/VirtualProduction/Avalanche/Source/AvalancheMediaEditor/Private/Rundown/DetailsView/RemoteControl/Controllers/SAvaRundownRCControllerPanel.h b/Engine/Plugins/VirtualProduction/Avalanche/Source/AvalancheMediaEditor/Private/Rundown/DetailsView/RemoteControl/Controllers/SAvaRundownRCControllerPanel.h
--- a/Engine/Plugins/VirtualProduction/Avalanche/Source/AvalancheMediaEditor/Private/Rundown/DetailsView/RemoteControl/Controllers/SAvaRundownRCControllerPanel.h
+++ b/Engine/Plugins/VirtualProduction/Avalanche/Source/AvalancheMediaEditor/Private/Rundown/DetailsView/RemoteControl/Controllers/SAvaRundownRCControllerPanel.h
@@ -3,6 +3,7 @@
 #pragma once
 
 #include "Misc/NotifyHook.h"
+#include "Rundown/AvaRundownManagedInstanceHandle.h"
 #include "Templates/SharedPointer.h"
 #include "Widgets/SCompoundWidget.h"
 
@@ -103,7 +104,7 @@ private:
 	
 	TArray<TUniquePtr<FPropertyRowGeneratorWrapper>> PropertyRowGenerators;
 
-	TArray<TSharedPtr<FAvaRundownManagedInstance>> ManagedInstances;
+	FAvaRundownManagedInstanceHandles ManagedHandles;
 	
 	TSharedPtr<SListView<FAvaRundownRCControllerItemPtr>> ControllerContainer;
 	
diff --git a/Engine/Plugins/VirtualProduction/Avalanche/Source/AvalancheMediaEditor/Private/Rundown/DetailsView/RemoteControl/Properties/SAvaRundownPageRemoteControlProps.cpp b/Engine/Plugins/VirtualProduction/Avalanche/Source/AvalancheMediaEditor/Private/Rundown/DetailsView/RemoteControl/Properties/SAvaRundownPageRemoteControlProps.cpp
--- a/Engine/Plugins/VirtualProduction/Avalanche/Source/AvalancheMediaEditor/Private/Rundown/DetailsView/RemoteControl/Properties/SAvaRundownPageRemoteControlProps.cpp
+++ b/Engine/Plugins/VirtualProduction/Avalanche/Source/AvalancheMediaEditor/Private/Rundown/DetailsView/RemoteControl/Properties/SAvaRundownPageRemoteControlProps.cpp
@@ -1,14 +1,14 @@
 // Copyright Epic Games, Inc. All Rights Reserved.
 
 #include "SAvaRundownPageRemoteControlProps.h"
+
 #include "Controller/RCCustomControllerUtilities.h"
+#include "IAvaMediaModule.h"
 #include "Item/AvaRundownRCFieldItem.h"
-#include "Modules/ModuleManager.h"
 #include "Playable/AvaPlayableRemoteControl.h"
 #include "Playable/AvaPlayableRemoteControlPresetInfo.h"
 #include "Playable/AvaPlayableRemoteControlValues.h"
 #include "Playback/AvaPlaybackUtils.h"
-#include "PropertyEditorModule.h"
 #include "RemoteControlEntity.h"
 #include "RemoteControlPreset.h"
 #include "Rundown/AvaRundownEditor.h"
@@ -19,7 +19,6 @@
 #include "Rundown/DetailsView/RemoteControl/Properties/AvaRundownPagePropertyContextMenu.h"
 #include "Rundown/DetailsView/SAvaRundownPageDetails.h"
 #include "SAvaRundownRCPropertyItemRow.h"
-#include "SlateOptMacros.h"
 #include "UObject/NameTypes.h"
 #include "Widgets/Views/ITableRow.h"
 #include "Widgets/Views/STableViewBase.h"
@@ -232,7 +231,9 @@ void SAvaRundownPageRemoteControlProps::Refresh(const TArray<int32>& InSelectedP
 			if (FAvaRundownPage* ActivePage = GetActivePage())
 			{
 				using namespace UE::AvaRundownEditor::Utils;
-				ManagedInstances = GetManagedInstancesForPage(Rundown, *ActivePage);
+				// Note: we don't use the handles api here because we are not acquiring the events.
+				// Todo: the managed instances should be more unified between the Entity/Prop and Controller panels.
+				ManagedInstances = IAvaMediaModule::Get().GetManagedInstanceCache().GetManagedInstancesForPage(Rundown, *ActivePage);
 
 				for (const TSharedPtr<FAvaRundownManagedInstance>& ManagedInstance : ManagedInstances)
 				{
