Motion Design: Fixed guide loading saving during load.

diff --git a/Engine/Plugins/VirtualProduction/Avalanche/Source/AvalancheLevelViewport/Private/SAvaLevelViewport.cpp b/Engine/Plugins/VirtualProduction/Avalanche/Source/AvalancheLevelViewport/Private/SAvaLevelViewport.cpp
--- a/Engine/Plugins/VirtualProduction/Avalanche/Source/AvalancheLevelViewport/Private/SAvaLevelViewport.cpp
+++ b/Engine/Plugins/VirtualProduction/Avalanche/Source/AvalancheLevelViewport/Private/SAvaLevelViewport.cpp
@@ -501,7 +501,7 @@ void SAvaLevelViewport::LoadGuides()
 
 		if (NewGuide.IsValid())
 		{
-			NewGuide->SetState(GuideInfo.State);
+			NewGuide->SetState(GuideInfo.State, /* Save Guides */ false);
 		}
 	}
 }
diff --git a/Engine/Plugins/VirtualProduction/Avalanche/Source/AvalancheLevelViewport/Private/Widgets/SAvaLevelViewportGuide.cpp b/Engine/Plugins/VirtualProduction/Avalanche/Source/AvalancheLevelViewport/Private/Widgets/SAvaLevelViewportGuide.cpp
--- a/Engine/Plugins/VirtualProduction/Avalanche/Source/AvalancheLevelViewport/Private/Widgets/SAvaLevelViewportGuide.cpp
+++ b/Engine/Plugins/VirtualProduction/Avalanche/Source/AvalancheLevelViewport/Private/Widgets/SAvaLevelViewportGuide.cpp
@@ -894,10 +894,15 @@ void SAvaLevelViewportGuide::ToggleEnabled()
 	}
 }
 
-void SAvaLevelViewportGuide::SetState(EAvaViewportGuideState InNewState)
+void SAvaLevelViewportGuide::SetState(EAvaViewportGuideState InNewState, bool bInSaveGuides)
 {
 	Info.State = InNewState;
 
+	if (!bInSaveGuides)
+	{
+		return;
+	}
+
 	TSharedPtr<SAvaLevelViewportFrame> ViewportFrame = GetViewportFrame();
 
 	if (!ViewportFrame.IsValid())
diff --git a/Engine/Plugins/VirtualProduction/Avalanche/Source/AvalancheLevelViewport/Private/Widgets/SAvaLevelViewportGuide.h b/Engine/Plugins/VirtualProduction/Avalanche/Source/AvalancheLevelViewport/Private/Widgets/SAvaLevelViewportGuide.h
--- a/Engine/Plugins/VirtualProduction/Avalanche/Source/AvalancheLevelViewport/Private/Widgets/SAvaLevelViewportGuide.h
+++ b/Engine/Plugins/VirtualProduction/Avalanche/Source/AvalancheLevelViewport/Private/Widgets/SAvaLevelViewportGuide.h
@@ -59,7 +59,7 @@ public:
 	void Construct(const FArguments& InArgs, TSharedPtr<SAvaLevelViewportFrame> InViewportFrame, int32 InIndex);
 
 	const FAvaViewportGuideInfo& GetGuideInfo() const { return Info; }
-	void SetState(EAvaViewportGuideState InNewState);
+	void SetState(EAvaViewportGuideState InNewState, bool bInSaveGuides = true);
 	void SetLocked(bool bInLocked) { Info.bLocked = bInLocked; }
 	bool IsBeingDragged() const { return bBeingDragged; }
 	EVisibility GetGuideVisibility() const;
