Motion Design: enabled transiton logic library functions to be used by context objects that are able to retrieve transition information.
Before, only transition nodes were allowed (for TL blueprints). Now, world context sensitive objects (e.g. actor blueprints, etc) are able to use these functions too.

diff --git a/Engine/Plugins/VirtualProduction/Avalanche/Config/DefaultAvalanche.ini b/Engine/Plugins/VirtualProduction/Avalanche/Config/DefaultAvalanche.ini
--- a/Engine/Plugins/VirtualProduction/Avalanche/Config/DefaultAvalanche.ini
+++ b/Engine/Plugins/VirtualProduction/Avalanche/Config/DefaultAvalanche.ini
@@ -307,5 +307,13 @@
 
 +PropertyRedirects=(OldName="/Script/AvalancheModifiers.AvaDynamicMeshConverterModifier.TransformUpdateInterval",NewName="/Script/AvalancheModifiers.AvaDynamicMeshConverterModifier.UpdateInterval")
 
+; 5.7 Redirects
+
+; Renaming TransitionNode to ContextObject to support any object that can retrieve the transition context (e.g. world context objects)
++PropertyRedirects=(OldName="AvaTransitionLibrary.IsTransitionActiveInLayer.InTransitionNode",NewName="InContextObject")
++PropertyRedirects=(OldName="AvaTransitionLibrary.GetTransitionType.InTransitionNode",NewName="InContextObject")
++PropertyRedirects=(OldName="AvaTransitionLibrary.AreScenesTransitioning.InTransitionNode",NewName="InContextObject")
++PropertyRedirects=(OldName="AvaTransitionLibrary.GetTransitionTree.InTransitionNode",NewName="InContextObject")
+
 [/Script/Avalanche.AvaAnimation]
 DefaultCompletionMode=KeepState
diff --git a/Engine/Plugins/VirtualProduction/Avalanche/Source/AvalancheTransition/Private/AvaTransitionLibrary.cpp b/Engine/Plugins/VirtualProduction/Avalanche/Source/AvalancheTransition/Private/AvaTransitionLibrary.cpp
--- a/Engine/Plugins/VirtualProduction/Avalanche/Source/AvalancheTransition/Private/AvaTransitionLibrary.cpp
+++ b/Engine/Plugins/VirtualProduction/Avalanche/Source/AvalancheTransition/Private/AvaTransitionLibrary.cpp
@@ -7,23 +7,135 @@
 #include "AvaTransitionSubsystem.h"
 #include "Behavior/AvaTransitionBehaviorInstance.h"
 #include "Behavior/AvaTransitionBehaviorInstanceCache.h"
+#include "Behavior/IAvaTransitionBehavior.h"
+#include "Engine/Engine.h"
 #include "Engine/Level.h"
 #include "Engine/World.h"
+#include "Execution/IAvaTransitionExecutor.h"
 #include "IAvaTransitionNodeInterface.h"
 #include "UObject/Package.h"
 
-bool UAvaTransitionLibrary::IsTransitionActiveInLayer(UObject* InTransitionNode
-	, EAvaTransitionComparisonResult InSceneComparisonType
-	, EAvaTransitionLayerCompareType InLayerComparisonType
-	, const FAvaTagHandleContainer& InSpecificLayers)
+namespace UE::AvaTransition::Private
 {
-	const IAvaTransitionNodeInterface* NodeInterface = Cast<IAvaTransitionNodeInterface>(InTransitionNode);
-	if (!NodeInterface)
+	EAvaTransitionType GetTransitionTypeFromFilter(EAvaTransitionFilterType InFilterType)
 	{
-		return false;
+		switch (InFilterType)
+		{
+		case EAvaTransitionFilterType::In:  return EAvaTransitionType::In;
+		case EAvaTransitionFilterType::Out: return EAvaTransitionType::Out;
+		}
+		return EAvaTransitionType::None;
+	}
+
+	ULevel* GetLevelFromContextObject(UObject* InContextObject)
+	{
+		if (!InContextObject)
+		{
+			return nullptr;
+		}
+
+		if (ULevel* ContextLevel = Cast<ULevel>(InContextObject))
+		{
+			return ContextLevel;
+		}
+
+		if (ULevel* ContextLevel = InContextObject->GetTypedOuter<ULevel>())
+		{
+			return ContextLevel;
+		}
+
+		UWorld* const World = GEngine ? GEngine->GetWorldFromContextObject(InContextObject, EGetWorldErrorMode::LogAndReturnNull) : nullptr;
+		return World ? World->PersistentLevel : nullptr;
 	}
 
-	const FAvaTransitionContext* TransitionContext = NodeInterface->GetBehaviorInstanceCache().GetTransitionContext();
+	UAvaTransitionSubsystem* GetSubsystemFromLevel(ULevel* InLevelChecked)
+	{
+		check(InLevelChecked);
+
+		// Use the main world of a level rather than the outer world (in the case it's a streamed in level)
+		if (ensure(InLevelChecked->OwningWorld))
+		{
+			return InLevelChecked->OwningWorld->GetSubsystem<UAvaTransitionSubsystem>();
+		}
+
+		return nullptr;
+	}
+	
+	/**
+	 * Retrieves the transition contexts corresponding to a given level
+	 * The expectation is that there will be only 1 or no transition contexts, but the API allows it to be technically possible to have more than 1.
+	 */
+	TArray<const FAvaTransitionContext*, TInlineAllocator<1>> GetTransitionContextsFromLevel(ULevel* InLevelChecked)
+	{
+		check(InLevelChecked);
+
+		UAvaTransitionSubsystem* const TransitionSubsystem = GetSubsystemFromLevel(InLevelChecked);
+		if (!TransitionSubsystem)
+		{
+			return {};
+		}
+
+		TArray<const FAvaTransitionContext*, TInlineAllocator<1>> TransitionContexts;
+
+		// Gather all the transition contexts where the transition scene level matches the provided level
+		TransitionSubsystem->ForEachTransitionExecutor(
+			[InLevelChecked, &TransitionContexts](IAvaTransitionExecutor& InExecutor)->EAvaTransitionIterationResult
+			{
+				InExecutor.ForEachBehaviorInstance(
+					[InLevelChecked, &TransitionContexts](const FAvaTransitionBehaviorInstance& InBehaviorInstance)
+					{
+						const FAvaTransitionContext& TransitionContext = InBehaviorInstance.GetTransitionContext();
+						const FAvaTransitionScene* TransitionScene = TransitionContext.GetTransitionScene();
+
+						if (TransitionScene && TransitionScene->GetLevel() == InLevelChecked)
+						{
+							TransitionContexts.Add(&TransitionContext);
+						}
+					});
+
+				return EAvaTransitionIterationResult::Continue;
+			});
+
+		return TransitionContexts;
+	}
+}
+
+const FAvaTransitionContext* UAvaTransitionLibrary::GetTransitionContext(UObject* InContextObject)
+{
+	using namespace UE::AvaTransition;
+
+	if (const IAvaTransitionNodeInterface* NodeInterface = Cast<IAvaTransitionNodeInterface>(InContextObject))
+	{
+		return NodeInterface->GetBehaviorInstanceCache().GetTransitionContext();
+	}
+
+	if (ULevel* const Level = Private::GetLevelFromContextObject(InContextObject))
+	{
+		const TArray<const FAvaTransitionContext*, TInlineAllocator<1>> TransitionContexts = Private::GetTransitionContextsFromLevel(Level);
+		if (TransitionContexts.IsEmpty())
+		{
+			return nullptr;
+		}
+
+#if UE_BUILD_DEBUG
+		// Expectation is that for a given level instance,
+		// there should only be one transition occurring, even if the system technically allows multiple instances to run for a behavior
+		ensure(TransitionContexts.Num() == 1);
+#endif
+		return TransitionContexts[0];
+	}
+
+	return nullptr;
+}
+
+bool UAvaTransitionLibrary::IsTransitionActiveInLayers(UObject* InWorldContextObject
+	, EAvaTransitionLayerCompareType InLayerComparisonType
+	, const FAvaTagHandleContainer& InSpecificLayers
+	, EAvaTransitionFilterType InTransitionFilterType)
+{
+	using namespace UE::AvaTransition;
+
+	const FAvaTransitionContext* TransitionContext = GetTransitionContext(InWorldContextObject);
 	if (!TransitionContext)
 	{
 		return false;
@@ -57,30 +169,43 @@ bool UAvaTransitionLibrary::IsTransitionActiveInLayer(UObject* InTransitionNode
 		return false;
 	}
 
+	const EAvaTransitionType TransitionType = Private::GetTransitionTypeFromFilter(InTransitionFilterType);
+
 	bool bHasMatchingInstance = BehaviorInstances.ContainsByPredicate(
-		[TransitionScene, InSceneComparisonType](const FAvaTransitionBehaviorInstance* InInstance)
+		[TransitionScene, TransitionType](const FAvaTransitionBehaviorInstance* InInstance)
 		{
+			// if transition type is set, ensure that the behavior instance matches the given transition type
+			if (TransitionType != EAvaTransitionType::None && TransitionType != InInstance->GetTransitionType())
+			{
+				return false;
+			}
+
 			const FAvaTransitionScene* OtherTransitionScene = InInstance->GetTransitionContext().GetTransitionScene();
 
 			const EAvaTransitionComparisonResult ComparisonResult = OtherTransitionScene
 				? TransitionScene->Compare(*OtherTransitionScene)
 				: EAvaTransitionComparisonResult::None;
 
-			return ComparisonResult == InSceneComparisonType;
+			return ComparisonResult != EAvaTransitionComparisonResult::None;
 		});
 
 	return bHasMatchingInstance;
 }
 
-EAvaTransitionType UAvaTransitionLibrary::GetTransitionType(UObject* InTransitionNode)
+bool UAvaTransitionLibrary::IsTransitionActiveInLayer(UObject* InWorldContextObject
+	, EAvaTransitionComparisonResult InSceneComparisonType
+	, EAvaTransitionLayerCompareType InLayerComparisonType
+	, const FAvaTagHandleContainer& InSpecificLayers)
 {
-	const IAvaTransitionNodeInterface* NodeInterface = Cast<IAvaTransitionNodeInterface>(InTransitionNode);
-	if (!NodeInterface)
-	{
-		return EAvaTransitionType::None;
-	}
+	return IsTransitionActiveInLayers(InWorldContextObject
+		, InLayerComparisonType
+		, InSpecificLayers
+		, EAvaTransitionFilterType::Any);
+}
 
-	const FAvaTransitionContext* TransitionContext = NodeInterface->GetBehaviorInstanceCache().GetTransitionContext();
+EAvaTransitionType UAvaTransitionLibrary::GetTransitionType(UObject* InContextObject)
+{
+	const FAvaTransitionContext* TransitionContext = GetTransitionContext(InContextObject);
 	if (!TransitionContext)
 	{
 		return EAvaTransitionType::None;
@@ -89,15 +214,11 @@ EAvaTransitionType UAvaTransitionLibrary::GetTransitionType(UObject* InTransitio
 	return TransitionContext->GetTransitionType();
 }
 
-bool UAvaTransitionLibrary::AreScenesTransitioning(UObject* InTransitionNode, const FAvaTagHandleContainer& InLayers, const TArray<TSoftObjectPtr<UWorld>>& InScenesToIgnore)
+bool UAvaTransitionLibrary::AreScenesTransitioning(UObject* InContextObject
+	, const FAvaTagHandleContainer& InLayers
+	, const TArray<TSoftObjectPtr<UWorld>>& InScenesToIgnore)
 {
-	const IAvaTransitionNodeInterface* NodeInterface = Cast<IAvaTransitionNodeInterface>(InTransitionNode);
-	if (!NodeInterface)
-	{
-		return false;
-	}
-
-	const FAvaTransitionContext* TransitionContext = NodeInterface->GetBehaviorInstanceCache().GetTransitionContext();
+	const FAvaTransitionContext* TransitionContext = GetTransitionContext(InContextObject);
 	if (!TransitionContext)
 	{
 		return false;
@@ -189,11 +310,22 @@ bool UAvaTransitionLibrary::AreScenesTransitioning(UObject* InTransitionNode, co
 	return false;
 }
 
-const UAvaTransitionTree* UAvaTransitionLibrary::GetTransitionTree(UObject* InTransitionNode)
+const UAvaTransitionTree* UAvaTransitionLibrary::GetTransitionTree(UObject* InContextObject)
 {
-	if (const IAvaTransitionNodeInterface* NodeInterface = Cast<IAvaTransitionNodeInterface>(InTransitionNode))
+	using namespace UE::AvaTransition;
+
+	if (const IAvaTransitionNodeInterface* NodeInterface = Cast<IAvaTransitionNodeInterface>(InContextObject))
 	{
 		return NodeInterface->GetBehaviorInstanceCache().GetTransitionTree();
 	}
+
+	if (ULevel* Level = Private::GetLevelFromContextObject(InContextObject))
+	{
+		if (UAvaTransitionSubsystem* const TransitionSubsystem = Private::GetSubsystemFromLevel(Level))
+		{
+			const IAvaTransitionBehavior* TransitionBehavior = TransitionSubsystem->GetTransitionBehavior(Level);
+			return TransitionBehavior ? TransitionBehavior->GetTransitionTree() : nullptr;
+		}
+	}
 	return nullptr;
 }
diff --git a/Engine/Plugins/VirtualProduction/Avalanche/Source/AvalancheTransition/Private/AvaTransitionLibrary.h b/Engine/Plugins/VirtualProduction/Avalanche/Source/AvalancheTransition/Private/AvaTransitionLibrary.h
deleted file mode 100644
diff --git a/Engine/Plugins/VirtualProduction/Avalanche/Source/AvalancheTransition/Public/AvaTransitionEnums.h b/Engine/Plugins/VirtualProduction/Avalanche/Source/AvalancheTransition/Public/AvaTransitionEnums.h
--- a/Engine/Plugins/VirtualProduction/Avalanche/Source/AvalancheTransition/Public/AvaTransitionEnums.h
+++ b/Engine/Plugins/VirtualProduction/Avalanche/Source/AvalancheTransition/Public/AvaTransitionEnums.h
@@ -27,6 +27,14 @@ enum class EAvaTransitionType : uint8
 };
 ENUM_CLASS_FLAGS(EAvaTransitionType);
 
+UENUM(BlueprintType, DisplayName="Motion Design Transition Filter Type")
+enum class EAvaTransitionFilterType : uint8
+{
+	In,
+	Out,
+	Any,
+};
+
 UENUM(BlueprintType, DisplayName="Motion Design Transition Run State")
 enum class EAvaTransitionRunState : uint8
 {
@@ -56,7 +64,7 @@ enum class EAvaTransitionLayerCompareType : uint8
 	None UMETA(Hidden),
 	Same,
 	Different UMETA(DisplayName="Other"),
-	MatchingTag,
+	MatchingTag UMETA(DisplayName="Specific Layers"),
 	Any,
 };
 
diff --git a/Engine/Plugins/VirtualProduction/Avalanche/Source/AvalancheTransition/Public/AvaTransitionLibrary.h b/Engine/Plugins/VirtualProduction/Avalanche/Source/AvalancheTransition/Public/AvaTransitionLibrary.h
new file mode 100644
--- /dev/null
+++ b/Engine/Plugins/VirtualProduction/Avalanche/Source/AvalancheTransition/Public/AvaTransitionLibrary.h
@@ -0,0 +1,72 @@
+// Copyright Epic Games, Inc. All Rights Reserved.
+
+#pragma once
+
+#include "AvaTransitionEnums.h"
+#include "Kismet/BlueprintFunctionLibrary.h"
+#include "AvaTransitionLibrary.generated.h"
+
+#define UE_API AVALANCHETRANSITION_API
+
+class UAvaTransitionTree;
+struct FAvaTagHandleContainer;
+struct FAvaTransitionContext;
+
+UCLASS(MinimalAPI)
+class UAvaTransitionLibrary : public UBlueprintFunctionLibrary
+{
+	GENERATED_BODY()
+
+public:
+	/** Retrieves the transition context from */
+	static UE_API const FAvaTransitionContext* GetTransitionContext(UObject* InContextObject);
+
+	/**
+	 * Determines whether a transition is currently taking place at a given layer
+	 * @param InContextObject the context to retrieve the active transition relating to the context.
+	 * @param InLayerComparisonType the type of layer to filter.
+	 *      "Same": the same layer as this context's.
+	 *      "Different": a layer different from this context's.
+	 *      "MatchingTag": a layer specified by "InSpecificLayers".
+	 *      "Any": any layer found.
+	 * @param InSpecificLayers the specific layers to filter. InLayerComparisonType should be set to "Specific Layers" for this to be considered.
+	 * @param InTransitionFilterType the transition type to filter (whether it should only consider Ins or Outs, or any type is allowed)
+	 */
+	UFUNCTION(BlueprintPure, Category="Transition Logic", meta=(DefaultToSelf="InContextObject"))
+	static UE_API bool IsTransitionActiveInLayers(UObject* InContextObject
+		, EAvaTransitionLayerCompareType InLayerComparisonType
+		, const FAvaTagHandleContainer& InSpecificLayers
+		, EAvaTransitionFilterType InTransitionFilterType = EAvaTransitionFilterType::Any);
+
+	UE_DEPRECATED(5.7, "Use IsTransitionActiveInLayers instead")
+	UFUNCTION(BlueprintPure, Category="Transition Logic", meta=(DefaultToSelf="InContextObject", DeprecatedFunction, DeprecationMessage="Use \"IsTransitionActiveInLayers\" instead"))
+	static bool IsTransitionActiveInLayer(UObject* InContextObject
+		, EAvaTransitionComparisonResult InSceneComparisonType
+		, EAvaTransitionLayerCompareType InLayerComparisonType
+		, const FAvaTagHandleContainer& InSpecificLayers);
+
+	/** Returns the transition type of the current transition taking place if any */
+	UFUNCTION(BlueprintPure, Category="Transition Logic", meta=(DefaultToSelf="InContextObject"))
+	static UE_API EAvaTransitionType GetTransitionType(UObject* InContextObject);
+
+	/**
+	 * Determines whether there are scenes/levels (other than the one from this context) transitioning.
+	 * Note: this ignores scenes that are transitioning out and to be discarded.
+	 * @param InContextObject the context to retrieve the active transition relating to the context.
+	 * @param InLayers the layers to consider
+	 * @param InScenesToIgnore the world assets to ignore
+	 * @return true if the given layers have assets that are involved in a transition.
+	 */
+	UFUNCTION(BlueprintPure, Category="Transition Logic", meta=(DefaultToSelf="InContextObject"))
+	static UE_API bool AreScenesTransitioning(UObject* InContextObject
+		, const FAvaTagHandleContainer& InLayers
+		, const TArray<TSoftObjectPtr<UWorld>>& InScenesToIgnore);
+
+	/**
+	 * Retrieves the transition tree from this context
+	 */
+	UFUNCTION(BlueprintPure, Category="Transition Logic", meta=(DefaultToSelf="InContextObject"))
+	static UE_API const UAvaTransitionTree* GetTransitionTree(UObject* InContextObject);
+};
+
+#undef UE_API
