[Motion Design] Reworking how RundownManagedInstanceHandles tracks and unregisters (from RC event routing) the managed  instances.

Improvements:
- Move Operator: Better logic to make sure to unregister instances that are not part of the new set.
- Destructor: It will also check to unregister only the instances that are from the same page. Instances registered to another page should not be unregistered, even if it is the same rundown.

Additonal changes:
- Fixing a few comments that where reported misleading or inaccurate by BugHawk.
- Moved the InvalidPageId to a runtime rundown header (instead of editor).

diff --git a/Engine/Plugins/VirtualProduction/Avalanche/Source/AvalancheMedia/Internal/Rundown/AvaRundownManagedInstance.h b/Engine/Plugins/VirtualProduction/Avalanche/Source/AvalancheMedia/Internal/Rundown/AvaRundownManagedInstance.h
--- a/Engine/Plugins/VirtualProduction/Avalanche/Source/AvalancheMedia/Internal/Rundown/AvaRundownManagedInstance.h
+++ b/Engine/Plugins/VirtualProduction/Avalanche/Source/AvalancheMedia/Internal/Rundown/AvaRundownManagedInstance.h
@@ -43,15 +43,21 @@ public:
 
 	const FSoftObjectPath& GetSourceAssetPath() const { return SourceAssetPath; }
 
-	/** Set the current rundown page for routing controller events. */
+	/** Set the current rundown page for routing controller events from this managed instance. */
 	void SetCurrentRundownPageForEvents(UAvaRundown* InRundown, int32 InPageId);
 
-	/** Returns the rundown currently set to receive the RC Controller events. */
+	/** Returns the rundown currently set to receive the RC Controller events from this managed instance. */
 	UAvaRundown* GetRundownForEvents() const
 	{
 		return RundownForEvents.Get();
 	}
 
+	/** Returns the rundown page id that is currently set to receive RC Controller events from this managed instance. */
+	int32 GetActivePageIdForEvents() const
+	{
+		return ActivePageId;
+	}
+
 protected:
 	FAvaPlayableRemoteControlPresetInfo& GetRemoteControlPresetInfoFromCache(URemoteControlPreset* InRemoteControlPreset) const;
 	void InvalidateFromCache();
diff --git a/Engine/Plugins/VirtualProduction/Avalanche/Source/AvalancheMedia/Internal/Rundown/AvaRundownManagedInstanceCache.h b/Engine/Plugins/VirtualProduction/Avalanche/Source/AvalancheMedia/Internal/Rundown/AvaRundownManagedInstanceCache.h
--- a/Engine/Plugins/VirtualProduction/Avalanche/Source/AvalancheMedia/Internal/Rundown/AvaRundownManagedInstanceCache.h
+++ b/Engine/Plugins/VirtualProduction/Avalanche/Source/AvalancheMedia/Internal/Rundown/AvaRundownManagedInstanceCache.h
@@ -91,7 +91,7 @@ public:
 	 * 
 	 * @param InRundown Rundown owning the current page.
 	 * @param InPage Page to get the assets from.
-	 * @return managed asset instances for the given page.
+	 * @return managed asset instance handles for the given page.
 	 */
 	FAvaRundownManagedInstanceHandles GetManagedHandlesForPage(UAvaRundown* InRundown, const FAvaRundownPage& InPage);
 	
diff --git a/Engine/Plugins/VirtualProduction/Avalanche/Source/AvalancheMedia/Internal/Rundown/AvaRundownManagedInstanceHandle.h b/Engine/Plugins/VirtualProduction/Avalanche/Source/AvalancheMedia/Internal/Rundown/AvaRundownManagedInstanceHandle.h
--- a/Engine/Plugins/VirtualProduction/Avalanche/Source/AvalancheMedia/Internal/Rundown/AvaRundownManagedInstanceHandle.h
+++ b/Engine/Plugins/VirtualProduction/Avalanche/Source/AvalancheMedia/Internal/Rundown/AvaRundownManagedInstanceHandle.h
@@ -3,6 +3,7 @@
 #pragma once
 
 #include "Containers/Array.h"
+#include "Rundown/AvaRundownDefines.h"
 #include "Templates/SharedPointer.h"
 #include "UObject/WeakObjectPtr.h"
 #include "UObject/WeakObjectPtrTemplates.h"
@@ -37,7 +38,7 @@ struct FAvaRundownManagedInstanceHandles
 		*this = MoveTemp(InOther);
 	}
 
-	/** Copy Move operator handles unregistering instances if needed. */
+	/** Move operator handles unregistering instances if needed. */
 	UE_API FAvaRundownManagedInstanceHandles& operator=(FAvaRundownManagedInstanceHandles&& InOther);
 
 	~FAvaRundownManagedInstanceHandles()
@@ -60,11 +61,15 @@ struct FAvaRundownManagedInstanceHandles
 	 */
 	UE_API void Reset();
 
+	/** Array of managed asset instances. */
+	TArray<TSharedPtr<FAvaRundownManagedInstance>> Instances;
+
+private:
 	/** Keep track of the registered Rundown to be able to unregister if needed. */
 	TWeakObjectPtr<UAvaRundown> RegisteredRundown;
 
-	/** Array of managed asset instances. */
-	TArray<TSharedPtr<FAvaRundownManagedInstance>> Instances;
+	/** Keep track of the registered page id to be able to unregister if needed. */
+	int32 RegisteredPageId = UE::AvaMedia::Rundown::InvalidPageId;
 };
 
 #undef UE_API
\ No newline at end of file
diff --git a/Engine/Plugins/VirtualProduction/Avalanche/Source/AvalancheMedia/Private/Rundown/AvaRundown.cpp b/Engine/Plugins/VirtualProduction/Avalanche/Source/AvalancheMedia/Private/Rundown/AvaRundown.cpp
--- a/Engine/Plugins/VirtualProduction/Avalanche/Source/AvalancheMedia/Private/Rundown/AvaRundown.cpp
+++ b/Engine/Plugins/VirtualProduction/Avalanche/Source/AvalancheMedia/Private/Rundown/AvaRundown.cpp
@@ -1983,7 +1983,7 @@ namespace UE::AvaMedia::Rundown::Private
 	}
 }
 
-bool UAvaRundown::PushRuntimeRemoteControlEvents(TConstArrayView<FGuid> InRcItemIds, int32 InPageId, bool bInPushToProgram, bool bInPushToPreview, FName InPreviewChannelName)
+bool UAvaRundown::PushRuntimeRemoteControlEvents(TConstArrayView<FGuid> InControllerIds, int32 InPageId, bool bInPushToProgram, bool bInPushToPreview, FName InPreviewChannelName)
 {
 	const FAvaRundownPage& Page = GetPage(InPageId);
 	if (!Page.IsValidPage())
@@ -2021,7 +2021,7 @@ bool UAvaRundown::PushRuntimeRemoteControlEvents(TConstArrayView<FGuid> InRcItem
 		if (!ControllerEvents.IsValid())
 		{
 			using namespace UE::AvaMedia::Rundown::Private;
-			ControllerEvents = MakeControllerEvents(InRcItemIds, Page.GetRemoteControlValues());
+			ControllerEvents = MakeControllerEvents(InControllerIds, Page.GetRemoteControlValues());
 		}
 
 		for (int32 InstanceIndex = 0; InstanceIndex < PagePlayer->GetNumInstancePlayers(); ++InstanceIndex)
@@ -2030,9 +2030,9 @@ bool UAvaRundown::PushRuntimeRemoteControlEvents(TConstArrayView<FGuid> InRcItem
 			{
 				constexpr EAvaPlayableRCUpdateFlags UpdateFlags = EAvaPlayableRCUpdateFlags::ExecuteControllerBehaviors | EAvaPlayableRCUpdateFlags::ControllerValuesAsEvents;
 				Playback->PushRemoteControlValues(PagePlayer->GetSourceAssetPath(InstanceIndex), PagePlayer->ChannelName, ControllerEvents.ToSharedRef(), UpdateFlags);
+				bValuesPushed = true;
 			}
 		}
-		bValuesPushed = true;
 	}
 	return bValuesPushed;
 }
@@ -2068,9 +2068,9 @@ bool UAvaRundown::PushRuntimeRemoteControlValues(int32 InPageId, bool bInIsPrevi
 				// Remote Control Value (live) updates, when not part of a transition, will execute the controller behaviors on the runtime instances.
 				constexpr EAvaPlayableRCUpdateFlags UpdateFlags = EAvaPlayableRCUpdateFlags::ExecuteControllerBehaviors;
 				Playback->PushRemoteControlValues(PagePlayer->GetSourceAssetPath(InstanceIndex), PagePlayer->ChannelName, SharedRCValues, UpdateFlags);
+				bValuesPushed = true;
 			}
 		}
-		bValuesPushed = true;
 	}
 	return bValuesPushed;
 }
diff --git a/Engine/Plugins/VirtualProduction/Avalanche/Source/AvalancheMedia/Private/Rundown/AvaRundownManagedInstanceHandle.cpp b/Engine/Plugins/VirtualProduction/Avalanche/Source/AvalancheMedia/Private/Rundown/AvaRundownManagedInstanceHandle.cpp
--- a/Engine/Plugins/VirtualProduction/Avalanche/Source/AvalancheMedia/Private/Rundown/AvaRundownManagedInstanceHandle.cpp
+++ b/Engine/Plugins/VirtualProduction/Avalanche/Source/AvalancheMedia/Private/Rundown/AvaRundownManagedInstanceHandle.cpp
@@ -4,33 +4,32 @@
 
 #include "Rundown/AvaRundown.h"
 #include "Rundown/AvaRundownManagedInstance.h"
-#include "Rundown/AvaRundownPage.h"
 
 namespace UE::AvaMedia::Rundown::Private
 {
-	/** Utility function to unregister the current rundown from controller events routing. */
-	void UnregisterRundownForEvents(const TSharedPtr<FAvaRundownManagedInstance>& InInstance, const UAvaRundown* InRundown)
+	/** Utility function to unregister the current rundown page from controller events routing. */
+	void UnregisterRundownForEvents(const TSharedPtr<FAvaRundownManagedInstance>& InInstance, const UAvaRundown* InRundown, int32 InPageId)
 	{
-		// Don't reset the rundown in case this instance is already used by another rundown.
-		if (InInstance && InInstance->GetRundownForEvents() == InRundown)
+		if (InInstance && InInstance->GetRundownForEvents() == InRundown && InInstance->GetActivePageIdForEvents() == InPageId)
 		{
-			InInstance->SetCurrentRundownPageForEvents(nullptr, FAvaRundownPage::InvalidPageId);
+			InInstance->SetCurrentRundownPageForEvents(nullptr, InvalidPageId);
 		}
 	}
 	
-	/** Utility function to unregister the current rundown from controller events routing. */
-	void UnregisterRundownForEvents(TConstArrayView<TSharedPtr<FAvaRundownManagedInstance>> InInstances, const UAvaRundown* InRundown)
+	/** Utility function to unregister the current rundown page from controller events routing. */
+	void UnregisterRundownForEvents(TConstArrayView<TSharedPtr<FAvaRundownManagedInstance>> InInstances, const UAvaRundown* InRundown, int32 InPageId)
 	{
 		for (const TSharedPtr<FAvaRundownManagedInstance>& Instance : InInstances)
 		{
-			UnregisterRundownForEvents(Instance, InRundown);
+			UnregisterRundownForEvents(Instance, InRundown, InPageId);
 		}
 	}
 }
 
 FAvaRundownManagedInstanceHandles::FAvaRundownManagedInstanceHandles(const TArray<TSharedPtr<FAvaRundownManagedInstance>>& InInstances, UAvaRundown* InRundown, int32 InPageId)
-	: RegisteredRundown(InRundown)
-	, Instances(InInstances)
+	: Instances(InInstances)
+	, RegisteredRundown(InRundown)
+	, RegisteredPageId(InPageId)
 {
 	for (const TSharedPtr<FAvaRundownManagedInstance>& Instance : Instances)
 	{
@@ -41,7 +40,7 @@ FAvaRundownManagedInstanceHandles::FAvaRundownManagedInstanceHandles(const TArra
 
 void FAvaRundownManagedInstanceHandles::Reset()
 {
-	UE::AvaMedia::Rundown::Private::UnregisterRundownForEvents(Instances, RegisteredRundown.Get());
+	UE::AvaMedia::Rundown::Private::UnregisterRundownForEvents(Instances, RegisteredRundown.Get(), RegisteredPageId);
 	Instances.Reset();
 }
 
@@ -51,15 +50,19 @@ FAvaRundownManagedInstanceHandles& FAvaRundownManagedInstanceHandles::operator=(
 	{
 		return *this;
 	}
-	
-	// Need to unregister previous instances from another rundown before copying.
-	if (!Instances.IsEmpty() && RegisteredRundown != InOther.RegisteredRundown)
+
+	// Remove any instances from the incoming set of instance.
+	// The remaining instances will be unregistered by the destructor.
+	for (TSharedPtr<FAvaRundownManagedInstance>& Instance : Instances)
 	{
-		UE::AvaMedia::Rundown::Private::UnregisterRundownForEvents(Instances, RegisteredRundown.Get());
-		RegisteredRundown = nullptr;
+		if (Instance && InOther.Instances.Contains(Instance))
+		{
+			Instance.Reset();
+		}
 	}
 
 	RegisteredRundown = MoveTemp(InOther.RegisteredRundown);
+	RegisteredPageId = MoveTemp(InOther.RegisteredPageId);
 	Instances = MoveTemp(InOther.Instances);
 	return *this;
 }
diff --git a/Engine/Plugins/VirtualProduction/Avalanche/Source/AvalancheMedia/Private/Rundown/AvaRundownPage.cpp b/Engine/Plugins/VirtualProduction/Avalanche/Source/AvalancheMedia/Private/Rundown/AvaRundownPage.cpp
--- a/Engine/Plugins/VirtualProduction/Avalanche/Source/AvalancheMedia/Private/Rundown/AvaRundownPage.cpp
+++ b/Engine/Plugins/VirtualProduction/Avalanche/Source/AvalancheMedia/Private/Rundown/AvaRundownPage.cpp
@@ -23,6 +23,7 @@
 
 FAvaRundownPage FAvaRundownPage::NullPage;
 const int32 FAvaRundownPage::InvalidPageId = -1;
+static_assert(FAvaRundownPage::InvalidPageId == static_cast<int32>(UE::AvaMedia::Rundown::InvalidPageId), "Mismatching InvalidPageId definitions.");
 
 FAvaRundownPage::FAvaRundownPage(int32 InPageId, int32 InTemplateId)
 	: PageId(InPageId)
diff --git a/Engine/Plugins/VirtualProduction/Avalanche/Source/AvalancheMedia/Public/Rundown/AvaRundown.h b/Engine/Plugins/VirtualProduction/Avalanche/Source/AvalancheMedia/Public/Rundown/AvaRundown.h
--- a/Engine/Plugins/VirtualProduction/Avalanche/Source/AvalancheMedia/Public/Rundown/AvaRundown.h
+++ b/Engine/Plugins/VirtualProduction/Avalanche/Source/AvalancheMedia/Public/Rundown/AvaRundown.h
@@ -634,17 +634,17 @@ public:
 	void OnParentWordBeginTearDown();
 
 	/**
-	 * Push an RC event on the given RC item(s) (controllers or entities) for the given page.
+	 * Push an RC event on the given controllers for the given page.
 	 * The propagation to channels in determined by rundown settings.
 	 * 
-	 * @param InRcItemIds List of rc item ids that should be triggered.
+	 * @param InControllerIds List of rc controller ids that should be triggered.
 	 * @param InPageId Page id these events apply to.
 	 * @param bInPushToProgram if true, the command will be pushed to page's program channel (if currently playing).
 	 * @param bInPushToPreview if true, the command will be pushed to preview channel(s).
 	 * @param InPreviewChannelName Allow to specify a specific preview channel to push to. If left none, will push to all preview channel this page is playing in.
 	 * @return true if the event command was pushed to the corresponding playback instance.
 	 */
-	bool PushRuntimeRemoteControlEvents(TConstArrayView<FGuid> InRcItemIds, int32 InPageId, bool bInPushToProgram, bool bInPushToPreview, FName InPreviewChannelName = NAME_None);
+	bool PushRuntimeRemoteControlEvents(TConstArrayView<FGuid> InControllerIds, int32 InPageId, bool bInPushToProgram, bool bInPushToPreview, FName InPreviewChannelName = NAME_None);
 	
 	/**
 	 * Push the page's RC values to the runtime playback instances.
diff --git a/Engine/Plugins/VirtualProduction/Avalanche/Source/AvalancheMedia/Public/Rundown/AvaRundownDefines.h b/Engine/Plugins/VirtualProduction/Avalanche/Source/AvalancheMedia/Public/Rundown/AvaRundownDefines.h
--- a/Engine/Plugins/VirtualProduction/Avalanche/Source/AvalancheMedia/Public/Rundown/AvaRundownDefines.h
+++ b/Engine/Plugins/VirtualProduction/Avalanche/Source/AvalancheMedia/Public/Rundown/AvaRundownDefines.h
@@ -11,3 +11,8 @@ enum class EAvaRundownPagePlayType : uint8
 	PreviewFromStart,
 	PreviewFromFrame
 };
+
+namespace UE::AvaMedia::Rundown
+{
+	enum { InvalidPageId = -1};
+}
\ No newline at end of file
diff --git a/Engine/Plugins/VirtualProduction/Avalanche/Source/AvalancheMediaEditor/Private/Rundown/Factories/Filters/IAvaRundownFilterSuggestionFactory.h b/Engine/Plugins/VirtualProduction/Avalanche/Source/AvalancheMediaEditor/Private/Rundown/Factories/Filters/IAvaRundownFilterSuggestionFactory.h
--- a/Engine/Plugins/VirtualProduction/Avalanche/Source/AvalancheMediaEditor/Private/Rundown/Factories/Filters/IAvaRundownFilterSuggestionFactory.h
+++ b/Engine/Plugins/VirtualProduction/Avalanche/Source/AvalancheMediaEditor/Private/Rundown/Factories/Filters/IAvaRundownFilterSuggestionFactory.h
@@ -3,7 +3,7 @@
 #pragma once
 
 #include "HAL/Platform.h"
-#include "Rundown/AvaRundownEditorDefines.h"
+#include "Rundown/AvaRundownDefines.h"
 #include "Templates/SharedPointer.h"
 #include "Templates/UnrealTypeTraits.h"
 
@@ -24,7 +24,7 @@ struct FAvaRundownFilterSuggestionPayload
 	const FString FilterValue;
 
 	/** Current item to check */
-	int32 ItemPageId = UE::AvaRundown::InvalidPageId;
+	int32 ItemPageId = UE::AvaMedia::Rundown::InvalidPageId;
 
 	/** Current item rundown */
 	const UAvaRundown* Rundown = nullptr;
diff --git a/Engine/Plugins/VirtualProduction/Avalanche/Source/AvalancheMediaEditor/Private/Rundown/Filters/AvaRundownPageFilterExpressionContext.h b/Engine/Plugins/VirtualProduction/Avalanche/Source/AvalancheMediaEditor/Private/Rundown/Filters/AvaRundownPageFilterExpressionContext.h
--- a/Engine/Plugins/VirtualProduction/Avalanche/Source/AvalancheMediaEditor/Private/Rundown/Filters/AvaRundownPageFilterExpressionContext.h
+++ b/Engine/Plugins/VirtualProduction/Avalanche/Source/AvalancheMediaEditor/Private/Rundown/Filters/AvaRundownPageFilterExpressionContext.h
@@ -4,7 +4,7 @@
 
 #include "Containers/UnrealString.h"
 #include "Misc/TextFilterExpressionEvaluator.h"
-#include "Rundown/AvaRundownEditorDefines.h"
+#include "Rundown/AvaRundownDefines.h"
 #include "Rundown/Factories/Filters/AvaRundownFactoriesUtils.h"
 
 class FName;
@@ -17,7 +17,7 @@ class FAvaRundownPageFilterExpressionContext : public ITextFilterExpressionConte
 {
 public:
 	FAvaRundownPageFilterExpressionContext()
-		: ItemPageId(UE::AvaRundown::InvalidPageId)
+		: ItemPageId(UE::AvaMedia::Rundown::InvalidPageId)
 		, ItemRundown(nullptr)
 		, RundownSearchListType(EAvaRundownSearchListType::None)
 	{}
diff --git a/Engine/Plugins/VirtualProduction/Avalanche/Source/AvalancheMediaEditor/Public/Rundown/AvaRundownEditorDefines.h b/Engine/Plugins/VirtualProduction/Avalanche/Source/AvalancheMediaEditor/Public/Rundown/AvaRundownEditorDefines.h
--- a/Engine/Plugins/VirtualProduction/Avalanche/Source/AvalancheMediaEditor/Public/Rundown/AvaRundownEditorDefines.h
+++ b/Engine/Plugins/VirtualProduction/Avalanche/Source/AvalancheMediaEditor/Public/Rundown/AvaRundownEditorDefines.h
@@ -51,8 +51,6 @@ enum class EAvaRundownPageSet : uint8
 
 namespace UE::AvaRundown
 {
-	enum {InvalidPageId	= -1};
-	
 	struct AVALANCHEMEDIAEDITOR_API FEditorMetrics
 	{
 		static constexpr float ColumnLeftOffset = 5.f;
