Motion Design: Minor comment fix

diff --git a/Engine/Plugins/VirtualProduction/Avalanche/Source/AvalancheTransition/Public/AvaTransitionLibrary.h b/Engine/Plugins/VirtualProduction/Avalanche/Source/AvalancheTransition/Public/AvaTransitionLibrary.h
--- a/Engine/Plugins/VirtualProduction/Avalanche/Source/AvalancheTransition/Public/AvaTransitionLibrary.h
+++ b/Engine/Plugins/VirtualProduction/Avalanche/Source/AvalancheTransition/Public/AvaTransitionLibrary.h
@@ -27,7 +27,7 @@ public:
 	 * @param InLayerComparisonType the type of layer to filter.
 	 *      "Same": the same layer as this context's.
 	 *      "Different": a layer different from this context's.
-	 *      "MatchingTag": a layer specified by "InSpecificLayers".
+	 *      "Specific Layers": layers specified by "InSpecificLayers".
 	 *      "Any": any layer found.
 	 * @param InSpecificLayers the specific layers to filter. InLayerComparisonType should be set to "Specific Layers" for this to be considered.
 	 * @param InTransitionFilterType the transition type to filter (whether it should only consider Ins or Outs, or any type is allowed)
