Motion Design: Guide loading issue redux.

diff --git a/Engine/Plugins/VirtualProduction/Avalanche/Source/AvalancheLevelViewport/Private/SAvaLevelViewport.cpp b/Engine/Plugins/VirtualProduction/Avalanche/Source/AvalancheLevelViewport/Private/SAvaLevelViewport.cpp
--- a/Engine/Plugins/VirtualProduction/Avalanche/Source/AvalancheLevelViewport/Private/SAvaLevelViewport.cpp
+++ b/Engine/Plugins/VirtualProduction/Avalanche/Source/AvalancheLevelViewport/Private/SAvaLevelViewport.cpp
@@ -401,7 +401,7 @@ TSharedPtr<FAvaLevelViewportClient> SAvaLevelViewport::GetAvaLevelViewportClient
 	return nullptr;
 }
 
-TSharedPtr<SAvaLevelViewportGuide> SAvaLevelViewport::AddGuide(EOrientation InOrientation, float InOffsetFraction)
+TSharedPtr<SAvaLevelViewportGuide> SAvaLevelViewport::AddGuide(EOrientation InOrientation, float InOffsetFraction, bool bInSaveGuides)
 {
 	if (!GuideCanvas.IsValid())
 	{
@@ -424,7 +424,10 @@ TSharedPtr<SAvaLevelViewportGuide> SAvaLevelViewport::AddGuide(EOrientation InOr
 
 	Guides.Add(NewGuide);
 
-	SaveGuides();
+	if (bInSaveGuides)
+	{
+		SaveGuides();
+	}
 
 	GuideCanvas->AddSlot()
 		.Position(NewGuide.Get(), &SAvaLevelViewportGuide::GetPosition)
@@ -495,13 +498,15 @@ void SAvaLevelViewport::LoadGuides()
 		return;
 	}
 
+	constexpr bool bSaveGuides = false;
+
 	for (const FAvaViewportGuideInfo& GuideInfo : Data->GuideData)
 	{
-		TSharedPtr<SAvaLevelViewportGuide> NewGuide = AddGuide(GuideInfo.Orientation, GuideInfo.OffsetFraction);
+		TSharedPtr<SAvaLevelViewportGuide> NewGuide = AddGuide(GuideInfo.Orientation, GuideInfo.OffsetFraction, bSaveGuides);
 
 		if (NewGuide.IsValid())
 		{
-			NewGuide->SetState(GuideInfo.State, /* Save Guides */ false);
+			NewGuide->SetState(GuideInfo.State, bSaveGuides);
 		}
 	}
 }
diff --git a/Engine/Plugins/VirtualProduction/Avalanche/Source/AvalancheLevelViewport/Private/SAvaLevelViewport.h b/Engine/Plugins/VirtualProduction/Avalanche/Source/AvalancheLevelViewport/Private/SAvaLevelViewport.h
--- a/Engine/Plugins/VirtualProduction/Avalanche/Source/AvalancheLevelViewport/Private/SAvaLevelViewport.h
+++ b/Engine/Plugins/VirtualProduction/Avalanche/Source/AvalancheLevelViewport/Private/SAvaLevelViewport.h
@@ -72,7 +72,7 @@ public:
 	TSharedPtr<SAvaLevelViewportSnapIndicators> GetSnapIndicators() const { return SnapIndicators; }
 
 	TConstArrayView<TSharedPtr<SAvaLevelViewportGuide>> GetGuides() const { return Guides; }
-	TSharedPtr<SAvaLevelViewportGuide> AddGuide(EOrientation InOrientation, float InOffsetFraction); /** 0 = left/top. 1 = right/bottom/ */
+	TSharedPtr<SAvaLevelViewportGuide> AddGuide(EOrientation InOrientation, float InOffsetFraction, bool bInSaveGuides = true); /** 0 = left/top. 1 = right/bottom/ */
 	bool RemoveGuide(TSharedPtr<SAvaLevelViewportGuide> InGuidetoRemove);
 	void SaveGuides();
 
