Motion Design: renamed TransitionFilterType  to TransitionTypeFilter for clarity + fixed display name disparity in comment.

diff --git a/Engine/Plugins/VirtualProduction/Avalanche/Source/AvalancheTransition/Private/AvaTransitionLibrary.cpp b/Engine/Plugins/VirtualProduction/Avalanche/Source/AvalancheTransition/Private/AvaTransitionLibrary.cpp
--- a/Engine/Plugins/VirtualProduction/Avalanche/Source/AvalancheTransition/Private/AvaTransitionLibrary.cpp
+++ b/Engine/Plugins/VirtualProduction/Avalanche/Source/AvalancheTransition/Private/AvaTransitionLibrary.cpp
@@ -17,12 +17,12 @@
 
 namespace UE::AvaTransition::Private
 {
-	EAvaTransitionType GetTransitionTypeFromFilter(EAvaTransitionFilterType InFilterType)
+	EAvaTransitionType GetTransitionTypeFromFilter(EAvaTransitionTypeFilter InFilterType)
 	{
 		switch (InFilterType)
 		{
-		case EAvaTransitionFilterType::In:  return EAvaTransitionType::In;
-		case EAvaTransitionFilterType::Out: return EAvaTransitionType::Out;
+		case EAvaTransitionTypeFilter::In:  return EAvaTransitionType::In;
+		case EAvaTransitionTypeFilter::Out: return EAvaTransitionType::Out;
 		}
 		return EAvaTransitionType::None;
 	}
@@ -131,7 +131,7 @@ const FAvaTransitionContext* UAvaTransitionLibrary::GetTransitionContext(UObject
 bool UAvaTransitionLibrary::IsTransitionActiveInLayers(UObject* InWorldContextObject
 	, EAvaTransitionLayerCompareType InLayerComparisonType
 	, const FAvaTagHandleContainer& InSpecificLayers
-	, EAvaTransitionFilterType InTransitionFilterType)
+	, EAvaTransitionTypeFilter InTransitionTypeFilter)
 {
 	using namespace UE::AvaTransition;
 
@@ -169,7 +169,7 @@ bool UAvaTransitionLibrary::IsTransitionActiveInLayers(UObject* InWorldContextOb
 		return false;
 	}
 
-	const EAvaTransitionType TransitionType = Private::GetTransitionTypeFromFilter(InTransitionFilterType);
+	const EAvaTransitionType TransitionType = Private::GetTransitionTypeFromFilter(InTransitionTypeFilter);
 
 	bool bHasMatchingInstance = BehaviorInstances.ContainsByPredicate(
 		[TransitionScene, TransitionType](const FAvaTransitionBehaviorInstance* InInstance)
@@ -200,7 +200,7 @@ bool UAvaTransitionLibrary::IsTransitionActiveInLayer(UObject* InWorldContextObj
 	return IsTransitionActiveInLayers(InWorldContextObject
 		, InLayerComparisonType
 		, InSpecificLayers
-		, EAvaTransitionFilterType::Any);
+		, EAvaTransitionTypeFilter::Any);
 }
 
 EAvaTransitionType UAvaTransitionLibrary::GetTransitionType(UObject* InContextObject)
diff --git a/Engine/Plugins/VirtualProduction/Avalanche/Source/AvalancheTransition/Public/AvaTransitionEnums.h b/Engine/Plugins/VirtualProduction/Avalanche/Source/AvalancheTransition/Public/AvaTransitionEnums.h
--- a/Engine/Plugins/VirtualProduction/Avalanche/Source/AvalancheTransition/Public/AvaTransitionEnums.h
+++ b/Engine/Plugins/VirtualProduction/Avalanche/Source/AvalancheTransition/Public/AvaTransitionEnums.h
@@ -28,7 +28,7 @@ enum class EAvaTransitionType : uint8
 ENUM_CLASS_FLAGS(EAvaTransitionType);
 
 UENUM(BlueprintType, DisplayName="Motion Design Transition Filter Type")
-enum class EAvaTransitionFilterType : uint8
+enum class EAvaTransitionTypeFilter : uint8
 {
 	In,
 	Out,
diff --git a/Engine/Plugins/VirtualProduction/Avalanche/Source/AvalancheTransition/Public/AvaTransitionLibrary.h b/Engine/Plugins/VirtualProduction/Avalanche/Source/AvalancheTransition/Public/AvaTransitionLibrary.h
--- a/Engine/Plugins/VirtualProduction/Avalanche/Source/AvalancheTransition/Public/AvaTransitionLibrary.h
+++ b/Engine/Plugins/VirtualProduction/Avalanche/Source/AvalancheTransition/Public/AvaTransitionLibrary.h
@@ -26,17 +26,17 @@ public:
 	 * @param InContextObject the context to retrieve the active transition relating to the context.
 	 * @param InLayerComparisonType the type of layer to filter.
 	 *      "Same": the same layer as this context's.
-	 *      "Different": a layer different from this context's.
+	 *      "Other": a layer different from this context's.
 	 *      "Specific Layers": layers specified by "InSpecificLayers".
 	 *      "Any": any layer found.
 	 * @param InSpecificLayers the specific layers to filter. InLayerComparisonType should be set to "Specific Layers" for this to be considered.
-	 * @param InTransitionFilterType the transition type to filter (whether it should only consider Ins or Outs, or any type is allowed)
+	 * @param InTransitionTypeFilter the transition type to filter (whether it should only consider Ins or Outs, or any type is allowed)
 	 */
 	UFUNCTION(BlueprintPure, Category="Transition Logic", meta=(DefaultToSelf="InContextObject"))
 	static UE_API bool IsTransitionActiveInLayers(UObject* InContextObject
 		, EAvaTransitionLayerCompareType InLayerComparisonType
 		, const FAvaTagHandleContainer& InSpecificLayers
-		, EAvaTransitionFilterType InTransitionFilterType = EAvaTransitionFilterType::Any);
+		, EAvaTransitionTypeFilter InTransitionTypeFilter = EAvaTransitionTypeFilter::Any);
 
 	UE_DEPRECATED(5.7, "Use IsTransitionActiveInLayers instead")
 	UFUNCTION(BlueprintPure, Category="Transition Logic", meta=(DefaultToSelf="InContextObject", DeprecatedFunction, DeprecationMessage="Use \"IsTransitionActiveInLayers\" instead"))
