[Motion Design] Fix FAvaRundownManagedInstanceHandles move operator to ensure the instances are properly unregistered.

I was mistakeny assuming the moved-from object destrutor would work, but the Instances array is getting reset somewhere else and the remaining instances where not unregistered. So, I am explicitly unregistering in the move operator, and setting everything up for the moved-from destructor too, just in case. After testing in the rundown editor with multiple page instances of different templates,  managed instances properly get unregistered when selecting different pages (of different templates).

diff --git a/Engine/Plugins/VirtualProduction/Avalanche/Source/AvalancheMedia/Private/Rundown/AvaRundownManagedInstanceHandle.cpp b/Engine/Plugins/VirtualProduction/Avalanche/Source/AvalancheMedia/Private/Rundown/AvaRundownManagedInstanceHandle.cpp
--- a/Engine/Plugins/VirtualProduction/Avalanche/Source/AvalancheMedia/Private/Rundown/AvaRundownManagedInstanceHandle.cpp
+++ b/Engine/Plugins/VirtualProduction/Avalanche/Source/AvalancheMedia/Private/Rundown/AvaRundownManagedInstanceHandle.cpp
@@ -12,6 +12,9 @@ namespace UE::AvaMedia::Rundown::Private
 	{
 		if (InInstance && InInstance->GetRundownForEvents() == InRundown && InInstance->GetActivePageIdForEvents() == InPageId)
 		{
+			UE_LOG(LogAvaRundown, Verbose, TEXT("Unregister Managed Instance %s page id %d from RC Events routing."),
+				*InInstance->GetSourceAssetPath().ToString(), InPageId);
+
 			InInstance->SetCurrentRundownPageForEvents(nullptr, InvalidPageId);
 		}
 	}
@@ -51,18 +54,29 @@ FAvaRundownManagedInstanceHandles& FAvaRundownManagedInstanceHandles::operator=(
 		return *this;
 	}
 
-	// Remove any instances from the incoming set of instance.
-	// The remaining instances will be unregistered by the destructor.
+	// Remove any instances that are also in the incoming set of instances to make sure
+	// they don't get unregistered by the moved-from object destructor.
 	for (TSharedPtr<FAvaRundownManagedInstance>& Instance : Instances)
 	{
-		if (Instance && InOther.Instances.Contains(Instance))
+		if (Instance)
 		{
-			Instance.Reset();
+			if (InOther.Instances.Contains(Instance))
+			{
+				Instance.Reset();
+			}
+			else
+			{
+				// Unregister instances from event routing immediately.
+				UE::AvaMedia::Rundown::Private::UnregisterRundownForEvents(Instance, RegisteredRundown.Get(), RegisteredPageId);
+			}
 		}
 	}
 
-	RegisteredRundown = MoveTemp(InOther.RegisteredRundown);
-	RegisteredPageId = MoveTemp(InOther.RegisteredPageId);
+	// Swap the registered rundown and page id so the moved-from object destructor unregisters
+	// the remaining unused instances from the previous page id.
+	Swap(RegisteredRundown,InOther.RegisteredRundown);
+	Swap(RegisteredPageId, InOther.RegisteredPageId);
+
 	Instances = MoveTemp(InOther.Instances);
 	return *this;
 }
