Motion Design Scene State: added a spawner ticker task that spawns an actor of a user-definable class, and queues it to a given ticker actor.
This allows the task to be optionally used with data link to set the actor to spawn with the data link output data mapped via the scene state property binding system.

diff --git a/Engine/Plugins/Experimental/AvalancheSceneState/Source/AvalancheSceneState/AvalancheSceneState.Build.cs b/Engine/Plugins/Experimental/AvalancheSceneState/Source/AvalancheSceneState/AvalancheSceneState.Build.cs
--- a/Engine/Plugins/Experimental/AvalancheSceneState/Source/AvalancheSceneState/AvalancheSceneState.Build.cs
+++ b/Engine/Plugins/Experimental/AvalancheSceneState/Source/AvalancheSceneState/AvalancheSceneState.Build.cs
@@ -18,6 +18,7 @@ public class AvalancheSceneState : ModuleRules
 				"CoreUObject",
 				"SceneState",
 				"SceneStateBinding",
+				"SceneStateTasks",
 				"SceneStateGameplay",
 			}
 		);
diff --git a/Engine/Plugins/Experimental/AvalancheSceneState/Source/AvalancheSceneState/Private/Ticker/AvaSceneStateSpawnerTickerTask.cpp b/Engine/Plugins/Experimental/AvalancheSceneState/Source/AvalancheSceneState/Private/Ticker/AvaSceneStateSpawnerTickerTask.cpp
new file mode 100644
--- /dev/null
+++ b/Engine/Plugins/Experimental/AvalancheSceneState/Source/AvalancheSceneState/Private/Ticker/AvaSceneStateSpawnerTickerTask.cpp
@@ -0,0 +1,55 @@
+// Copyright Epic Games, Inc. All Rights Reserved.
+
+#include "Ticker/AvaSceneStateSpawnerTickerTask.h"
+#include "Framework/Ticker/AvaTickerComponent.h"
+#include "GameFramework/Actor.h"
+
+#define LOCTEXT_NAMESPACE "AvaSceneStateSpawnerTickerTask"
+
+#if WITH_EDITOR
+const UScriptStruct* FAvaSceneStateSpawnerTickerTask::OnGetTaskInstanceType() const
+{
+	return FInstanceDataType::StaticStruct();
+}
+#endif
+
+bool FAvaSceneStateSpawnerTickerTask::ShouldSpawnActor(FStructView InTaskInstance, FText& OutErrorMessage) const
+{
+	FInstanceDataType& Instance = InTaskInstance.Get<FInstanceDataType>();
+	if (!Instance.Ticker)
+	{
+		OutErrorMessage = LOCTEXT("Error_InvalidTickerActor", "Specified ticker actor is invalid.");
+		return false;
+	}
+
+	if (!Instance.Ticker->FindComponentByClass<UAvaTickerComponent>())
+	{
+		OutErrorMessage = LOCTEXT("Error_NoTickerComponent", "Specified actor does not contain a ticker component.");
+		return false;
+	}
+
+	return true;
+}
+
+void FAvaSceneStateSpawnerTickerTask::OnActorSpawned(AActor* InActorChecked, FStructView InTaskInstance) const
+{
+	FInstanceDataType& Instance = InTaskInstance.Get<FInstanceDataType>();
+
+	// Ensure here since ShouldSpawnActor should've already validated that the actor is valid
+	if (!ensureMsgf(Instance.Ticker, TEXT("Valid ticker actor unexpectedly not found in task instance")))
+	{
+		return;
+	}
+
+	UAvaTickerComponent* TickerComponent = Instance.Ticker->FindComponentByClass<UAvaTickerComponent>();
+
+	// Ensure here since ShouldSpawnActor should've already validated that the actor has a ticker component
+	if (!ensureMsgf(TickerComponent, TEXT("Ticker actor %s expected to have a ticker component"), *Instance.Ticker->GetActorNameOrLabel()))
+	{
+		return;
+	}
+
+	TickerComponent->QueueActor(InActorChecked);
+}
+
+#undef LOCTEXT_NAMESPACE
diff --git a/Engine/Plugins/Experimental/AvalancheSceneState/Source/AvalancheSceneState/Public/Ticker/AvaSceneStateSpawnerTickerTask.h b/Engine/Plugins/Experimental/AvalancheSceneState/Source/AvalancheSceneState/Public/Ticker/AvaSceneStateSpawnerTickerTask.h
new file mode 100644
--- /dev/null
+++ b/Engine/Plugins/Experimental/AvalancheSceneState/Source/AvalancheSceneState/Public/Ticker/AvaSceneStateSpawnerTickerTask.h
@@ -0,0 +1,44 @@
+// Copyright Epic Games, Inc. All Rights Reserved.
+
+#pragma once
+
+#include "SceneStateSpawnActorTask.h"
+#include "AvaSceneStateSpawnerTickerTask.generated.h"
+
+#define UE_API AVALANCHESCENESTATE_API
+
+USTRUCT()
+struct FAvaSceneStateTickerTaskInstance : public FSceneStateSpawnActorTaskInstance
+{
+	GENERATED_BODY()
+
+	/** The actor to find the ticker component for */
+	UPROPERTY(EditAnywhere, Category="Ticker")
+	TObjectPtr<AActor> Ticker;
+};
+
+/**
+ * Spawns an actor and adds it to the Ticker.
+ * Note: if the Ticker actor is invalid or does not contain a Ticker Component, the actor will not spawn.
+ */
+USTRUCT(DisplayName="Spawn Actor to Ticker", Category="Motion Design")
+struct FAvaSceneStateSpawnerTickerTask : public FSceneStateSpawnActorTask
+{
+	GENERATED_BODY()
+
+	using FInstanceDataType = FAvaSceneStateTickerTaskInstance;
+
+protected:
+	//~ Begin FSceneStateTask
+#if WITH_EDITOR
+	UE_API virtual const UScriptStruct* OnGetTaskInstanceType() const override;
+#endif
+	//~ End FSceneStateTask
+
+	//~ Begin FSceneStateSpawnActorTask
+	UE_API virtual bool ShouldSpawnActor(FStructView InTaskInstance, FText& OutErrorMessage) const override;
+	UE_API virtual void OnActorSpawned(AActor* InActorChecked, FStructView InTaskInstance) const override;
+	//~ End FSceneStateSpawnActorTask
+};
+
+#undef UE_API
diff --git a/Engine/Plugins/VirtualProduction/Avalanche/Source/Avalanche/Internal/Framework/Ticker/AvaTickerActor.h b/Engine/Plugins/VirtualProduction/Avalanche/Source/Avalanche/Internal/Framework/Ticker/AvaTickerActor.h
deleted file mode 100644
diff --git a/Engine/Plugins/VirtualProduction/Avalanche/Source/Avalanche/Internal/Framework/Ticker/AvaTickerComponent.h b/Engine/Plugins/VirtualProduction/Avalanche/Source/Avalanche/Internal/Framework/Ticker/AvaTickerComponent.h
deleted file mode 100644
diff --git a/Engine/Plugins/VirtualProduction/Avalanche/Source/Avalanche/Public/Framework/Ticker/AvaTickerActor.h b/Engine/Plugins/VirtualProduction/Avalanche/Source/Avalanche/Public/Framework/Ticker/AvaTickerActor.h
new file mode 100644
--- /dev/null
+++ b/Engine/Plugins/VirtualProduction/Avalanche/Source/Avalanche/Public/Framework/Ticker/AvaTickerActor.h
@@ -0,0 +1,32 @@
+// Copyright Epic Games, Inc. All Rights Reserved.
+
+#pragma once
+
+#include "GameFramework/Actor.h"
+#include "AvaTickerActor.generated.h"
+
+class UAvaTickerComponent;
+
+UCLASS(MinimalAPI, DisplayName="Motion Design Ticker Actor")
+class AAvaTickerActor : public AActor
+{
+	GENERATED_BODY()
+
+public:
+	AAvaTickerActor();
+
+	UAvaTickerComponent* GetTickerComponent() const
+	{
+		return TickerComponent;
+	}
+
+	//~ Begin AActor
+#if WITH_EDITOR
+	virtual FString GetDefaultActorLabel() const override;
+#endif
+	//~ End AActor
+
+protected:
+	UPROPERTY(VisibleAnywhere, BlueprintReadOnly, Category="Ticker")
+	TObjectPtr<UAvaTickerComponent> TickerComponent;
+};
diff --git a/Engine/Plugins/VirtualProduction/Avalanche/Source/Avalanche/Public/Framework/Ticker/AvaTickerComponent.h b/Engine/Plugins/VirtualProduction/Avalanche/Source/Avalanche/Public/Framework/Ticker/AvaTickerComponent.h
new file mode 100644
--- /dev/null
+++ b/Engine/Plugins/VirtualProduction/Avalanche/Source/Avalanche/Public/Framework/Ticker/AvaTickerComponent.h
@@ -0,0 +1,213 @@
+// Copyright Epic Games, Inc. All Rights Reserved.
+
+#pragma once
+
+#include "Components/PrimitiveComponent.h"
+#include "AvaTickerComponent.generated.h"
+
+#define UE_API AVALANCHE_API
+
+UENUM()
+enum class EAvaTickerQueueLimitType : uint8
+{
+	None UMETA(Hidden),
+
+	/** No new elements will be queued if queue is full*/
+	DisableQueueing,
+
+	/** New elements will be queued by discarding the oldest (front of queue) */
+	DiscardOldest,
+};
+
+USTRUCT()
+struct FAvaTickerElement
+{
+	GENERATED_BODY()
+
+	FAvaTickerElement() = default;
+
+	UE_API explicit FAvaTickerElement(AActor* InActor);
+
+	UE_API bool operator==(const FAvaTickerElement& InOther) const;
+
+	/** Whether the element is a valid element */
+	bool IsValid() const;
+
+	/** Sets the game and editor visibility of the element */
+	void SetVisibility(bool bInVisible);
+
+	/** Retrieves the world space location of the element */
+	FVector GetElementLocation() const;
+
+	/** Moves the given element to the given world space location */
+	void SetLocation(const FVector& InLocation);
+
+	/** Moves the given element by the given displacement vector */
+	void Displace(const FVector& InDisplacement);
+
+	/** Destroys the element */
+	void Destroy();
+
+	/** Gets the bounds for the element */
+	void GetBounds(FVector& OutOrigin, FVector& OutExtents) const;
+
+	/** Actor mapped to this ticker element */
+	UPROPERTY()
+	TObjectPtr<AActor> Actor;
+};
+
+UCLASS(MinimalAPI, Blueprintable, DisplayName="Motion Design Ticker Component")
+class UAvaTickerComponent : public UPrimitiveComponent
+{
+	GENERATED_BODY()
+
+public:
+	UE_API UAvaTickerComponent();
+
+	/** Returns whether the queue can be added to */
+	UFUNCTION(BlueprintPure, Category="Ticker")
+	bool CanQueueElements() const;
+
+	/**
+	 * Adds the given actor to the queue
+	 * @param InActor the actor to queue
+	 * @param bInDestroyOnFailure whether to destroy the actor if it failed to be queued
+	 * @return true if the actor was successfully queued
+	 */
+	UFUNCTION(BlueprintCallable, Category="Ticker")
+	UE_API bool QueueActor(AActor* InActor, bool bInDestroyOnFailure = true);
+
+	UFUNCTION(BlueprintPure, Category="Ticker")
+	FVector GetStartLocation() const
+	{
+		return StartLocation;
+	}
+
+	UFUNCTION(BlueprintCallable, Category="Ticker")
+	UE_API void SetStartLocation(const FVector& InStartLocation);
+
+	UFUNCTION(BlueprintPure, Category="Ticker")
+	double GetDestroyDistance() const
+	{
+		return DestroyDistance;
+	}
+
+	UFUNCTION(BlueprintCallable, Category="Ticker")
+	UE_API void SetDestroyDistance(double InDestroyDistance);
+
+	UFUNCTION(BlueprintPure, Category="Ticker")
+	FVector GetVelocity() const
+	{
+		return Velocity;
+	}
+
+	UFUNCTION(BlueprintCallable, Category="Ticker")
+	UE_API void SetVelocity(const FVector& InVelocity);
+
+	UFUNCTION(BlueprintPure, Category="Ticker")
+	double GetPadding() const
+	{
+		return Padding;
+	}
+
+	UFUNCTION(BlueprintCallable, Category="Ticker")
+	UE_API void SetPadding(double InPadding);
+
+	UFUNCTION(BlueprintPure, Category="Ticker")
+	bool ShouldLimitQueue() const
+	{
+		return bLimitQueue;
+	}
+
+	UFUNCTION(BlueprintCallable, Category="Ticker")
+	UE_API void SetLimitQueue(bool bInLimitQueue);
+
+	UFUNCTION(BlueprintPure, Category="Ticker")
+	int32 GetQueueLimitCount() const
+	{
+		return QueueLimitCount;
+	}
+
+	UFUNCTION(BlueprintCallable, Category="Ticker")
+	UE_API void SetQueueLimitCount(int32 InQueueLimitCount);
+
+	UFUNCTION(BlueprintPure, Category="Ticker")
+	EAvaTickerQueueLimitType GetQueueLimitType() const
+	{
+		return QueueLimitType;
+	}
+
+	UFUNCTION(BlueprintCallable, Category="Ticker")
+	UE_API void SetQueueLimitType(EAvaTickerQueueLimitType InQueueLimitType);
+
+	//~ Begin UPrimitiveComponent
+	UE_API virtual FPrimitiveSceneProxy* CreateSceneProxy() override;
+	//~ End UPrimitiveComponent
+
+	//~ Begin USceneComponent
+	UE_API virtual FBoxSphereBounds CalcBounds(const FTransform& InLocalToWorld) const override;
+	//~ End USceneComponent
+
+	//~ Begin UActorComponent
+	UE_API virtual void TickComponent(float InDeltaTime, ELevelTick InTickType, FActorComponentTickFunction* InThisTickFunction) override;
+	//~ End UActorComponent
+
+protected:
+	bool IsQueueFull() const;
+
+	/** Adds the given element to the queue */
+	bool QueueElement(FAvaTickerElement&& InElement, bool bInDestroyOnFailure);
+
+	/** Retrieves the start location transformed to World Space */
+	FVector GetWorldStartLocation() const;
+
+	/** Retrieves the destroy displacement transformed to World Space */
+	FVector GetWorldDestroyDisplacement() const;
+
+	/** Retrieves the velocity transformed to World Space */
+	FVector GetWorldVelocity() const;
+
+	/** Removes the first element from the queue and activates it if the last active element isn't on the way */
+	void TryDequeueElement();
+
+	/** Updates the active elements to move with the given velocity */
+	void UpdateActiveElements(float InDeltaTime);
+
+	/** Elements being moved by this ticker */
+	UPROPERTY(Transient, DuplicateTransient, TextExportTransient)
+	TArray<FAvaTickerElement> ActiveElements;
+
+	/** Elements currently hidden, awaiting to become available */
+	UPROPERTY(Transient, DuplicateTransient, TextExportTransient)
+	TArray<FAvaTickerElement> QueuedElements;
+
+	/** The start location (in local space) where active elements start at */
+	UPROPERTY(EditAnywhere, Category="Ticker", meta=(AllowPrivateAccess="true"))
+	FVector StartLocation = FVector(0, 500, 0);
+
+	/** The distance the last part of the active element must cross before it gets destroyed */
+	UPROPERTY(EditAnywhere, Category="Ticker", meta=(AllowPrivateAccess="true"))
+	double DestroyDistance = 1000.0;
+
+	/** The velocity (in local space) to move the active elements at */
+	UPROPERTY(EditAnywhere, Category="Ticker", meta=(AllowPrivateAccess="true"))
+	FVector Velocity = FVector(0, -200, 0);
+
+	/** Padding between the active ticker elements */
+	UPROPERTY(EditAnywhere, Category="Ticker", meta=(AllowPrivateAccess="true"))
+	double Padding = 30.0;
+
+	/** Indicates whether to limit the queue when it gets past a certain number of elements */
+	UPROPERTY(EditAnywhere, Category="Ticker|Queue", meta=(AllowPrivateAccess="true"))
+	bool bLimitQueue = true;
+
+	/** The number of items allowed in the queue */
+	UPROPERTY(EditAnywhere, Category="Ticker|Queue", meta=(EditCondition="bLimitQueue", EditConditionHides, ClampMin=0, AllowPrivateAccess="true"))
+	int32 QueueLimitCount = 2;
+
+	/** The handling method when the queue is full */
+	UPROPERTY(EditAnywhere, Category="Ticker|Queue", meta=(EditCondition="bLimitQueue", EditConditionHides, AllowPrivateAccess="true"))
+	EAvaTickerQueueLimitType QueueLimitType = EAvaTickerQueueLimitType::DisableQueueing;
+};
+
+#undef UE_API
