[Motion Design] Adding a broadcast channel render target resize event to allow broadcast editor to resize the channel preview brushes.

diff --git a/Engine/Plugins/VirtualProduction/Avalanche/Source/AvalancheMedia/Internal/AvaMediaRenderTargetUtils.h b/Engine/Plugins/VirtualProduction/Avalanche/Source/AvalancheMedia/Internal/AvaMediaRenderTargetUtils.h
--- a/Engine/Plugins/VirtualProduction/Avalanche/Source/AvalancheMedia/Internal/AvaMediaRenderTargetUtils.h
+++ b/Engine/Plugins/VirtualProduction/Avalanche/Source/AvalancheMedia/Internal/AvaMediaRenderTargetUtils.h
@@ -24,8 +24,17 @@ namespace UE::AvaMediaRenderTargetUtils
 	 * @return true if the TargetGamma was modified, false otherwise.
 	 */
 	AVALANCHEMEDIA_API bool SetupRenderTargetGamma(UTextureRenderTarget2D* InRenderTarget);
+
+	/**
+	 * Returns the size of the render target.
+	 * This version of the function doesn't ensure if the render target is not valid.
+	 */ 
+	AVALANCHEMEDIA_API FIntPoint GetRenderTargetSizeSafe(const UTextureRenderTarget2D* InRenderTarget);
 	
-	/** Returns the size of the render target. */ 
+	/** 
+	 * Returns the size of the render target.
+	 * An ensure will be fired if the render target is not valid.
+	 */ 
 	AVALANCHEMEDIA_API FIntPoint GetRenderTargetSize(const UTextureRenderTarget2D* InRenderTarget);
 
 	/** Returns true if render target's format is floating point. */
diff --git a/Engine/Plugins/VirtualProduction/Avalanche/Source/AvalancheMedia/Private/AvaMediaRenderTargetUtils.cpp b/Engine/Plugins/VirtualProduction/Avalanche/Source/AvalancheMedia/Private/AvaMediaRenderTargetUtils.cpp
--- a/Engine/Plugins/VirtualProduction/Avalanche/Source/AvalancheMedia/Private/AvaMediaRenderTargetUtils.cpp
+++ b/Engine/Plugins/VirtualProduction/Avalanche/Source/AvalancheMedia/Private/AvaMediaRenderTargetUtils.cpp
@@ -67,7 +67,12 @@ namespace UE::AvaMediaRenderTargetUtils
 			InRenderTarget->ResizeTarget(InSize.X, InSize.Y);
 		}
 	}
-	
+
+	FIntPoint GetRenderTargetSizeSafe(const UTextureRenderTarget2D* InRenderTarget)
+	{
+		return InRenderTarget ? FIntPoint(InRenderTarget->SizeX, InRenderTarget->SizeY) : FIntPoint(0,0);
+	}
+
 	FIntPoint GetRenderTargetSize(const UTextureRenderTarget2D* InRenderTarget)
 	{
 		if (ensure(IsValid(InRenderTarget)))
diff --git a/Engine/Plugins/VirtualProduction/Avalanche/Source/AvalancheMedia/Private/Broadcast/Channel/AvaBroadcastOutputChannel.cpp b/Engine/Plugins/VirtualProduction/Avalanche/Source/AvalancheMedia/Private/Broadcast/Channel/AvaBroadcastOutputChannel.cpp
--- a/Engine/Plugins/VirtualProduction/Avalanche/Source/AvalancheMedia/Private/Broadcast/Channel/AvaBroadcastOutputChannel.cpp
+++ b/Engine/Plugins/VirtualProduction/Avalanche/Source/AvalancheMedia/Private/Broadcast/Channel/AvaBroadcastOutputChannel.cpp
@@ -465,14 +465,18 @@ FAvaBroadcastMediaOutputInfo* FAvaBroadcastOutputChannel::GetMediaOutputInfoMuta
 	return (MediaOutputInfos.IsValidIndex(Index)) ? &MediaOutputInfos[Index] : nullptr; 
 }
 
-inline void UpdatePlaceHolderRenderTarget(TObjectPtr<UTextureRenderTarget2D>& InOutRenderTarget, FName InBaseName, const FIntPoint& InSize, EPixelFormat InFormat, const FLinearColor& InClearColor)
+namespace UE::AvaMedia::BroadcastChannel
 {
-	if (!InOutRenderTarget)
+	inline void UpdatePlaceHolderRenderTarget(TObjectPtr<UTextureRenderTarget2D>& InOutRenderTarget, FName InBaseName, const FIntPoint& InSize, EPixelFormat InFormat, const FLinearColor& InClearColor)
 	{
-		InOutRenderTarget = UE::AvaMediaRenderTargetUtils::CreateDefaultRenderTarget(InBaseName);
-	}
+		using namespace UE::AvaMediaRenderTargetUtils;
+		if (!InOutRenderTarget)
+		{
+			InOutRenderTarget = CreateDefaultRenderTarget(InBaseName);
+		}
 
-	UE::AvaMediaRenderTargetUtils::UpdateRenderTarget(InOutRenderTarget.Get(), InSize, InFormat, InClearColor);
+		UpdateRenderTarget(InOutRenderTarget.Get(), InSize, InFormat, InClearColor);
+	}
 }
 
 bool FAvaBroadcastOutputChannel::IsEditingEnabled() const
@@ -539,8 +543,18 @@ void FAvaBroadcastOutputChannel::UpdatePlaceholderRenderTargets()
 	static const FName PlaceholderRenderTargetName = TEXT("AvaOutputChannel_PlaceHolderRT");
 	static const FName PlaceholderRenderTargetTmpName = TEXT("AvaOutputChannel_TmpPlaceHolderRT");
 
+	using namespace UE::AvaMediaRenderTargetUtils;
+	const FIntPoint PreviousRenderTargetSize = GetRenderTargetSizeSafe(PlaceholderRenderTarget);
+	
+	using namespace UE::AvaMedia::BroadcastChannel;
 	UpdatePlaceHolderRenderTarget(PlaceholderRenderTarget, PlaceholderRenderTargetName, RenderTargetSize, RenderTargetFormat, ClearColor);
 	UpdatePlaceHolderRenderTarget(PlaceholderRenderTargetTmp, PlaceholderRenderTargetTmpName, RenderTargetSize, RenderTargetFormat, ClearColor);
+
+	// Broadcast the render target resize for editor to resize preview brushes.
+	if (PreviousRenderTargetSize != GetRenderTargetSizeSafe(PlaceholderRenderTarget))
+	{
+		OnChannelChanged.Broadcast(*this, EAvaBroadcastChannelChange::RenderTarget);
+	}
 }
 
 void FAvaBroadcastOutputChannel::UpdateOverridePlaceholder()
