Motion Desing: renamed ava scene state to attribute container to avoid confusion with the scene state system.

diff --git a/Engine/Plugins/VirtualProduction/Avalanche/Config/DefaultAvalanche.ini b/Engine/Plugins/VirtualProduction/Avalanche/Config/DefaultAvalanche.ini
--- a/Engine/Plugins/VirtualProduction/Avalanche/Config/DefaultAvalanche.ini
+++ b/Engine/Plugins/VirtualProduction/Avalanche/Config/DefaultAvalanche.ini
@@ -315,5 +315,9 @@
 +PropertyRedirects=(OldName="AvaTransitionLibrary.AreScenesTransitioning.InTransitionNode",NewName="InContextObject")
 +PropertyRedirects=(OldName="AvaTransitionLibrary.GetTransitionTree.InTransitionNode",NewName="InContextObject")
 
+; Renaming 'AvaSceneState' to 'AvaAttributeContainer' to avoid confusion with the Scene State system
++ClassRedirects=(OldName="/Script/Avalanche.AvaSceneState",NewName="/Script/Avalanche.AvaAttributeContainer")
++PropertyRedirects=(OldName="/Script/Avalanche.AvaScene.SceneState",NewName="/Script/Avalanche.AvaScene.AttributeContainer")
+
 [/Script/Avalanche.AvaAnimation]
 DefaultCompletionMode=KeepState
diff --git a/Engine/Plugins/VirtualProduction/Avalanche/Source/Avalanche/Private/AvaAttributeContainer.cpp b/Engine/Plugins/VirtualProduction/Avalanche/Source/Avalanche/Private/AvaAttributeContainer.cpp
new file mode 100644
--- /dev/null
+++ b/Engine/Plugins/VirtualProduction/Avalanche/Source/Avalanche/Private/AvaAttributeContainer.cpp
@@ -0,0 +1,134 @@
+// Copyright Epic Games, Inc. All Rights Reserved.
+
+#include "AvaAttributeContainer.h"
+#include "AvaNameAttribute.h"
+#include "AvaSceneSettings.h"
+#include "AvaTagHandle.h"
+#include "AvaTagHandleContainer.h"
+#include "Tags/AvaTagAttribute.h"
+#include "Tags/AvaTagAttributeBase.h"
+
+void UAvaAttributeContainer::Initialize(UAvaSceneSettings* InSceneSettings)
+{
+	SceneAttributes.Reset();
+	if (InSceneSettings)
+	{
+		SceneAttributes = InSceneSettings->GetSceneAttributes();
+	}
+}
+
+bool UAvaAttributeContainer::AddTagAttribute(const FAvaTagHandle& InTagHandle)
+{
+	if (!InTagHandle.IsValid())
+	{
+		return false;
+	}
+
+	// Return true if already existing
+	if (ContainsTagAttribute(InTagHandle))
+	{
+		return true;
+	}
+
+	UAvaTagAttribute* TagAttribute = NewObject<UAvaTagAttribute>(this, NAME_None, RF_Transient);
+	check(TagAttribute);
+	TagAttribute->Tag = InTagHandle;
+
+	SceneAttributes.Add(TagAttribute);
+	return true;
+}
+
+bool UAvaAttributeContainer::RemoveTagAttribute(const FAvaTagHandle& InTagHandle)
+{
+	uint32 TagsCleared = 0;
+
+	for (UAvaAttribute* Attribute : SceneAttributes)
+	{
+		UAvaTagAttributeBase* TagAttribute = Cast<UAvaTagAttributeBase>(Attribute);
+		if (!TagAttribute)
+		{
+			continue;
+		}
+
+		// Attempt to clear the given tag handle for the attribute. This will return true if it did remove the entry
+		// Do not remove the attribute itself from the list as it could still have valid tags, or later have valid tags
+		if (TagAttribute->ClearTagHandle(InTagHandle))
+		{
+			++TagsCleared;
+		}
+	}
+
+	return TagsCleared > 0;
+}
+
+bool UAvaAttributeContainer::ContainsTagAttribute(const FAvaTagHandle& InTagHandle) const
+{
+	return SceneAttributes.ContainsByPredicate(
+		[&InTagHandle](UAvaAttribute* InAttribute)
+		{
+			UAvaTagAttributeBase* TagAttribute = Cast<UAvaTagAttributeBase>(InAttribute);
+			return TagAttribute && TagAttribute->ContainsTag(InTagHandle);
+		});
+}
+
+bool UAvaAttributeContainer::AddNameAttribute(FName InName)
+{
+	if (InName.IsNone())
+	{
+		return false;
+	}
+
+	// Return true if already existing
+	if (ContainsNameAttribute(InName))
+	{
+		return true;
+	}
+
+	UAvaNameAttribute* NameAttribute = NewObject<UAvaNameAttribute>(this, NAME_None, RF_Transient);
+	check(NameAttribute);
+	NameAttribute->Name = InName;
+
+	SceneAttributes.Add(NameAttribute);
+	return true;
+}
+
+bool UAvaAttributeContainer::RemoveNameAttribute(FName InName)
+{
+	uint32 NamesCleared = 0;
+
+	for (TArray<TObjectPtr<UAvaAttribute>>::TIterator Iter(SceneAttributes); Iter; ++Iter)
+	{
+		UAvaNameAttribute* NameAttribute = Cast<UAvaNameAttribute>(*Iter);
+		if (!NameAttribute || NameAttribute->Name != InName)
+		{
+			continue;
+		}
+
+		// Attempt to clear the given name for the attribute.
+		// If this owns the attribute, remove from the Scene Attributes list
+		if (NameAttribute->GetOuter() == this)
+		{
+			Iter.RemoveCurrent();
+		}
+		// Do not remove the attribute itself from the list if it's an external attribute as it could still have valid tags,
+		// or later have valid tags (due to a dynamic change)
+		else
+		{
+			NameAttribute->Name = NAME_None;
+		}
+
+		++NamesCleared;
+	}
+
+	return NamesCleared > 0;
+}
+
+bool UAvaAttributeContainer::ContainsNameAttribute(FName InName) const
+{
+	return SceneAttributes.ContainsByPredicate(
+		[InName](UAvaAttribute* InAttribute)
+		{
+			UAvaNameAttribute* NameAttribute = Cast<UAvaNameAttribute>(InAttribute);
+			return NameAttribute && NameAttribute->Name == InName;
+		});
+}
diff --git a/Engine/Plugins/VirtualProduction/Avalanche/Source/Avalanche/Private/AvaScene.cpp b/Engine/Plugins/VirtualProduction/Avalanche/Source/Avalanche/Private/AvaScene.cpp
--- a/Engine/Plugins/VirtualProduction/Avalanche/Source/Avalanche/Private/AvaScene.cpp
+++ b/Engine/Plugins/VirtualProduction/Avalanche/Source/Avalanche/Private/AvaScene.cpp
@@ -2,10 +2,10 @@
 
 #include "AvaScene.h"
 #include "AvaAssetTags.h"
+#include "AvaAttributeContainer.h"
 #include "AvaCameraSubsystem.h"
 #include "AvaRemoteControlUtils.h"
 #include "AvaSceneSettings.h"
-#include "AvaSceneState.h"
 #include "AvaSceneSubsystem.h"
 #include "AvaSequencePlaybackObject.h"
 #include "AvaSequenceSubsystem.h"
@@ -83,7 +83,13 @@ AAvaScene::AAvaScene()
 {
 	SceneSettings = CreateDefaultSubobject<UAvaSceneSettings>(TEXT("SceneSettings"));
 
-	SceneState = CreateDefaultSubobject<UAvaSceneState>(TEXT("SceneState"));
+	AttributeContainer = CreateDefaultSubobject<UAvaAttributeContainer>(TEXT("SceneState"));
+
+#if WITH_EDITOR
+PRAGMA_DISABLE_DEPRECATION_WARNINGS
+	SceneState = AttributeContainer;
+PRAGMA_ENABLE_DEPRECATION_WARNINGS
+#endif
 
 	RemoteControlPreset = CreateDefaultSubobject<URemoteControlPreset>(TEXT("RemoteControlPreset"));
 
@@ -319,18 +325,18 @@ void AAvaScene::PostLoad()
 	Super::PostLoad();
 	RegisterObjects();
 
-	if (SceneState)
+	if (AttributeContainer)
 	{
-		SceneState->Initialize(SceneSettings);
+		AttributeContainer->Initialize(SceneSettings);
 	}
 }
 
 void AAvaScene::PostInitializeComponents()
 {
 	Super::PostInitializeComponents();
-	if (SceneState)
+	if (AttributeContainer)
 	{
-		SceneState->Initialize(SceneSettings);
+		AttributeContainer->Initialize(SceneSettings);
 	}
 }
 
diff --git a/Engine/Plugins/VirtualProduction/Avalanche/Source/Avalanche/Private/AvaSceneState.cpp b/Engine/Plugins/VirtualProduction/Avalanche/Source/Avalanche/Private/AvaSceneState.cpp
deleted file mode 100644
diff --git a/Engine/Plugins/VirtualProduction/Avalanche/Source/Avalanche/Private/Transition/AvaTransitionAttributeLibrary.cpp b/Engine/Plugins/VirtualProduction/Avalanche/Source/Avalanche/Private/Transition/AvaTransitionAttributeLibrary.cpp
--- a/Engine/Plugins/VirtualProduction/Avalanche/Source/Avalanche/Private/Transition/AvaTransitionAttributeLibrary.cpp
+++ b/Engine/Plugins/VirtualProduction/Avalanche/Source/Avalanche/Private/Transition/AvaTransitionAttributeLibrary.cpp
@@ -1,7 +1,7 @@
 // Copyright Epic Games, Inc. All Rights Reserved.
 
 #include "AvaTransitionAttributeLibrary.h"
-#include "AvaSceneState.h"
+#include "AvaAttributeContainer.h"
 #include "AvaSceneSubsystem.h"
 #include "AvaTransitionContext.h"
 #include "AvaTransitionLayer.h"
@@ -14,7 +14,7 @@
 
 namespace UE::Ava::Private
 {
-	UAvaSceneState* GetSceneState(UObject* InTransitionNode)
+	UAvaAttributeContainer* GetAttributeContainer(UObject* InTransitionNode)
 	{
 		const IAvaTransitionNodeInterface* NodeInterface = Cast<IAvaTransitionNodeInterface>(InTransitionNode);
 		if (!NodeInterface)
@@ -36,7 +36,7 @@ namespace UE::Ava::Private
 
 		if (IAvaSceneInterface* SceneInterface = UAvaSceneSubsystem::FindSceneInterface(TransitionScene->GetLevel()))
 		{
-			return SceneInterface->GetSceneState();
+			return SceneInterface->GetAttributeContainer();
 		}
 		return nullptr;
 	}
@@ -44,54 +44,54 @@ namespace UE::Ava::Private
 
 bool UAvaTransitionAttributeLibrary::AddTagAttribute(UObject* InTransitionNode, const FAvaTagHandle& InTagHandle)
 {
-	if (UAvaSceneState* SceneState = UE::Ava::Private::GetSceneState(InTransitionNode))
+	if (UAvaAttributeContainer* AttributeContainer = UE::Ava::Private::GetAttributeContainer(InTransitionNode))
 	{
-		return SceneState->AddTagAttribute(InTagHandle);
+		return AttributeContainer->AddTagAttribute(InTagHandle);
 	}
 	return false;
 }
 
 bool UAvaTransitionAttributeLibrary::RemoveTagAttribute(UObject* InTransitionNode, const FAvaTagHandle& InTagHandle)
 {
-	if (UAvaSceneState* SceneState = UE::Ava::Private::GetSceneState(InTransitionNode))
+	if (UAvaAttributeContainer* AttributeContainer = UE::Ava::Private::GetAttributeContainer(InTransitionNode))
 	{
-		return SceneState->RemoveTagAttribute(InTagHandle);
+		return AttributeContainer->RemoveTagAttribute(InTagHandle);
 	}
 	return false;
 }
 
 bool UAvaTransitionAttributeLibrary::ContainsTagAttribute(UObject* InTransitionNode, const FAvaTagHandle& InTagHandle)
 {
-	if (UAvaSceneState* SceneState = UE::Ava::Private::GetSceneState(InTransitionNode))
+	if (UAvaAttributeContainer* AttributeContainer = UE::Ava::Private::GetAttributeContainer(InTransitionNode))
 	{
-		return SceneState->ContainsTagAttribute(InTagHandle);
+		return AttributeContainer->ContainsTagAttribute(InTagHandle);
 	}
 	return false;
 }
 
 bool UAvaTransitionAttributeLibrary::AddNameAttribute(UObject* InTransitionNode, FName InName)
 {
-	if (UAvaSceneState* SceneState = UE::Ava::Private::GetSceneState(InTransitionNode))
+	if (UAvaAttributeContainer* AttributeContainer = UE::Ava::Private::GetAttributeContainer(InTransitionNode))
 	{
-		return SceneState->AddNameAttribute(InName);
+		return AttributeContainer->AddNameAttribute(InName);
 	}
 	return false;
 }
 
 bool UAvaTransitionAttributeLibrary::RemoveNameAttribute(UObject* InTransitionNode, FName InName)
 {
-	if (UAvaSceneState* SceneState = UE::Ava::Private::GetSceneState(InTransitionNode))
+	if (UAvaAttributeContainer* AttributeContainer = UE::Ava::Private::GetAttributeContainer(InTransitionNode))
 	{
-		return SceneState->RemoveNameAttribute(InName);
+		return AttributeContainer->RemoveNameAttribute(InName);
 	}
 	return false;
 }
 
 bool UAvaTransitionAttributeLibrary::ContainsNameAttribute(UObject* InTransitionNode, FName InName)
 {
-	if (UAvaSceneState* SceneState = UE::Ava::Private::GetSceneState(InTransitionNode))
+	if (UAvaAttributeContainer* AttributeContainer = UE::Ava::Private::GetAttributeContainer(InTransitionNode))
 	{
-		return SceneState->ContainsNameAttribute(InName);
+		return AttributeContainer->ContainsNameAttribute(InName);
 	}
 	return false;
 }
diff --git a/Engine/Plugins/VirtualProduction/Avalanche/Source/Avalanche/Private/Transition/Conditions/AvaSceneContainsTagAttributeCondition.cpp b/Engine/Plugins/VirtualProduction/Avalanche/Source/Avalanche/Private/Transition/Conditions/AvaSceneContainsTagAttributeCondition.cpp
--- a/Engine/Plugins/VirtualProduction/Avalanche/Source/Avalanche/Private/Transition/Conditions/AvaSceneContainsTagAttributeCondition.cpp
+++ b/Engine/Plugins/VirtualProduction/Avalanche/Source/Avalanche/Private/Transition/Conditions/AvaSceneContainsTagAttributeCondition.cpp
@@ -1,7 +1,7 @@
 // Copyright Epic Games, Inc. All Rights Reserved.
 
 #include "Transition/Conditions/AvaSceneContainsTagAttributeCondition.h"
-#include "AvaSceneState.h"
+#include "AvaAttributeContainer.h"
 #include "AvaSceneSubsystem.h"
 #include "AvaTransitionLayer.h"
 #include "AvaTransitionLayerUtils.h"
@@ -131,8 +131,8 @@ bool FAvaSceneContainsTagAttributeConditionBase::ContainsTagAttribute(FStateTree
 			continue;
 		}
 
-		const UAvaSceneState* SceneState = Scene->GetSceneState();
-		if (SceneState && SceneState->ContainsTagAttribute(InstanceData.TagAttribute))
+		const UAvaAttributeContainer* AttributeContainer = Scene->GetAttributeContainer();
+		if (AttributeContainer && AttributeContainer->ContainsTagAttribute(InstanceData.TagAttribute))
 		{
 			return true;
 		}
diff --git a/Engine/Plugins/VirtualProduction/Avalanche/Source/Avalanche/Private/Transition/Tasks/AvaSceneAddTagAttributeTask.cpp b/Engine/Plugins/VirtualProduction/Avalanche/Source/Avalanche/Private/Transition/Tasks/AvaSceneAddTagAttributeTask.cpp
--- a/Engine/Plugins/VirtualProduction/Avalanche/Source/Avalanche/Private/Transition/Tasks/AvaSceneAddTagAttributeTask.cpp
+++ b/Engine/Plugins/VirtualProduction/Avalanche/Source/Avalanche/Private/Transition/Tasks/AvaSceneAddTagAttributeTask.cpp
@@ -1,7 +1,7 @@
 // Copyright Epic Games, Inc. All Rights Reserved.
 
 #include "Transition/Tasks/AvaSceneAddTagAttributeTask.h"
-#include "AvaSceneState.h"
+#include "AvaAttributeContainer.h"
 #include "IAvaSceneInterface.h"
 #include "StateTreeExecutionContext.h"
 
@@ -28,14 +28,14 @@ EStateTreeRunStatus FAvaSceneAddTagAttributeTask::EnterState(FStateTreeExecution
 		return EStateTreeRunStatus::Failed;
 	}
 
-	UAvaSceneState* SceneState = Scene->GetSceneState();
-	if (!SceneState)
+	UAvaAttributeContainer* AttributeContainer = Scene->GetAttributeContainer();
+	if (!AttributeContainer)
 	{
 		return EStateTreeRunStatus::Failed;
 	}
 
 	const FInstanceDataType& InstanceData = InContext.GetInstanceData(*this);
-	if (SceneState->AddTagAttribute(InstanceData.TagAttribute))
+	if (AttributeContainer->AddTagAttribute(InstanceData.TagAttribute))
 	{
 		return EStateTreeRunStatus::Succeeded;
 	}
diff --git a/Engine/Plugins/VirtualProduction/Avalanche/Source/Avalanche/Private/Transition/Tasks/AvaSceneRemoveTagAttributeTask.cpp b/Engine/Plugins/VirtualProduction/Avalanche/Source/Avalanche/Private/Transition/Tasks/AvaSceneRemoveTagAttributeTask.cpp
--- a/Engine/Plugins/VirtualProduction/Avalanche/Source/Avalanche/Private/Transition/Tasks/AvaSceneRemoveTagAttributeTask.cpp
+++ b/Engine/Plugins/VirtualProduction/Avalanche/Source/Avalanche/Private/Transition/Tasks/AvaSceneRemoveTagAttributeTask.cpp
@@ -1,7 +1,7 @@
 // Copyright Epic Games, Inc. All Rights Reserved.
 
 #include "Transition/Tasks/AvaSceneRemoveTagAttributeTask.h"
-#include "AvaSceneState.h"
+#include "AvaAttributeContainer.h"
 #include "IAvaSceneInterface.h"
 #include "StateTreeExecutionContext.h"
 
@@ -28,14 +28,14 @@ EStateTreeRunStatus FAvaSceneRemoveTagAttributeTask::EnterState(FStateTreeExecut
 		return EStateTreeRunStatus::Failed;
 	}
 
-	UAvaSceneState* SceneState = Scene->GetSceneState();
-	if (!SceneState)
+	UAvaAttributeContainer* AttributeContainer = Scene->GetAttributeContainer();
+	if (!AttributeContainer)
 	{
 		return EStateTreeRunStatus::Failed;
 	}
 
 	const FInstanceDataType& InstanceData = InContext.GetInstanceData(*this);
-	if (SceneState->RemoveTagAttribute(InstanceData.TagAttribute))
+	if (AttributeContainer->RemoveTagAttribute(InstanceData.TagAttribute))
 	{
 		return EStateTreeRunStatus::Succeeded;
 	}
diff --git a/Engine/Plugins/VirtualProduction/Avalanche/Source/Avalanche/Public/AvaAttributeContainer.h b/Engine/Plugins/VirtualProduction/Avalanche/Source/Avalanche/Public/AvaAttributeContainer.h
new file mode 100644
--- /dev/null
+++ b/Engine/Plugins/VirtualProduction/Avalanche/Source/Avalanche/Public/AvaAttributeContainer.h
@@ -0,0 +1,40 @@
+// Copyright Epic Games, Inc. All Rights Reserved.
+
+#pragma once
+
+#include "AvaTag.h"
+#include "Containers/Array.h"
+#include "Containers/Set.h"
+#include "UObject/Object.h"
+#include "AvaAttributeContainer.generated.h"
+
+#define UE_API AVALANCHE_API
+
+class UAvaAttribute;
+class UAvaSceneSettings;
+struct FAvaTagHandle;
+
+/** Object providing attribute information of the Scene */
+UCLASS(MinimalAPI)
+class UAvaAttributeContainer : public UObject
+{
+	GENERATED_BODY()
+
+public:
+	void Initialize(UAvaSceneSettings* InSceneSettings);
+
+	UE_API bool AddTagAttribute(const FAvaTagHandle& InTagHandle);
+	UE_API bool RemoveTagAttribute(const FAvaTagHandle& InTagHandle);
+	UE_API bool ContainsTagAttribute(const FAvaTagHandle& InTagHandle) const;
+
+	UE_API bool AddNameAttribute(FName InName);
+	UE_API bool RemoveNameAttribute(FName InName);
+	UE_API bool ContainsNameAttribute(FName InName) const;
+
+private:
+	/** In-play Scene Attributes. Can be added to / removed from while in-play */
+	UPROPERTY(Transient)
+	TArray<TObjectPtr<UAvaAttribute>> SceneAttributes;
+};
+
+#undef UE_API
diff --git a/Engine/Plugins/VirtualProduction/Avalanche/Source/Avalanche/Public/AvaScene.h b/Engine/Plugins/VirtualProduction/Avalanche/Source/Avalanche/Public/AvaScene.h
--- a/Engine/Plugins/VirtualProduction/Avalanche/Source/Avalanche/Public/AvaScene.h
+++ b/Engine/Plugins/VirtualProduction/Avalanche/Source/Avalanche/Public/AvaScene.h
@@ -14,8 +14,8 @@
 #include "AvaScene.generated.h"
 
 class IAvaSequencePlaybackObject;
+class UAvaAttributeContainer;
 class UAvaSceneSettings;
-class UAvaSceneState;
 class UAvaSequence;
 class URemoteControlPreset;
 
@@ -63,7 +63,7 @@ public:
 	//~ Begin IAvaSceneInterface
 	virtual ULevel* GetSceneLevel() const override;
 	virtual UAvaSceneSettings* GetSceneSettings() const override { return SceneSettings; }
-	virtual UAvaSceneState* GetSceneState() const override { return SceneState; }
+	virtual UAvaAttributeContainer* GetAttributeContainer() const override { return AttributeContainer; }
 	virtual FAvaSceneTree& GetSceneTree() override { return SceneTree; }
 	virtual const FAvaSceneTree& GetSceneTree() const override { return SceneTree; }
 	virtual IAvaSequencePlaybackObject* GetPlaybackObject() const override;
@@ -130,7 +130,12 @@ protected:
 	TObjectPtr<UAvaSceneSettings> SceneSettings;
 
 	UPROPERTY()
-	TObjectPtr<UAvaSceneState> SceneState;
+	TObjectPtr<UAvaAttributeContainer> AttributeContainer;
+
+#if WITH_EDITORONLY_DATA
+	UE_DEPRECATED(5.7, "This property is deprecated. Please use AttributeContainer")
+	TObjectPtr<UAvaAttributeContainer> SceneState;
+#endif // WITH_EDITORONLY_DATA
 
 	UPROPERTY()
 	TObjectPtr<URemoteControlPreset> RemoteControlPreset;
diff --git a/Engine/Plugins/VirtualProduction/Avalanche/Source/Avalanche/Public/AvaSceneState.h b/Engine/Plugins/VirtualProduction/Avalanche/Source/Avalanche/Public/AvaSceneState.h
--- a/Engine/Plugins/VirtualProduction/Avalanche/Source/Avalanche/Public/AvaSceneState.h
+++ b/Engine/Plugins/VirtualProduction/Avalanche/Source/Avalanche/Public/AvaSceneState.h
@@ -2,35 +2,11 @@
 
 #pragma once
 
-#include "AvaTag.h"
-#include "Containers/Array.h"
-#include "Containers/Set.h"
-#include "UObject/Object.h"
-#include "AvaSceneState.generated.h"
+#include "Misc/CoreMiscDefines.h"
 
-class UAvaAttribute;
-class UAvaSceneSettings;
-struct FAvaTagHandle;
+/** Renamed AvaSceneState to AvaAttributeContainer to avoid confusion with the 'Scene State' system */
+UE_DEPRECATED_HEADER(5.7, "Use AvaAttributeContainer.h instead")
+#include "AvaAttributeContainer.h"
 
-/** Object providing attribute information of the Scene */
-UCLASS(MinimalAPI)
-class UAvaSceneState : public UObject
-{
-	GENERATED_BODY()
-
-public:
-	void Initialize(UAvaSceneSettings* InSceneSettings);
-
-	AVALANCHE_API bool AddTagAttribute(const FAvaTagHandle& InTagHandle);
-	AVALANCHE_API bool RemoveTagAttribute(const FAvaTagHandle& InTagHandle);
-	AVALANCHE_API bool ContainsTagAttribute(const FAvaTagHandle& InTagHandle) const;
-
-	AVALANCHE_API bool AddNameAttribute(FName InName);
-	AVALANCHE_API bool RemoveNameAttribute(FName InName);
-	AVALANCHE_API bool ContainsNameAttribute(FName InName) const;
-
-private:
-	/** In-play Scene Attributes. Can be added to / removed from while in-play */
-	UPROPERTY(Transient)
-	TArray<TObjectPtr<UAvaAttribute>> SceneAttributes;
-};
+UE_DEPRECATED(5.7, "UAvaSceneState has been renamed to UAvaAttributeContainer. Use UAvaAttributeContainer instead.")
+typedef UAvaAttributeContainer UAvaSceneState;
diff --git a/Engine/Plugins/VirtualProduction/Avalanche/Source/Avalanche/Public/IAvaSceneInterface.h b/Engine/Plugins/VirtualProduction/Avalanche/Source/Avalanche/Public/IAvaSceneInterface.h
--- a/Engine/Plugins/VirtualProduction/Avalanche/Source/Avalanche/Public/IAvaSceneInterface.h
+++ b/Engine/Plugins/VirtualProduction/Avalanche/Source/Avalanche/Public/IAvaSceneInterface.h
@@ -7,8 +7,8 @@
 
 class IAvaSequencePlaybackObject;
 class IAvaSequenceProvider;
+class UAvaAttributeContainer;
 class UAvaSceneSettings;
-class UAvaSceneState;
 class ULevel;
 class URemoteControlPreset;
 struct FAvaSceneTree;
@@ -36,7 +36,13 @@ public:
 
 	virtual UAvaSceneSettings* GetSceneSettings() const = 0;
 
-	virtual UAvaSceneState* GetSceneState() const = 0;
+	virtual UAvaAttributeContainer* GetAttributeContainer() const = 0;
+
+	UE_DEPRECATED(5.7, "GetSceneState deprecated. Use GetAttributeContainer instead")
+	virtual UAvaAttributeContainer* GetSceneState() const
+	{
+		return GetAttributeContainer();
+	}
 
 	virtual const FAvaSceneTree& GetSceneTree() const = 0;
 
diff --git a/Engine/Plugins/VirtualProduction/Avalanche/Source/AvalancheEditor/Private/Scene/AvaSceneRigExtension.cpp b/Engine/Plugins/VirtualProduction/Avalanche/Source/AvalancheEditor/Private/Scene/AvaSceneRigExtension.cpp
--- a/Engine/Plugins/VirtualProduction/Avalanche/Source/AvalancheEditor/Private/Scene/AvaSceneRigExtension.cpp
+++ b/Engine/Plugins/VirtualProduction/Avalanche/Source/AvalancheEditor/Private/Scene/AvaSceneRigExtension.cpp
@@ -4,7 +4,6 @@
 #include "AvaSceneRigEditorCommands.h"
 #include "AvaSceneRigSubsystem.h"
 #include "AvaSceneSettings.h"
-#include "AvaSceneState.h"
 #include "AvaSceneSubsystem.h"
 #include "Engine/Level.h"
 #include "Engine/LevelStreaming.h"
@@ -14,11 +13,11 @@
 #include "Item/AvaOutlinerActor.h"
 #include "ScopedTransaction.h"
 #include "Styling/AppStyle.h"
-#include "ToolMenus.h"
 #include "ToolMenu.h"
 #include "ToolMenuContext/AvaOutlinerItemsContext.h"
 #include "ToolMenuEntry.h"
 #include "ToolMenuSection.h"
+#include "ToolMenus.h"
 
 #define LOCTEXT_NAMESPACE "AvaSceneRigExtension"
 
@@ -378,26 +377,4 @@ TArray<AActor*> FAvaSceneRigExtension::OutlinerItemsToActors(const TConstArrayVi
 	return OutActors;
 }
 
-UAvaSceneState* FAvaSceneRigExtension::GetSceneState(UWorld* const InWorld)
-{
-	if (!InWorld)
-	{
-		return nullptr;
-	}
-
-	const UAvaSceneSubsystem* const SceneSubsystem = InWorld->GetSubsystem<UAvaSceneSubsystem>();
-	if (!SceneSubsystem)
-	{
-		return nullptr;
-	}
-
-	IAvaSceneInterface* const AvaScene = SceneSubsystem->GetSceneInterface();
-	if (!AvaScene)
-	{
-		return nullptr;
-	}
-
-	return AvaScene->GetSceneState();
-}
-
 #undef LOCTEXT_NAMESPACE
diff --git a/Engine/Plugins/VirtualProduction/Avalanche/Source/AvalancheEditor/Private/Scene/AvaSceneRigExtension.h b/Engine/Plugins/VirtualProduction/Avalanche/Source/AvalancheEditor/Private/Scene/AvaSceneRigExtension.h
--- a/Engine/Plugins/VirtualProduction/Avalanche/Source/AvalancheEditor/Private/Scene/AvaSceneRigExtension.h
+++ b/Engine/Plugins/VirtualProduction/Avalanche/Source/AvalancheEditor/Private/Scene/AvaSceneRigExtension.h
@@ -11,7 +11,6 @@ class AActor;
 class FUICommandList;
 class IAvaEditor;
 class UAvaOutlinerItemsContext;
-class UAvaSceneState;
 class UToolMenu;
 
 class FAvaSceneRigExtension : public FAvaEditorExtension
@@ -19,8 +18,6 @@ class FAvaSceneRigExtension : public FAvaEditorExtension
 public:
 	UE_AVA_INHERITS(FAvaSceneRigExtension, FAvaEditorExtension);
 
-	static UAvaSceneState* GetSceneState(UWorld* const InWorld);
-
 	void PromptToSaveSceneRigFromOutlinerItems();
 	bool CanSaveSceneRigFromOutlinerItems() const;
 
