MotionDesign : Mask
- Fix mask modifier not restoring original materials when they are transient and causing issues on level reload

diff --git a/Engine/Plugins/VirtualProduction/Avalanche/Source/AvalancheMask/Private/Handling/AvaMaskActorMaterialCollectionHandle.cpp b/Engine/Plugins/VirtualProduction/Avalanche/Source/AvalancheMask/Private/Handling/AvaMaskActorMaterialCollectionHandle.cpp
--- a/Engine/Plugins/VirtualProduction/Avalanche/Source/AvalancheMask/Private/Handling/AvaMaskActorMaterialCollectionHandle.cpp
+++ b/Engine/Plugins/VirtualProduction/Avalanche/Source/AvalancheMask/Private/Handling/AvaMaskActorMaterialCollectionHandle.cpp
@@ -295,10 +295,7 @@ void FAvaMaskActorMaterialCollectionHandle::MapEachMaterialHandle(TFunctionRef<T
 		GetMaterialHandles();
 
 		TArray<TObjectPtr<UMaterialInterface>> MappedMaterials;
-		MappedMaterials.Reserve(WeakMaterials.Num());
-
 		TBitArray<> SetFlags;
-		SetFlags.Init(false, MappedMaterials.Num());
 
 		TGuardValue<bool> IteratingGuard(bIsIteratingMaterials, true);
 		for (TPair<FSoftComponentReference, TArray<TSharedPtr<IAvaMaskMaterialHandle>>>& ComponentMaterialHandles : MaterialHandles)
@@ -311,15 +308,15 @@ void FAvaMaskActorMaterialCollectionHandle::MapEachMaterialHandle(TFunctionRef<T
 					UMaterialInterface* RemappedMaterial = nullptr;
 					if (MappedMaterialHandle.IsValid())
 					{
-						RemappedMaterial = MappedMaterialHandle->GetMaterial();	
+						RemappedMaterial = MappedMaterialHandle->GetMaterial();
 					}
 					else
 					{
 						RemappedMaterial = MaterialHandle->GetMaterial();
 					}
 
-					MappedMaterials[SlotIdx] = RemappedMaterial;
-					SetFlags[SlotIdx] = RemappedMaterial != nullptr;
+					MappedMaterials.Add(RemappedMaterial);
+					SetFlags.Add(RemappedMaterial != nullptr);
 				}
 			}
 		}
diff --git a/Engine/Plugins/VirtualProduction/Avalanche/Source/AvalancheMask/Private/Handling/AvaMaskAvaShapeMaterialCollectionHandle.cpp b/Engine/Plugins/VirtualProduction/Avalanche/Source/AvalancheMask/Private/Handling/AvaMaskAvaShapeMaterialCollectionHandle.cpp
--- a/Engine/Plugins/VirtualProduction/Avalanche/Source/AvalancheMask/Private/Handling/AvaMaskAvaShapeMaterialCollectionHandle.cpp
+++ b/Engine/Plugins/VirtualProduction/Avalanche/Source/AvalancheMask/Private/Handling/AvaMaskAvaShapeMaterialCollectionHandle.cpp
@@ -255,12 +255,7 @@ void FAvaMaskAvaShapeMaterialCollectionHandle::MapEachMaterialHandle(TFunctionRe
 		GetMaterialHandles();
 		
 		TArray<TObjectPtr<UMaterialInterface>> MappedMaterials;
-        MappedMaterials.Init(nullptr, GetNumMaterials());
-
 		TBitArray<> SetFlags;
-		SetFlags.Init(false, MappedMaterials.Num());
-
-        int32 MeshIdx = 0;
 		
 		const FSoftComponentReference ComponentReference = UE::AvaMask::Internal::MakeComponentReference(ShapeComponent->GetOwner(), ShapeComponent);
 		for (TPair<int32, TSharedPtr<IAvaMaskMaterialHandle>>& MaterialHandle : MaterialHandles)
@@ -276,8 +271,8 @@ void FAvaMaskAvaShapeMaterialCollectionHandle::MapEachMaterialHandle(TFunctionRe
 				RemappedMaterial = MaterialHandle.Value->GetMaterial();
 			}
 
-			MappedMaterials[++MeshIdx] = RemappedMaterial;
-			SetFlags[MeshIdx] = RemappedMaterial != nullptr;
+			MappedMaterials.Add(RemappedMaterial);
+			SetFlags.Add(RemappedMaterial != nullptr);
 		}
 
 		SetMaterials(MappedMaterials, SetFlags);
diff --git a/Engine/Plugins/VirtualProduction/Avalanche/Source/AvalancheMask/Private/Handling/AvaMaskDesignedMaterialHandle.h b/Engine/Plugins/VirtualProduction/Avalanche/Source/AvalancheMask/Private/Handling/AvaMaskDesignedMaterialHandle.h
--- a/Engine/Plugins/VirtualProduction/Avalanche/Source/AvalancheMask/Private/Handling/AvaMaskDesignedMaterialHandle.h
+++ b/Engine/Plugins/VirtualProduction/Avalanche/Source/AvalancheMask/Private/Handling/AvaMaskDesignedMaterialHandle.h
@@ -58,6 +58,7 @@ public:
 	virtual UMaterialInterface* GetMaterial() override { return Super::GetMaterial(); }
 	virtual void CopyParametersFrom(UMaterialInstance* InSourceMaterial) override { Super::CopyParametersFrom(InSourceMaterial); }
 	virtual bool HasRequiredParameters(TArray<FString>& OutMissingParameterNames) override;
+	virtual void SetParentMaterial(UMaterialInterface* InMaterial) override { Super::SetParentMaterial(InMaterial); }
 	// ~End IAvaMaterialHandle
 
 	// ~Begin IAvaObjectHandle
diff --git a/Engine/Plugins/VirtualProduction/Avalanche/Source/AvalancheMask/Private/Handling/AvaMaskMaterialAssignmentObserver.cpp b/Engine/Plugins/VirtualProduction/Avalanche/Source/AvalancheMask/Private/Handling/AvaMaskMaterialAssignmentObserver.cpp
--- a/Engine/Plugins/VirtualProduction/Avalanche/Source/AvalancheMask/Private/Handling/AvaMaskMaterialAssignmentObserver.cpp
+++ b/Engine/Plugins/VirtualProduction/Avalanche/Source/AvalancheMask/Private/Handling/AvaMaskMaterialAssignmentObserver.cpp
@@ -5,21 +5,23 @@
 #include "Components/ActorComponent.h"
 #include "Components/PrimitiveComponent.h"
 #include "GameFramework/Actor.h"
+#include "Materials/Material.h"
+#include "Materials/MaterialInstanceDynamic.h"
 #include "Materials/MaterialInterface.h"
 
 FAvaMaskMaterialAssignmentObserver::FAvaMaskMaterialAssignmentObserver(UPrimitiveComponent* InComponent)
 	: ComponentWeak(InComponent)
 {
 	check(InComponent);
-	
+
 	SlotCount = InComponent->GetNumMaterials();
 	Materials.AddDefaulted(InComponent->GetNumMaterials());
-	
+
 	for (int32 MaterialIdx = 0; MaterialIdx < InComponent->GetNumMaterials(); ++MaterialIdx)
 	{
 		Materials[MaterialIdx] = InComponent->GetMaterial(MaterialIdx);
 	}
-	
+
 	if (!UPrimitiveComponent::MarkRenderStateDirtyEvent.IsBoundToObject(this))
 	{
 		RenderStateDirtyHandle = UPrimitiveComponent::MarkRenderStateDirtyEvent.AddRaw(this, &FAvaMaskMaterialAssignmentObserver::OnComponentRenderStateDirty);
@@ -39,6 +41,11 @@ FAvaMaskMaterialAssignmentObserver::FOnMaterialAssignedDelegate& FAvaMaskMateria
 	return OnMaterialAssignedDelegate;
 }
 
+UPrimitiveComponent* FAvaMaskMaterialAssignmentObserver::GetPrimitiveComponent() const
+{
+	return ComponentWeak.Get();
+}
+
 void FAvaMaskMaterialAssignmentObserver::OnComponentRenderStateDirty(UActorComponent& InComponent)
 {
 	// @note: this is global delegate, so we need to check that we care about the component
@@ -56,16 +63,23 @@ void FAvaMaskMaterialAssignmentObserver::OnComponentRenderStateDirty(UActorCompo
 			Materials.SetNum(NewSlotCount);
 			for (int32 MaterialIdx = 0; MaterialIdx < Component->GetNumMaterials(); ++MaterialIdx)
 			{
-				if (UMaterialInterface* NewMaterial = Component->GetMaterial(MaterialIdx))
+				UMaterialInterface* NewMaterial = Component->GetMaterial(MaterialIdx);
+
+				if (NewMaterial)
 				{
+					const UMaterialInterface* OldMaterial = Cast<UMaterialInterface>(Materials[MaterialIdx].ResolveObjectPtr());
+
+					const bool bDifferentObject = OldMaterial != NewMaterial;
+					const bool bSameBaseMaterial = OldMaterial != nullptr && NewMaterial->IsA<UMaterialInstanceDynamic>() && OldMaterial->GetMaterial() == NewMaterial->GetMaterial();
+
 					// Compare actual material
-					if (Materials[MaterialIdx] != NewMaterial)
+					if (bDifferentObject && !bSameBaseMaterial)
 					{
 						ChangedMaterialSlots.Add(MaterialIdx);
 					}
 				}
-				
-				Materials[MaterialIdx] = Component->GetMaterial(MaterialIdx);
+
+				Materials[MaterialIdx] = NewMaterial;
 			}
 
 			if (!ChangedMaterialSlots.IsEmpty())
diff --git a/Engine/Plugins/VirtualProduction/Avalanche/Source/AvalancheMask/Private/Handling/AvaMaskMaterialAssignmentObserver.h b/Engine/Plugins/VirtualProduction/Avalanche/Source/AvalancheMask/Private/Handling/AvaMaskMaterialAssignmentObserver.h
--- a/Engine/Plugins/VirtualProduction/Avalanche/Source/AvalancheMask/Private/Handling/AvaMaskMaterialAssignmentObserver.h
+++ b/Engine/Plugins/VirtualProduction/Avalanche/Source/AvalancheMask/Private/Handling/AvaMaskMaterialAssignmentObserver.h
@@ -22,14 +22,14 @@ public:
 
 	FOnMaterialAssignedDelegate& OnMaterialAssigned();
 
+	UPrimitiveComponent* GetPrimitiveComponent() const;
+
 private:
 	void OnComponentRenderStateDirty(UActorComponent& InComponent);
 
-private:
 	TWeakObjectPtr<UPrimitiveComponent> ComponentWeak;
 	FDelegateHandle RenderStateDirtyHandle;
 	FOnMaterialAssignedDelegate OnMaterialAssignedDelegate;
-	
 	int32 SlotCount = 0;
 	TArray<FObjectKey> Materials; 
 };
diff --git a/Engine/Plugins/VirtualProduction/Avalanche/Source/AvalancheMask/Private/Handling/AvaMaskMaterialInstanceHandle.h b/Engine/Plugins/VirtualProduction/Avalanche/Source/AvalancheMask/Private/Handling/AvaMaskMaterialInstanceHandle.h
--- a/Engine/Plugins/VirtualProduction/Avalanche/Source/AvalancheMask/Private/Handling/AvaMaskMaterialInstanceHandle.h
+++ b/Engine/Plugins/VirtualProduction/Avalanche/Source/AvalancheMask/Private/Handling/AvaMaskMaterialInstanceHandle.h
@@ -57,6 +57,7 @@ public:
 	virtual FString GetMaterialName() override { return Super::GetMaterialName(); }
 	virtual UMaterialInterface* GetMaterial() override { return FAvaMaterialInstanceHandle::GetMaterial(); }
 	virtual void CopyParametersFrom(UMaterialInstance* InSourceMaterial) override { FAvaMaterialInstanceHandle::CopyParametersFrom(InSourceMaterial); }
+	virtual void SetParentMaterial(UMaterialInterface* InMaterial) override { FAvaMaterialInstanceHandle::SetParentMaterial(InMaterial); }
 	// ~End IAvaMaterialHandle
 
 	// ~Begin IAvaObjectHandle	
diff --git a/Engine/Plugins/VirtualProduction/Avalanche/Source/AvalancheMask/Private/Handling/AvaMaskParametricMaterialHandle.h b/Engine/Plugins/VirtualProduction/Avalanche/Source/AvalancheMask/Private/Handling/AvaMaskParametricMaterialHandle.h
--- a/Engine/Plugins/VirtualProduction/Avalanche/Source/AvalancheMask/Private/Handling/AvaMaskParametricMaterialHandle.h
+++ b/Engine/Plugins/VirtualProduction/Avalanche/Source/AvalancheMask/Private/Handling/AvaMaskParametricMaterialHandle.h
@@ -53,6 +53,7 @@ public:
 	virtual FString GetMaterialName() override { return Super::GetMaterialName(); }
 	virtual UMaterialInterface* GetMaterial() override { return FAvaParametricMaterialHandle::GetMaterial(); }
 	virtual void CopyParametersFrom(UMaterialInstance* InSourceMaterial) override { FAvaParametricMaterialHandle::CopyParametersFrom(InSourceMaterial); }
+	virtual void SetParentMaterial(UMaterialInterface* InMaterial) override { FAvaParametricMaterialHandle::SetParentMaterial(InMaterial); }
 	// ~End IAvaMaterialHandle
 
 	// ~Begin IAvaObjectHandle
diff --git a/Engine/Plugins/VirtualProduction/Avalanche/Source/AvalancheMask/Private/Handling/AvaMaskText3DActorMaterialCollectionHandle.cpp b/Engine/Plugins/VirtualProduction/Avalanche/Source/AvalancheMask/Private/Handling/AvaMaskText3DActorMaterialCollectionHandle.cpp
--- a/Engine/Plugins/VirtualProduction/Avalanche/Source/AvalancheMask/Private/Handling/AvaMaskText3DActorMaterialCollectionHandle.cpp
+++ b/Engine/Plugins/VirtualProduction/Avalanche/Source/AvalancheMask/Private/Handling/AvaMaskText3DActorMaterialCollectionHandle.cpp
@@ -321,7 +321,7 @@ void FAvaMaskText3DActorMaterialCollectionHandle::MapEachMaterialHandle(TFunctio
 		GetMaterialHandles();
 
 		TArray<TObjectPtr<UMaterialInterface>> MappedMaterials;
-		MappedMaterials.Reserve(GetNumMaterials());
+		MappedMaterials.SetNum(GetNumMaterials());
 
 		TBitArray<> SetFlags;
 		SetFlags.Init(false, GetNumMaterials());
@@ -333,7 +333,7 @@ void FAvaMaskText3DActorMaterialCollectionHandle::MapEachMaterialHandle(TFunctio
 			int32 InMaterialGroupIdx = static_cast<int32>(InMaterialGroup);			
 			
 			TSharedPtr<IAvaMaskMaterialHandle>& MaterialHandle = MaterialHandles[InMaterialGroupIdx];
-			if (ensureAlways(UE::Ava::Internal::IsHandleValid(MaterialHandle)))
+
 			{
 				const TSharedPtr<IAvaMaskMaterialHandle> MappedMaterialHandle = InFunction(ComponentReference, InMaterialGroupIdx, true, MaterialHandle);
 				UMaterialInterface* RemappedMaterial = nullptr;
@@ -341,14 +341,14 @@ void FAvaMaskText3DActorMaterialCollectionHandle::MapEachMaterialHandle(TFunctio
 				{
 					RemappedMaterial = MappedMaterialHandle->GetMaterial();	
 				}
-				else
+				else if (MaterialHandle.IsValid())
 				{
 					RemappedMaterial = MaterialHandle->GetMaterial();
 				}
 
 				MappedMaterials[InMaterialGroupIdx] = RemappedMaterial;
 				SetFlags[InMaterialGroupIdx] = RemappedMaterial != nullptr;
-			}	
+			}
 		};
 
 		TryMap(EText3DGroupType::Front);
diff --git a/Engine/Plugins/VirtualProduction/Avalanche/Source/AvalancheMask/Private/Handling/AvaMaterialInstanceHandle.cpp b/Engine/Plugins/VirtualProduction/Avalanche/Source/AvalancheMask/Private/Handling/AvaMaterialInstanceHandle.cpp
--- a/Engine/Plugins/VirtualProduction/Avalanche/Source/AvalancheMask/Private/Handling/AvaMaterialInstanceHandle.cpp
+++ b/Engine/Plugins/VirtualProduction/Avalanche/Source/AvalancheMask/Private/Handling/AvaMaterialInstanceHandle.cpp
@@ -43,6 +43,11 @@ void FAvaMaterialInstanceHandle::CopyParametersFrom(UMaterialInstance* InSourceM
 	}
 }
 
+void FAvaMaterialInstanceHandle::SetParentMaterial(UMaterialInterface* InMaterial)
+{
+	WeakParentMaterial = InMaterial;
+}
+
 bool FAvaMaterialInstanceHandle::IsValid() const
 {
 	// @note: MaterialInstance doesn't have to be valid - this handle should deal with MID creation
diff --git a/Engine/Plugins/VirtualProduction/Avalanche/Source/AvalancheMask/Private/Handling/AvaMaterialInstanceHandle.h b/Engine/Plugins/VirtualProduction/Avalanche/Source/AvalancheMask/Private/Handling/AvaMaterialInstanceHandle.h
--- a/Engine/Plugins/VirtualProduction/Avalanche/Source/AvalancheMask/Private/Handling/AvaMaterialInstanceHandle.h
+++ b/Engine/Plugins/VirtualProduction/Avalanche/Source/AvalancheMask/Private/Handling/AvaMaterialInstanceHandle.h
@@ -27,6 +27,7 @@ public:
 	virtual FString GetMaterialName() override;
 	virtual UMaterialInterface* GetMaterial() override;
 	virtual void CopyParametersFrom(UMaterialInstance* InSourceMaterial) override;
+	virtual void SetParentMaterial(UMaterialInterface* InMaterial) override;
 	// ~End IAvaMaterialHandle
 
 	// ~Begin IAvaObjectHandle
diff --git a/Engine/Plugins/VirtualProduction/Avalanche/Source/AvalancheMask/Private/Handling/IAvaMaterialHandle.h b/Engine/Plugins/VirtualProduction/Avalanche/Source/AvalancheMask/Private/Handling/IAvaMaterialHandle.h
--- a/Engine/Plugins/VirtualProduction/Avalanche/Source/AvalancheMask/Private/Handling/IAvaMaterialHandle.h
+++ b/Engine/Plugins/VirtualProduction/Avalanche/Source/AvalancheMask/Private/Handling/IAvaMaterialHandle.h
@@ -18,6 +18,7 @@ public:
 
 	virtual FString GetMaterialName() = 0;
 	virtual UMaterialInterface* GetMaterial() = 0;
+	virtual void SetParentMaterial(UMaterialInterface* InMaterial) = 0;
 
 	virtual void CopyParametersFrom(UMaterialInstance* InSourceMaterial) = 0;
 };
diff --git a/Engine/Plugins/VirtualProduction/Avalanche/Source/AvalancheMask/Private/Mask2D/AvaMask2DBaseModifier.cpp b/Engine/Plugins/VirtualProduction/Avalanche/Source/AvalancheMask/Private/Mask2D/AvaMask2DBaseModifier.cpp
--- a/Engine/Plugins/VirtualProduction/Avalanche/Source/AvalancheMask/Private/Mask2D/AvaMask2DBaseModifier.cpp
+++ b/Engine/Plugins/VirtualProduction/Avalanche/Source/AvalancheMask/Private/Mask2D/AvaMask2DBaseModifier.cpp
@@ -415,16 +415,30 @@ void UAvaMask2DBaseModifier::SaveActorPreState(AActor* InActor, FAvaMask2DActorD
 			{
 				TSharedPtr<IAvaMaskMaterialCollectionHandle> Handle = GetObjectHandleSubsystem()->MakeHandle<IAvaMaskMaterialCollectionHandle>(InActor, UE::AvaMask::Internal::HandleTag);
 
-				for (const TPair<FAvaMask2DComponentMaterialPath, TObjectPtr<UMaterialInterface>>& OriginalMaterial : InActorData.OriginalMaterials)
-				{
-					Handle->SetMaterial(OriginalMaterial.Key.Component, OriginalMaterial.Key.SlotIdx, OriginalMaterial.Value);
-				}
-
-				InActorData.OriginalMaterials.Empty();
+				Handle->ForEachMaterialHandle([this, &InActorData](
+					const FSoftComponentReference& InComponent,
+					const int32 InIdx,
+					const bool bInIsSlotOccupied,
+					const TSharedPtr<IAvaMaskMaterialHandle>& InMaterialHandle)->bool
+					{
+						if (InMaterialHandle.IsValid())
+						{
+							if (const TObjectPtr<UMaterialInterface>* Material = InActorData.OriginalMaterials.Find(FAvaMask2DComponentMaterialPath(InComponent, InIdx)))
+							{
+								// Refresh parent material to create MID from
+								InMaterialHandle->SetParentMaterial(Material->Get());
+							}
+						}
+
+						return true;
+					}
+				);
 
 				Handle->ForEachMaterial([&InActorData](const FSoftComponentReference& InComponent, const int32 InIdx, UMaterialInterface* InMaterial)
 				{
-					InActorData.OriginalMaterials.Add(FAvaMask2DComponentMaterialPath(InComponent, InIdx), InMaterial);
+					// Save parent material if material is transient otherwise can't restore original
+					UMaterialInterface* Material = InMaterial && InMaterial->HasAnyFlags(RF_Transient) ? InMaterial->GetBaseMaterial() : InMaterial;
+					InActorData.OriginalMaterials.Add(FAvaMask2DComponentMaterialPath(InComponent, InIdx), Material);
 					return true;
 				});
 
@@ -448,16 +462,21 @@ void UAvaMask2DBaseModifier::RestoreActorPreState(AActor* InActor, FAvaMask2DAct
     	UE::Ava::Internal::FindOrAddHandleByLambda(
     		MaterialCollectionHandles
     		, InActor
-    		, [this, InActor]()
+    		, [this, InActor, &InActorData]()
     		{
     			TSharedPtr<IAvaMaskMaterialCollectionHandle> Handle = GetObjectHandleSubsystem()->MakeHandle<IAvaMaskMaterialCollectionHandle>(InActor, UE::AvaMask::Internal::HandleTag);
+
+    			for (const TPair<FAvaMask2DComponentMaterialPath, TObjectPtr<UMaterialInterface>>& OriginalMaterial : InActorData.OriginalMaterials)
+    			{
+    				Handle->SetMaterial(OriginalMaterial.Key.Component, OriginalMaterial.Key.SlotIdx, OriginalMaterial.Value);
+				}
+
     			return Handle;
-    		});
+    		}
+    	);
 
     FInstancedStruct& MaterialCollectionData = MaterialCollectionHandleData.FindOrAdd(InActor, MaterialCollectionHandle->MakeDataStruct());
     MaterialCollectionHandle->ApplyOriginalState(MaterialCollectionData);
-
-	InActorData.OriginalMaterials.Empty();
 }
 
 void UAvaMask2DBaseModifier::OnSceneTreeTrackedActorParentChanged(
diff --git a/Engine/Plugins/VirtualProduction/ClonerEffector/Source/ClonerEffector/Private/Cloner/Attachments/CEClonerAttachmentTree.cpp b/Engine/Plugins/VirtualProduction/ClonerEffector/Source/ClonerEffector/Private/Cloner/Attachments/CEClonerAttachmentTree.cpp
--- a/Engine/Plugins/VirtualProduction/ClonerEffector/Source/ClonerEffector/Private/Cloner/Attachments/CEClonerAttachmentTree.cpp
+++ b/Engine/Plugins/VirtualProduction/ClonerEffector/Source/ClonerEffector/Private/Cloner/Attachments/CEClonerAttachmentTree.cpp
@@ -437,7 +437,7 @@ void FCEClonerAttachmentTree::UpdateAttachment(AActor* InActor, AActor* InParent
 
 		// Check materials change
 		TArray<TWeakObjectPtr<UMaterialInterface>> UnsupportedMaterials;
-		if (AttachmentItem->CheckMaterialsChanged(/** Update */true, &UnsupportedMaterials))
+		if (AttachmentItem->CheckMaterialsChanged(/** CheckAndUpdate */true, &UnsupportedMaterials))
 		{
 			// Show warning for unset materials
 			if (!UnsupportedMaterials.IsEmpty())
@@ -448,6 +448,8 @@ void FCEClonerAttachmentTree::UpdateAttachment(AActor* InActor, AActor* InParent
 				}
 			}
 
+			BehaviorImplementation->UpdateMesh(*this, InActor);
+
 			if (AttachmentItem->bSetupDone)
 			{
 				AttachmentItem->Status = ECEClonerAttachmentStatus::Outdated;
diff --git a/Engine/Plugins/VirtualProduction/ClonerEffector/Source/ClonerEffector/Private/Cloner/Attachments/CEClonerAttachmentTreeBehavior.cpp b/Engine/Plugins/VirtualProduction/ClonerEffector/Source/ClonerEffector/Private/Cloner/Attachments/CEClonerAttachmentTreeBehavior.cpp
--- a/Engine/Plugins/VirtualProduction/ClonerEffector/Source/ClonerEffector/Private/Cloner/Attachments/CEClonerAttachmentTreeBehavior.cpp
+++ b/Engine/Plugins/VirtualProduction/ClonerEffector/Source/ClonerEffector/Private/Cloner/Attachments/CEClonerAttachmentTreeBehavior.cpp
@@ -220,6 +220,11 @@ void FCEClonerAttachmentGroupBehavior::OnItemDetached(FCEClonerAttachmentTree& I
 	}
 }
 
+void FCEClonerAttachmentGroupBehavior::UpdateMesh(FCEClonerAttachmentTree& InTree, AActor* InActor)
+{
+	RefreshBakedMeshRootMaterials(InTree, InActor);
+}
+
 void FCEClonerAttachmentFlatBehavior::GetOrderedChildrenActors(FCEClonerAttachmentTree& InTree, AActor* InActor, TArray<AActor*>& OutChildren) const
 {
 	if (InActor != InTree.GetAttachmentRoot())
diff --git a/Engine/Plugins/VirtualProduction/ClonerEffector/Source/ClonerEffector/Public/Cloner/Attachments/CEClonerAttachmentTreeBehavior.h b/Engine/Plugins/VirtualProduction/ClonerEffector/Source/ClonerEffector/Public/Cloner/Attachments/CEClonerAttachmentTreeBehavior.h
--- a/Engine/Plugins/VirtualProduction/ClonerEffector/Source/ClonerEffector/Public/Cloner/Attachments/CEClonerAttachmentTreeBehavior.h
+++ b/Engine/Plugins/VirtualProduction/ClonerEffector/Source/ClonerEffector/Public/Cloner/Attachments/CEClonerAttachmentTreeBehavior.h
@@ -29,6 +29,9 @@ public:
 	/** Attachment mesh status is out of date and needs an update */
 	virtual void InvalidateMesh(FCEClonerAttachmentTree& InTree, AActor* InActor) = 0;
 
+	/** Update attachment mesh (materials) */
+	virtual void UpdateMesh(FCEClonerAttachmentTree& InTree, AActor* InActor) = 0;
+
 	/** Item is attached in the tree */
 	virtual void OnItemAttached(FCEClonerAttachmentTree& InTree, FCEClonerAttachmentItem& InItem) = 0;
 
@@ -47,6 +50,7 @@ protected:
 	virtual void InvalidateMesh(FCEClonerAttachmentTree& InTree, AActor* InActor) override;
 	virtual void OnItemAttached(FCEClonerAttachmentTree& InTree, FCEClonerAttachmentItem& InItem) override;
 	virtual void OnItemDetached(FCEClonerAttachmentTree& InTree, FCEClonerAttachmentItem& InItem) override;
+	virtual void UpdateMesh(FCEClonerAttachmentTree& InTree, AActor* InActor) override;
 	//~ End ICEClonerAttachmentImplementation
 
 	void BindActorDelegates(FCEClonerAttachmentTree& InTree, AActor* InActor);
