Motion Design: Repurposed 'add' button in md sequence to now give additional options for presets for faster creation of sequences.
In addition to that, new motion design levels auto-create 4 sequences for TL.

diff --git a/Engine/Plugins/VirtualProduction/Avalanche/Source/Avalanche/Private/AvaScene.cpp b/Engine/Plugins/VirtualProduction/Avalanche/Source/Avalanche/Private/AvaScene.cpp
--- a/Engine/Plugins/VirtualProduction/Avalanche/Source/Avalanche/Private/AvaScene.cpp
+++ b/Engine/Plugins/VirtualProduction/Avalanche/Source/Avalanche/Private/AvaScene.cpp
@@ -193,21 +193,36 @@ bool AAvaScene::CreateDirectorInstance(UAvaSequence& InSequence
 
 bool AAvaScene::AddSequence(UAvaSequence* InSequence)
 {
-	if (IsValid(InSequence) && !Animations.Contains(InSequence))
-	{
-		Animations.Add(InSequence);
-		ScheduleRebuildSequenceTree();
+	return AddSequences(MakeArrayView(&InSequence, 1)) > 0;
+}
 
-#if WITH_EDITOR
-		if (FEngineAnalytics::IsAvailable())
+uint32 AAvaScene::AddSequences(TConstArrayView<UAvaSequence*> InSequences)
+{
+	uint32 AddedCount = 0;
+
+	for (UAvaSequence* Sequence : InSequences)
+	{
+		if (Sequence && !Animations.Contains(Sequence))
 		{
-			FEngineAnalytics::GetProvider().RecordEvent(TEXT("Editor.Usage.MotionDesign.Sequence"), TEXT("Action"), TEXT("AddedSequence"));
+			Animations.Add(Sequence);
+			++AddedCount;
 		}
-#endif
+	}
+
+	if (AddedCount == 0)
+	{
+		return false;
+	}
 
-		return true;
+	ScheduleRebuildSequenceTree();
+
+#if WITH_EDITOR
+	if (FEngineAnalytics::IsAvailable())
+	{
+		FEngineAnalytics::GetProvider().RecordEvent(TEXT("Editor.Usage.MotionDesign.Sequence"), TEXT("Action"), TEXT("AddedSequence"));
 	}
-	return false;
+#endif
+	return true;
 }
 
 void AAvaScene::RemoveSequence(UAvaSequence* InSequence)
diff --git a/Engine/Plugins/VirtualProduction/Avalanche/Source/Avalanche/Public/AvaScene.h b/Engine/Plugins/VirtualProduction/Avalanche/Source/Avalanche/Public/AvaScene.h
--- a/Engine/Plugins/VirtualProduction/Avalanche/Source/Avalanche/Public/AvaScene.h
+++ b/Engine/Plugins/VirtualProduction/Avalanche/Source/Avalanche/Public/AvaScene.h
@@ -80,6 +80,7 @@ public:
 	virtual UWorld* GetContextWorld() const override;
 	virtual bool CreateDirectorInstance(UAvaSequence& InSequence, IMovieScenePlayer& InPlayer, const FMovieSceneSequenceID& InSequenceID, UObject*& OutDirectorInstance) override;
 	virtual bool AddSequence(UAvaSequence* InSequence) override;
+	virtual uint32 AddSequences(TConstArrayView<UAvaSequence*> InSequences) override;
 	virtual void RemoveSequence(UAvaSequence* InSequence) override;
 	virtual void SetDefaultSequence(UAvaSequence* InSequence) override;
 	virtual UAvaSequence* GetDefaultSequence() const override;
diff --git a/Engine/Plugins/VirtualProduction/Avalanche/Source/AvalancheEditor/Private/Sequencer/AvaSequencerExtension.cpp b/Engine/Plugins/VirtualProduction/Avalanche/Source/AvalancheEditor/Private/Sequencer/AvaSequencerExtension.cpp
--- a/Engine/Plugins/VirtualProduction/Avalanche/Source/AvalancheEditor/Private/Sequencer/AvaSequencerExtension.cpp
+++ b/Engine/Plugins/VirtualProduction/Avalanche/Source/AvalancheEditor/Private/Sequencer/AvaSequencerExtension.cpp
@@ -14,6 +14,7 @@
 #include "LevelEditor.h"
 #include "Misc/PackageName.h"
 #include "MovieScene.h"
+#include "Settings/AvaSequencerSettings.h"
 #include "StatusBarSubsystem.h"
 #include "UObject/Package.h"
 #include "Viewport/AvaViewportExtension.h"
@@ -67,6 +68,7 @@ void FAvaSequencerExtension::Activate()
 		if (UAvaSequencerSubsystem* SequencerSubsystem = World->GetSubsystem<UAvaSequencerSubsystem>())
 		{
 			AvaSequencer = SequencerSubsystem->GetOrCreateSequencer(*this, MakeSequencerArgs());
+			ConditionallyCreateDefaultSequences();
 			ValidateDirectorBlueprints();
 		}
 		else
@@ -107,6 +109,11 @@ void FAvaSequencerExtension::Cleanup()
 	}
 }
 
+void FAvaSequencerExtension::OnSceneObjectCreated(UObject* InSceneObject)
+{
+	bCreateDefaultSequences = true;
+}
+
 void FAvaSequencerExtension::OnSceneObjectChanged(UObject* InOldSceneObject, UObject* InNewSceneObject)
 {
 	if (IAvaSequenceProvider* SequenceProvider = Cast<IAvaSequenceProvider>(InNewSceneObject))
@@ -306,6 +313,24 @@ void FAvaSequencerExtension::UnregisterSequenceDrawerWidget()
 	}
 }
 
+void FAvaSequencerExtension::ConditionallyCreateDefaultSequences()
+{
+	if (!bCreateDefaultSequences || !AvaSequencer.IsValid())
+	{
+		return;
+	}
+
+	const UAvaSequencerSettings* const SequencerSettings = GetDefault<UAvaSequencerSettings>();
+	if (!SequencerSettings)
+	{
+		return;
+	}
+
+	const TArray<const FAvaSequencePreset*> SequencePresets = SequencerSettings->GatherPresetsFromGroup(UAvaSequencerSettings::TransitionLogicStarterGroup);
+	AvaSequencer->AddSequenceFromPresets(SequencePresets);
+	bCreateDefaultSequences = false;
+}
+
 void FAvaSequencerExtension::ValidateDirectorBlueprints()
 {
 	IAvaSequenceProvider* SequenceProvider = GetSequenceProvider();
diff --git a/Engine/Plugins/VirtualProduction/Avalanche/Source/AvalancheEditor/Private/Sequencer/AvaSequencerExtension.h b/Engine/Plugins/VirtualProduction/Avalanche/Source/AvalancheEditor/Private/Sequencer/AvaSequencerExtension.h
--- a/Engine/Plugins/VirtualProduction/Avalanche/Source/AvalancheEditor/Private/Sequencer/AvaSequencerExtension.h
+++ b/Engine/Plugins/VirtualProduction/Avalanche/Source/AvalancheEditor/Private/Sequencer/AvaSequencerExtension.h
@@ -37,6 +37,7 @@ public:
 	virtual void Activate() override;
 	virtual void Deactivate() override;
 	virtual void Cleanup() override;
+	virtual void OnSceneObjectCreated(UObject* InSceneObject) override;
 	virtual void OnSceneObjectChanged(UObject* InOldSceneObject, UObject* InNewSceneObject) override;
 	virtual void RegisterTabSpawners(const TSharedRef<IAvaEditor>& InEditor) const override;
 	virtual void ExtendLevelEditorLayout(FLayoutExtender& InExtender) const override;
@@ -69,6 +70,9 @@ private:
 
 	void UnregisterSequenceDrawerWidget();
 
+	/** Creates the default sequences gotten from the default preset group to use */
+	void ConditionallyCreateDefaultSequences();
+
 	/** Temporary fix ensuring that any Director Blueprint that got renamed to the incorrect outer gets fixed to the correct outer*/
 	void ValidateDirectorBlueprints();
 
@@ -78,4 +82,10 @@ private:
 	const FName SequencerTabId;
 
 	const bool bSupportsDrawerWidget;
+
+	/**
+	 * Flag to indicate whether to create the default sequences when activating
+	 * Set to true when the scene object is first created, but before the editor itself is activated.
+	 */
+	bool bCreateDefaultSequences = false;
 };
diff --git a/Engine/Plugins/VirtualProduction/Avalanche/Source/AvalancheEditorCore/Private/AvaEditor.cpp b/Engine/Plugins/VirtualProduction/Avalanche/Source/AvalancheEditorCore/Private/AvaEditor.cpp
--- a/Engine/Plugins/VirtualProduction/Avalanche/Source/AvalancheEditorCore/Private/AvaEditor.cpp
+++ b/Engine/Plugins/VirtualProduction/Avalanche/Source/AvalancheEditorCore/Private/AvaEditor.cpp
@@ -451,13 +451,27 @@ UObject* FAvaEditor::GetSceneObject(EAvaEditorObjectQueryType InQueryType) const
 	}
 
 	UObject* SceneObject = nullptr;
+	bool bWasCreated = false;
 
 	if (UWorld* SceneWorld = GetWorld())
 	{
-		SceneObject = Provider->GetSceneObject(SceneWorld, InQueryType);
+		// Check that the query types from here on are either find or create
+		check(InQueryType == EAvaEditorObjectQueryType::SearchOnly || InQueryType == EAvaEditorObjectQueryType::CreateIfNotFound);
+
+		SceneObject = Provider->GetSceneObject(SceneWorld, EAvaEditorObjectQueryType::SearchOnly);
+
+		if (!SceneObject && InQueryType == EAvaEditorObjectQueryType::CreateIfNotFound)
+		{
+			SceneObject = Provider->GetSceneObject(SceneWorld, EAvaEditorObjectQueryType::CreateIfNotFound);
+			bWasCreated = true;
+		}
 	}
 
 	const_cast<FAvaEditor*>(this)->SceneObjectWeak = SceneObject;
+	if (bWasCreated)
+	{
+		NotifySceneObjectCreated(SceneObject);
+	}
 	return SceneObject;
 }
 
@@ -798,6 +812,15 @@ void FAvaEditor::DeactivateExtensions()
 		});
 }
 
+void FAvaEditor::NotifySceneObjectCreated(UObject* InSceneObject) const
+{
+	ForEachExtension(
+		[InSceneObject](const TSharedRef<IAvaEditorExtension>& InExtension)
+		{
+			InExtension->OnSceneObjectCreated(InSceneObject);
+		});
+}
+
 bool FAvaEditor::CopyToString(FString& OutCopyData) const
 {
 	UE::AvaEditorCore::Private::FAvaActorEditActionScope ActionScope(*this);
diff --git a/Engine/Plugins/VirtualProduction/Avalanche/Source/AvalancheEditorCore/Private/AvaEditor.h b/Engine/Plugins/VirtualProduction/Avalanche/Source/AvalancheEditorCore/Private/AvaEditor.h
--- a/Engine/Plugins/VirtualProduction/Avalanche/Source/AvalancheEditorCore/Private/AvaEditor.h
+++ b/Engine/Plugins/VirtualProduction/Avalanche/Source/AvalancheEditorCore/Private/AvaEditor.h
@@ -76,6 +76,9 @@ protected:
 	void DeactivateExtensions();
 
 private:
+	/** Called when the scene object has been created for the first time */
+	void NotifySceneObjectCreated(UObject* InSceneObject) const;
+
 	bool CopyToString(FString& OutCopyData) const;
 
 	bool PasteFromString(FString& InPastedData) const;
diff --git a/Engine/Plugins/VirtualProduction/Avalanche/Source/AvalancheEditorCore/Public/IAvaEditorExtension.h b/Engine/Plugins/VirtualProduction/Avalanche/Source/AvalancheEditorCore/Public/IAvaEditorExtension.h
--- a/Engine/Plugins/VirtualProduction/Avalanche/Source/AvalancheEditorCore/Public/IAvaEditorExtension.h
+++ b/Engine/Plugins/VirtualProduction/Avalanche/Source/AvalancheEditorCore/Public/IAvaEditorExtension.h
@@ -51,6 +51,10 @@ public:
 
 	virtual void BindCommands(const TSharedRef<FUICommandList>& InCommandList) {}
 
+	/** Called when the scene object has been created */
+	virtual void OnSceneObjectCreated(UObject* InSceneObject) {}
+
+	/** Called when the scene object has been changed/replaced*/
 	virtual void OnSceneObjectChanged(UObject* InOldSceneObject, UObject* InNewSceneObject) {}
 
 	virtual void RegisterTabSpawners(const TSharedRef<IAvaEditor>& InEditor) const {}
diff --git a/Engine/Plugins/VirtualProduction/Avalanche/Source/AvalancheSequence/Public/IAvaSequenceProvider.h b/Engine/Plugins/VirtualProduction/Avalanche/Source/AvalancheSequence/Public/IAvaSequenceProvider.h
--- a/Engine/Plugins/VirtualProduction/Avalanche/Source/AvalancheSequence/Public/IAvaSequenceProvider.h
+++ b/Engine/Plugins/VirtualProduction/Avalanche/Source/AvalancheSequence/Public/IAvaSequenceProvider.h
@@ -54,6 +54,13 @@ public:
 
 	virtual bool AddSequence(UAvaSequence* InSequence) = 0;
 
+	/**
+	 * Adds the given sequences to the sequence provider
+	 * @param InSequences the sequences to add, if valid and not already in the list.
+	 * @return the number of sequences that were successfully added.
+	 */
+	virtual uint32 AddSequences(TConstArrayView<UAvaSequence*> InSequences) = 0;
+
 	virtual void RemoveSequence(UAvaSequence* InSequence) = 0;
 
 	/** Sets the Default Sequence to use when no particular sequence is active */
diff --git a/Engine/Plugins/VirtualProduction/Avalanche/Source/AvalancheSequencer/Private/AvaSequencer.cpp b/Engine/Plugins/VirtualProduction/Avalanche/Source/AvalancheSequencer/Private/AvaSequencer.cpp
--- a/Engine/Plugins/VirtualProduction/Avalanche/Source/AvalancheSequencer/Private/AvaSequencer.cpp
+++ b/Engine/Plugins/VirtualProduction/Avalanche/Source/AvalancheSequencer/Private/AvaSequencer.cpp
@@ -1337,6 +1337,68 @@ void FAvaSequencer::DeleteSequences(const TSet<UAvaSequence*>& InSequences)
 	}
 }
 
+uint32 FAvaSequencer::AddSequenceFromPresets(TConstArrayView<const FAvaSequencePreset*> InPresets)
+{
+	IAvaSequenceProvider* const SequenceProvider = Provider.GetSequenceProvider();
+	if (!SequenceProvider)
+	{
+		return 0;
+	}
+
+	TArray<UAvaSequence*> NewSequences;
+	NewSequences.Reserve(InPresets.Num());
+
+	for (const FAvaSequencePreset* Preset : InPresets)
+	{
+		if (!Preset)
+		{
+			continue;
+		}
+
+		if (UAvaSequence* const Sequence = CreateSequence())
+		{
+			Preset->ApplyPreset(Sequence);
+			NewSequences.Add(Sequence);
+		}
+	}
+
+	return SequenceProvider->AddSequences(NewSequences);
+}
+
+void FAvaSequencer::AddSequencesFromPresetGroup(FName InPresetGroup)
+{
+	const UAvaSequencerSettings* SequencerSettings = GetDefault<UAvaSequencerSettings>();
+	if (!SequencerSettings)
+	{
+		return;
+	}
+
+	TArray<const FAvaSequencePreset*> Presets = SequencerSettings->GatherPresetsFromGroup(InPresetGroup);
+	if (Presets.IsEmpty())
+	{
+		return;
+	}
+
+	IAvaSequenceProvider* const SequenceProvider = Provider.GetSequenceProvider();
+	if (!SequenceProvider)
+	{
+		return;
+	}
+
+	FScopedTransaction Transaction(LOCTEXT("AddSequencesFromPresetGroup", "Add Sequences from Preset Group"));
+
+	if (UObject* SequenceProviderObject = SequenceProvider->ToUObject())
+	{
+		SequenceProviderObject->Modify();
+	}
+
+	const uint32 AddedCount = AddSequenceFromPresets(Presets);
+	if (AddedCount == 0)
+	{
+		Transaction.Cancel();
+	}
+}
+
 UObject* FAvaSequencer::FindObjectToPossess(UObject* InParentObject, const FMovieScenePossessable& InPossessable)
 {
 	if (!InParentObject)
diff --git a/Engine/Plugins/VirtualProduction/Avalanche/Source/AvalancheSequencer/Private/AvaSequencer.h b/Engine/Plugins/VirtualProduction/Avalanche/Source/AvalancheSequencer/Private/AvaSequencer.h
--- a/Engine/Plugins/VirtualProduction/Avalanche/Source/AvalancheSequencer/Private/AvaSequencer.h
+++ b/Engine/Plugins/VirtualProduction/Avalanche/Source/AvalancheSequencer/Private/AvaSequencer.h
@@ -171,7 +171,6 @@ public:
 	TArrayView<TWeakObjectPtr<>> ResolveBoundObjects(const FGuid& InBindingId, UMovieSceneSequence* Sequence) const;
 
 	//~ Begin IAvaSequencer
-
 	virtual const IAvaSequencerProvider& GetProvider() const override { return Provider; }
 	virtual TSharedRef<ISequencer> GetSequencer() const override;
 	virtual USequencerSettings* GetSequencerSettings() const override;
@@ -186,15 +185,15 @@ public:
 	virtual void OnEditorSelectionChanged(const FAvaEditorSelection& InEditorSelection) override;
 	virtual void NotifyOnSequenceTreeChanged() override;
 	virtual const TArray<TWeakObjectPtr<UAvaSequence>>& GetRootSequences() const override;
-
 	virtual bool CanAddSequence() const override;
 	virtual UAvaSequence* AddSequence(UAvaSequence* const InParentSequence = nullptr) override;
-
 	virtual void DeleteSequences(const TSet<UAvaSequence*>& InSequences) override;
-
+	virtual uint32 AddSequenceFromPresets(TConstArrayView<const FAvaSequencePreset*> InPresets) override;
 	//~ End IAvaSequencer
 
-	// BEGIN REMOVE AFTER OLD TREE VIEW REMOVAL
+	/** Creates new sequences based on a given preset group */
+	void AddSequencesFromPresetGroup(FName InPresetGroup);
+
 	bool AddSequence_CanExecute() const;
 	void AddSequence_Execute();
 
@@ -221,7 +220,6 @@ public:
 
 	bool StopSequences_CanExecute() const;
 	void StopSequences_Execute();
-	// END REMOVE AFTER OLD TREE VIEW REMOVAL
 
 	/** Called when PostUndo has been called for the given sequence */
 	void OnSequenceEditUndo(UAvaSequence* InSequence);
diff --git a/Engine/Plugins/VirtualProduction/Avalanche/Source/AvalancheSequencer/Private/SequenceTree/Widgets/SAvaSequenceTree.cpp b/Engine/Plugins/VirtualProduction/Avalanche/Source/AvalancheSequencer/Private/SequenceTree/Widgets/SAvaSequenceTree.cpp
--- a/Engine/Plugins/VirtualProduction/Avalanche/Source/AvalancheSequencer/Private/SequenceTree/Widgets/SAvaSequenceTree.cpp
+++ b/Engine/Plugins/VirtualProduction/Avalanche/Source/AvalancheSequencer/Private/SequenceTree/Widgets/SAvaSequenceTree.cpp
@@ -113,7 +113,7 @@ void SAvaSequenceTree::Construct(const FArguments& InArgs
 				.AutoWidth()
 				[
 					SNew(SPositiveActionButton)
-					.OnClicked(this, &SAvaSequenceTree::OnNewSequenceClicked)
+					.OnGetMenuContent(this, &SAvaSequenceTree::CreateAddSequenceMenu)
 					.Text(LOCTEXT("AddSequenceText", "Add"))
 				]
 				+ SHorizontalBox::Slot()
@@ -219,6 +219,10 @@ void SAvaSequenceTree::BindCommands(const TSharedRef<FAvaSequencer>& InSequencer
 		, FExecuteAction::CreateSP(InSequencer, &FAvaSequencer::StopSequences_Execute)
 		, FCanExecuteAction::CreateSP(InSequencer, &FAvaSequencer::StopSequences_CanExecute));
 
+	CommandList->MapAction(AvaSequencerCommands.AddNew
+		, FExecuteAction::CreateSP(InSequencer, &FAvaSequencer::AddSequence_Execute)
+		, FCanExecuteAction::CreateSP(InSequencer, &FAvaSequencer::AddSequence_CanExecute));
+
 	CommandList->MapAction(AvaSequencerCommands.ExportSequence
 		, FExecuteAction::CreateSP(InSequencer, &FAvaSequencer::ExportSequences_Execute)
 		, FCanExecuteAction::CreateSP(InSequencer, &FAvaSequencer::ExportSequences_CanExecute));
@@ -320,14 +324,87 @@ FReply SAvaSequenceTree::OnKeyDown(const FGeometry& InGeometry, const FKeyEvent&
 	return FReply::Unhandled();
 }
 
-FReply SAvaSequenceTree::OnNewSequenceClicked()
+void SAvaSequenceTree::CreateAddSequencePresetMenu(UToolMenu* InMenu)
 {
-	if (SequencerWeak.IsValid())
+	const UAvaSequencerMenuContext* const MenuContext = InMenu->FindContext<UAvaSequencerMenuContext>();
+	if (!MenuContext)
 	{
-		SequencerWeak.Pin()->AddSequence_Execute();
-		return FReply::Handled();
+		return;
 	}
-	return FReply::Unhandled();
+
+	const TSharedPtr<FAvaSequencer> AvaSequencer = MenuContext->GetAvaSequencer();
+	if (!AvaSequencer.IsValid())
+	{
+		return;
+	}
+
+	const UAvaSequencerSettings* const SequencerSettings = GetDefault<UAvaSequencerSettings>();
+	if (!SequencerSettings)
+	{
+		return;
+	}
+
+	auto AddGroup = [AvaSequenceRef = AvaSequencer.ToSharedRef()](FToolMenuSection& InSection, const FAvaSequencePresetGroup& InGroup)
+		{
+			FUIAction Action;
+			Action.ExecuteAction.BindSP(AvaSequenceRef, &FAvaSequencer::AddSequencesFromPresetGroup, InGroup.GroupName);
+			Action.CanExecuteAction.BindSP(AvaSequenceRef, &FAvaSequencer::CanAddSequence);
+
+			InSection.AddMenuEntry(InGroup.GroupName
+				, FText::FromName(InGroup.GroupName)
+				, FText::FromName(InGroup.GroupName)
+				, FSlateIcon()
+				, FToolUIActionChoice(Action));
+		};
+
+	// Transition Logic Groups
+	FToolMenuSection& TransitionLogicSection = InMenu->FindOrAddSection(TEXT("TransitionLogicGroups"));
+	TransitionLogicSection.Label = LOCTEXT("TransitionLogicGroupsLabel", "Transition Logic");
+	for (const FAvaSequencePresetGroup& Group : SequencerSettings->GetTransitionLogicPresetGroups())
+	{
+		AddGroup(TransitionLogicSection, Group);
+	}
+
+	// Custom Groups
+	if (!SequencerSettings->GetCustomPresetGroups().IsEmpty())
+	{
+		FToolMenuSection& CustomSection = InMenu->FindOrAddSection(TEXT("CustomGroups"));
+		TransitionLogicSection.Label = LOCTEXT("CustomGroupsLabel", "Custom");
+		for (const FAvaSequencePresetGroup& Group : SequencerSettings->GetCustomPresetGroups())
+		{
+			AddGroup(CustomSection, Group);
+		}
+	}
+}
+
+TSharedRef<SWidget> SAvaSequenceTree::CreateAddSequenceMenu() const
+{
+	UToolMenus* ToolMenus = UToolMenus::Get();
+	if (!ToolMenus)
+	{
+		return SNullWidget::NullWidget;
+	}
+
+	constexpr const TCHAR* ContextMenuName = TEXT("SAvaSequenceAddNewContextMenu");
+
+	if (!ToolMenus->IsMenuRegistered(ContextMenuName))
+	{
+		UToolMenu* const ToolMenu = ToolMenus->RegisterMenu(ContextMenuName, NAME_None, EMultiBoxType::Menu);
+
+		FToolMenuSection& Section = ToolMenu->FindOrAddSection(NAME_None);
+		Section.AddMenuEntry(FAvaSequencerCommands::Get().AddNew
+			, LOCTEXT("AddSequenceLabel", "New sequence")
+			, LOCTEXT("AddSequenceTooltip", "Adds a new Motion Design sequence.")
+			, FSlateIcon(FAppStyle::GetAppStyleSetName(), TEXT("Icons.Plus")));
+
+		ToolMenu->AddDynamicSection(TEXT("SequencePresets"), FNewToolMenuDelegate::CreateStatic(&SAvaSequenceTree::CreateAddSequencePresetMenu));
+	}
+
+	UAvaSequencerMenuContext* MenuContext = NewObject<UAvaSequencerMenuContext>();
+	MenuContext->SetAvaSequencer(SequencerWeak);
+
+	FToolMenuContext ToolMenuContext(CommandList, /*Extender*/nullptr, MenuContext);	
+	return ToolMenus->GenerateWidget(ContextMenuName, ToolMenuContext);
 }
 
 #undef LOCTEXT_NAMESPACE
diff --git a/Engine/Plugins/VirtualProduction/Avalanche/Source/AvalancheSequencer/Private/SequenceTree/Widgets/SAvaSequenceTree.h b/Engine/Plugins/VirtualProduction/Avalanche/Source/AvalancheSequencer/Private/SequenceTree/Widgets/SAvaSequenceTree.h
--- a/Engine/Plugins/VirtualProduction/Avalanche/Source/AvalancheSequencer/Private/SequenceTree/Widgets/SAvaSequenceTree.h
+++ b/Engine/Plugins/VirtualProduction/Avalanche/Source/AvalancheSequencer/Private/SequenceTree/Widgets/SAvaSequenceTree.h
@@ -10,10 +10,11 @@ class FAvaSequencer;
 class FUICommandList;
 class ITableRow;
 class SHeaderRow;
-template<typename ItemType> class STreeView;
 class SSearchBox;
 class STableViewBase;
 class UAvaSequence;
+class UToolMenu;
+template<typename ItemType> class STreeView;
 
 class SAvaSequenceTree : public SCompoundWidget
 {
@@ -39,8 +40,6 @@ public:
 
 	void OnSelectionChanged(FAvaSequenceItemPtr InSelectedItem, ESelectInfo::Type InSelectionInfo);
 
-	FReply OnNewSequenceClicked();
-
 	void OnSearchChanged(const FText& InSearchText);
 
 	//~ Begin SWidget
@@ -48,6 +47,12 @@ public:
 	//~ End SWidget
 	
 protected:
+	/** Adds the preset options in the add sequence menu */
+	static void CreateAddSequencePresetMenu(UToolMenu* InMenu);
+
+	/** Creates the add sequence menu when the 'add' button gets clicked */
+	TSharedRef<SWidget> CreateAddSequenceMenu() const;
+
 	TWeakPtr<FAvaSequencer> SequencerWeak;
 	
 	TSharedPtr<STreeView<FAvaSequenceItemPtr>> SequenceTreeView;
diff --git a/Engine/Plugins/VirtualProduction/Avalanche/Source/AvalancheSequencer/Private/Settings/AvaSequencePreset.cpp b/Engine/Plugins/VirtualProduction/Avalanche/Source/AvalancheSequencer/Private/Settings/AvaSequencePreset.cpp
--- a/Engine/Plugins/VirtualProduction/Avalanche/Source/AvalancheSequencer/Private/Settings/AvaSequencePreset.cpp
+++ b/Engine/Plugins/VirtualProduction/Avalanche/Source/AvalancheSequencer/Private/Settings/AvaSequencePreset.cpp
@@ -48,12 +48,21 @@ void FAvaSequencePreset::ApplyPreset(UAvaSequence* InSequence) const
 		MovieScene->Modify();	
 	}
 
+	const FFrameRate TickResolution = MovieScene->GetTickResolution();
+
+	if (bEnableEndTime)
+	{
+		const FFrameTime EndFrame = EndTime * TickResolution;
+
+		const FFrameNumber StartFrame = FMath::Min(FFrameNumber(0), EndFrame.FrameNumber);
+		MovieScene->SetPlaybackRange(TRange<FFrameNumber>(StartFrame, EndFrame.FrameNumber+1));
+	}
+
 	if (bEnableMarks)
 	{
 		MovieScene->DeleteMarkedFrames();
 
 		const FFrameRate DisplayRate = MovieScene->GetDisplayRate();
-		const FFrameRate TickResolution = MovieScene->GetTickResolution();
 
 		for (const FAvaMarkSetting& MarkSetting : Marks)
 		{
diff --git a/Engine/Plugins/VirtualProduction/Avalanche/Source/AvalancheSequencer/Private/Settings/AvaSequencerSettings.cpp b/Engine/Plugins/VirtualProduction/Avalanche/Source/AvalancheSequencer/Private/Settings/AvaSequencerSettings.cpp
--- a/Engine/Plugins/VirtualProduction/Avalanche/Source/AvalancheSequencer/Private/Settings/AvaSequencerSettings.cpp
+++ b/Engine/Plugins/VirtualProduction/Avalanche/Source/AvalancheSequencer/Private/Settings/AvaSequencerSettings.cpp
@@ -6,14 +6,22 @@
 
 namespace UE::AvaSequencer::Private
 {
-	FAvaSequencePreset CreateDefaultPreset(FName InSequenceType, const FAvaTagSoftHandle& InSequenceTag)
+	FAvaSequencePreset CreateDefaultPreset(FName InSequenceType)
 	{
 		FAvaSequencePreset Preset;
-		Preset.PresetName    = InSequenceType;
-		Preset.SequenceLabel = InSequenceType;
-		Preset.SequenceTag   = InSequenceTag;
-		Preset.bEnableLabel  = true;
-		Preset.bEnableTag    = true;
+		Preset.PresetName     = InSequenceType;
+		Preset.bEnableLabel   = true;
+		Preset.SequenceLabel  = InSequenceType;
+		Preset.bEnableEndTime = true;
+		Preset.EndTime        = 1.0;
+		return Preset;
+	}
+
+	FAvaSequencePreset CreateDefaultPreset(FName InSequenceType, const FAvaTagSoftHandle& InSequenceTag)
+	{
+		FAvaSequencePreset Preset = CreateDefaultPreset(InSequenceType);
+		Preset.bEnableTag  = true;
+		Preset.SequenceTag = InSequenceTag;
 		return Preset;
 	}
 
@@ -22,27 +30,47 @@ namespace UE::AvaSequencer::Private
 		const FAvaSequenceDefaultTags& DefaultTags = FAvaSequenceDefaultTags::Get();
 
 		// Guids found via the entries that DefaultSequenceTags had for In/Out/Change
-		TArray<FAvaSequencePreset> DefaultPresets =
-			{
-				CreateDefaultPreset(TEXT("In")    , DefaultTags.In),
-				CreateDefaultPreset(TEXT("Out")   , DefaultTags.Out),
-				CreateDefaultPreset(TEXT("Change"), DefaultTags.Change),
-			};
-
-		// Configure Change Marks
+		return
 		{
-			FAvaSequencePreset& ChangePreset = DefaultPresets[2];
-			ChangePreset.bEnableMarks = true;
+			CreateDefaultPreset(TEXT("In"), DefaultTags.In),
+			CreateDefaultPreset(TEXT("Out"), DefaultTags.Out),
+			CreateDefaultPreset(TEXT("NewLayer_ChangeIn")),
+			CreateDefaultPreset(TEXT("NewLayer_ChangeOut")),
+		};
+	}
 
-			FAvaMarkSetting& MarkSetting = ChangePreset.Marks.AddDefaulted_GetRef();
-        	MarkSetting.Label = TEXT("A");
-			MarkSetting.FrameNumber = 60;
+	TArray<FAvaSequencePresetGroup> MakeTransitionLogicPresetGroups()
+	{
+		TArray<FAvaSequencePresetGroup> PresetGroups;
+
+		// Transition Logic Starter. 4 sequences: global in/out, and a sub layer's in/out pair
+		{
+			FAvaSequencePresetGroup& PresetGroup = PresetGroups.AddDefaulted_GetRef();
+			PresetGroup.GroupName = UAvaSequencerSettings::TransitionLogicStarterGroup;
+
+			PresetGroup.PresetNames.SetNum(4);
+			PresetGroup.PresetNames[0] = TEXT("In");
+			PresetGroup.PresetNames[1] = TEXT("Out");
+			PresetGroup.PresetNames[2] = TEXT("NewLayer_ChangeIn");
+			PresetGroup.PresetNames[3] = TEXT("NewLayer_ChangeOut");
 		}
 
-		return DefaultPresets;
+		// Transition Logic Sub Layer. 2 sequences: a sub layer's in/out pair
+		{
+			FAvaSequencePresetGroup& PresetGroup = PresetGroups.AddDefaulted_GetRef();
+			PresetGroup.GroupName = UAvaSequencerSettings::TransitionLogicSubLayerGroup;
+
+			PresetGroup.PresetNames.SetNum(2);
+			PresetGroup.PresetNames[0] = TEXT("NewLayer_ChangeIn");
+			PresetGroup.PresetNames[1] = TEXT("NewLayer_ChangeOut");
+		}
+		return PresetGroups;
 	}
 }
 
+const FLazyName UAvaSequencerSettings::TransitionLogicStarterGroup = TEXT("Transition Logic Starter");
+const FLazyName UAvaSequencerSettings::TransitionLogicSubLayerGroup = TEXT("Transition Logic Sub Layer");
+
 UAvaSequencerSettings::UAvaSequencerSettings()
 {
 	CategoryName = SettingsCategory;
@@ -51,8 +79,69 @@ UAvaSequencerSettings::UAvaSequencerSettings()
 	DisplayRate.FrameRate = FFrameRate(60000, 1001);
 }
 
+const FAvaSequencePreset* UAvaSequencerSettings::FindPreset(FName InPresetName) const
+{
+	const FAvaSequencePreset::FKeyType Key(InPresetName);
+
+	if (const FAvaSequencePreset* FoundPreset = CustomSequencePresets.FindByHash(GetTypeHash(Key), Key))
+	{
+		return FoundPreset;
+	}
+
+	if (const FAvaSequencePreset* FoundPreset = GetDefaultSequencePresets().FindByKey(Key))
+	{
+		return FoundPreset;
+	}
+
+	return nullptr;
+}
+
+const FAvaSequencePresetGroup* UAvaSequencerSettings::FindPresetGroup(FName InGroupName) const
+{
+	const FAvaSequencePresetGroup::FKeyType Key(InGroupName);
+
+	if (const FAvaSequencePresetGroup* FoundGroup = CustomPresetGroups.FindByHash(GetTypeHash(Key), Key))
+	{
+		return FoundGroup;
+	}
+
+	if (const FAvaSequencePresetGroup* FoundGroup = GetTransitionLogicPresetGroups().FindByKey(Key))
+	{
+		return FoundGroup;
+	}
+
+	return nullptr;
+}
+
+TArray<const FAvaSequencePreset*> UAvaSequencerSettings::GatherPresetsFromGroup(FName InGroupName) const
+{
+	const FAvaSequencePresetGroup* PresetGroup = FindPresetGroup(InGroupName);
+	if (!PresetGroup)
+	{
+		return {};
+	}
+
+	TArray<const FAvaSequencePreset*> Presets;
+	Presets.Reserve(PresetGroup->PresetNames.Num());
+
+	for (const FName PresetName : PresetGroup->PresetNames)
+	{
+		if (const FAvaSequencePreset* Preset = FindPreset(PresetName))
+		{
+			Presets.Add(Preset);
+		}
+	}
+	return Presets;
+}
+
 TConstArrayView<FAvaSequencePreset> UAvaSequencerSettings::GetDefaultSequencePresets() const
 {
 	static const TArray<FAvaSequencePreset> DefaultPresets = UE::AvaSequencer::Private::MakeDefaultSequencePresets();
 	return DefaultPresets;
 }
+
+TConstArrayView<FAvaSequencePresetGroup> UAvaSequencerSettings::GetTransitionLogicPresetGroups() const
+{
+	static const TArray<FAvaSequencePresetGroup> TransitionLogicGroups = UE::AvaSequencer::Private::MakeTransitionLogicPresetGroups();
+	return TransitionLogicGroups;
+}
diff --git a/Engine/Plugins/VirtualProduction/Avalanche/Source/AvalancheSequencer/Public/IAvaSequencer.h b/Engine/Plugins/VirtualProduction/Avalanche/Source/AvalancheSequencer/Public/IAvaSequencer.h
--- a/Engine/Plugins/VirtualProduction/Avalanche/Source/AvalancheSequencer/Public/IAvaSequencer.h
+++ b/Engine/Plugins/VirtualProduction/Avalanche/Source/AvalancheSequencer/Public/IAvaSequencer.h
@@ -21,6 +21,7 @@ class SWidget;
 class UAvaSequence;
 class UObject;
 class USequencerSettings;
+struct FAvaSequencePreset;
 
 class IAvaSequencer
 {
@@ -74,5 +75,8 @@ public:
 	virtual bool CanAddSequence() const = 0;
 	virtual UAvaSequence* AddSequence(UAvaSequence* const InParentSequence = nullptr) = 0;
 
+	/** Creates new sequences based on the given presets */
+	virtual uint32 AddSequenceFromPresets(TConstArrayView<const FAvaSequencePreset*> InPresets) = 0;
+
 	virtual void DeleteSequences(const TSet<UAvaSequence*>& InSequences) = 0;
 };
diff --git a/Engine/Plugins/VirtualProduction/Avalanche/Source/AvalancheSequencer/Public/Settings/AvaSequencePreset.h b/Engine/Plugins/VirtualProduction/Avalanche/Source/AvalancheSequencer/Public/Settings/AvaSequencePreset.h
--- a/Engine/Plugins/VirtualProduction/Avalanche/Source/AvalancheSequencer/Public/Settings/AvaSequencePreset.h
+++ b/Engine/Plugins/VirtualProduction/Avalanche/Source/AvalancheSequencer/Public/Settings/AvaSequencePreset.h
@@ -14,6 +14,8 @@ struct FAvaSequencePreset
 {
 	GENERATED_BODY()
 
+	using FKeyType = FName;
+
 	AVALANCHESEQUENCER_API FAvaSequencePreset() = default;
 	AVALANCHESEQUENCER_API virtual ~FAvaSequencePreset() = default;
 
@@ -33,6 +35,11 @@ struct FAvaSequencePreset
 		return PresetName == InOther.PresetName;
 	}
 
+	bool operator==(const FAvaSequencePreset::FKeyType& InKey) const
+	{
+		return PresetName == InKey;
+	}
+
 	friend uint32 GetTypeHash(const FAvaSequencePreset& InPreset)
 	{
 		return GetTypeHash(InPreset.PresetName);
@@ -49,15 +56,54 @@ struct FAvaSequencePreset
 	UPROPERTY(EditAnywhere, Category="Sequence", meta=(EditCondition="bEnableTag"))
 	FAvaTagSoftHandle SequenceTag;
 
+	/** Initial end time to set for a sequence, in seconds */
+	UPROPERTY(EditAnywhere, Category="Sequence", meta=(EditCondition="bEnableEndTime"))
+	double EndTime = 1.0;
+
 	UPROPERTY(EditAnywhere, Category="Sequence", meta=(EditCondition="bEnableMarks"))
 	TArray<FAvaMarkSetting> Marks;
 
-	UPROPERTY(meta = (InlineEditConditionToggle))
+	UPROPERTY(meta=(InlineEditConditionToggle))
 	bool bEnableLabel = false;
 
-	UPROPERTY(meta = (InlineEditConditionToggle))
+	UPROPERTY(meta=(InlineEditConditionToggle))
 	bool bEnableTag = false;
 
-	UPROPERTY(meta = (InlineEditConditionToggle))
+	UPROPERTY(meta=(InlineEditConditionToggle))
+    bool bEnableEndTime = false;
+
+	UPROPERTY(meta=(InlineEditConditionToggle))
 	bool bEnableMarks = false;
 };
+
+/** Represents one or many presets to apply */
+USTRUCT()
+struct FAvaSequencePresetGroup
+{
+	GENERATED_BODY()
+
+	using FKeyType = FName;
+	
+	bool operator==(const FAvaSequencePresetGroup& InOther) const
+	{
+		return GroupName == InOther.GroupName;
+	}
+
+	bool operator==(const FAvaSequencePreset::FKeyType& InKey) const
+	{
+		return GroupName == InKey;
+	}
+
+	friend uint32 GetTypeHash(const FAvaSequencePresetGroup& InPresetGroup)
+	{
+		return GetTypeHash(InPresetGroup.GroupName);
+	}
+
+	/** Identifier of this group */
+	UPROPERTY(EditAnywhere, Category="Sequence")
+	FName GroupName;
+
+	/** The preset that this group represents */
+	UPROPERTY(EditAnywhere, Category="Sequence")
+	TArray<FName> PresetNames;
+};
diff --git a/Engine/Plugins/VirtualProduction/Avalanche/Source/AvalancheSequencer/Public/Settings/AvaSequencerSettings.h b/Engine/Plugins/VirtualProduction/Avalanche/Source/AvalancheSequencer/Public/Settings/AvaSequencerSettings.h
--- a/Engine/Plugins/VirtualProduction/Avalanche/Source/AvalancheSequencer/Public/Settings/AvaSequencerSettings.h
+++ b/Engine/Plugins/VirtualProduction/Avalanche/Source/AvalancheSequencer/Public/Settings/AvaSequencerSettings.h
@@ -9,7 +9,7 @@
 #include "UObject/Object.h"
 #include "AvaSequencerSettings.generated.h"
 
-UCLASS(config=EditorPerProjectUserSettings, meta = (DisplayName = "Sequencer"))
+UCLASS(MinimalAPI, config=EditorPerProjectUserSettings, meta = (DisplayName = "Sequencer"))
 class UAvaSequencerSettings : public UDeveloperSettings
 {
 	GENERATED_BODY()
@@ -18,6 +18,10 @@ public:
 	static constexpr const TCHAR* SettingsCategory = TEXT("Motion Design");
 	static constexpr const TCHAR* SettingsSection  = TEXT("Sequencer");
 
+	/** Group name for the built-in default groups */
+	AVALANCHESEQUENCER_API static const FLazyName TransitionLogicStarterGroup;
+	AVALANCHESEQUENCER_API static const FLazyName TransitionLogicSubLayerGroup;
+
 	UAvaSequencerSettings();
 
 	FFrameRate GetDisplayRate() const
@@ -35,13 +39,29 @@ public:
 		return EndTime;
 	}
 
+	/** Attempts to find the preset with the given preset name, first searching in the custom presets, then looking at the default presets */
+	const FAvaSequencePreset* FindPreset(FName InPresetName) const;
+
+	/** Attempts to find the preset group with the given preset name, first searching in the custom groups, then looking at the default (TL) groups */
+	const FAvaSequencePresetGroup* FindPresetGroup(FName InGroupName) const;
+
+	/** Finds all the presets within a preset group */
+	AVALANCHESEQUENCER_API TArray<const FAvaSequencePreset*> GatherPresetsFromGroup(FName InGroupName) const;
+
 	AVALANCHESEQUENCER_API TConstArrayView<FAvaSequencePreset> GetDefaultSequencePresets() const;
 
+	AVALANCHESEQUENCER_API TConstArrayView<FAvaSequencePresetGroup> GetTransitionLogicPresetGroups() const;
+
 	const TSet<FAvaSequencePreset>& GetCustomSequencePresets() const
 	{
 		return CustomSequencePresets;
 	}
 
+	const TSet<FAvaSequencePresetGroup>& GetCustomPresetGroups() const
+	{
+		return CustomPresetGroups;
+	}
+
 	UE_DEPRECATED(5.6, "The Motion Design sidebar drawer was moved to the Sequencer sidebar and now stores its state in USequencerSettings. Use USequencerSettings::GetSidebarState() instead.")
 	FSidebarState& GetSidebarState()
 	{
@@ -71,6 +91,10 @@ private:
 	UPROPERTY(Config, EditAnywhere, Category = "Sequencer")
 	TSet<FAvaSequencePreset> CustomSequencePresets;
 
+	/** Sequence Preset groups that will be served as options when creating sequences */
+	UPROPERTY(Config, EditAnywhere, Category = "Sequencer")
+	TSet<FAvaSequencePresetGroup> CustomPresetGroups;
+
 	/** The state of a sidebar to be restored when Sequencer is initialized */
 	UPROPERTY(Config, meta = (DeprecatedProperty, DeprecationMessage = "The Motion Design sidebar drawer was moved to the Sequencer sidebar and now stores its state in USequencerSettings"))
 	FSidebarState SidebarState_DEPRECATED;
