Default TL Layer

#virtualized

diff --git a/Engine/Plugins/VirtualProduction/Avalanche/Source/AvalancheTransition/Private/AvaTransitionDefaultLayerTags.cpp b/Engine/Plugins/VirtualProduction/Avalanche/Source/AvalancheTransition/Private/AvaTransitionDefaultLayerTags.cpp
new file mode 100644
--- /dev/null
+++ b/Engine/Plugins/VirtualProduction/Avalanche/Source/AvalancheTransition/Private/AvaTransitionDefaultLayerTags.cpp
@@ -0,0 +1,20 @@
+// Copyright Epic Games, Inc. All Rights Reserved.
+
+#include "AvaTransitionDefaultLayerTags.h"
+
+namespace UE::AvaTransition
+{
+
+FDefaultTags::FDefaultTags()
+{
+	TSoftObjectPtr<UAvaTagCollection> TagCollection(FSoftObjectPath(TEXT("/Avalanche/Tags/DefaultTransitionLayerTags.DefaultTransitionLayerTags")));
+	DefaultLayer = FAvaTagSoftHandle(TagCollection, FAvaTagId(FGuid(0x89084F7E, 0x4D80443E, 0xAC48609D, 0x7A348C8A)));
+}
+
+const FDefaultTags& FDefaultTags::Get()
+{
+	static const FDefaultTags DefaultTags;
+	return DefaultTags;
+}
+
+} // UE::AvaTransition
diff --git a/Engine/Plugins/VirtualProduction/Avalanche/Source/AvalancheTransition/Public/AvaTransitionDefaultLayerTags.h b/Engine/Plugins/VirtualProduction/Avalanche/Source/AvalancheTransition/Public/AvaTransitionDefaultLayerTags.h
new file mode 100644
--- /dev/null
+++ b/Engine/Plugins/VirtualProduction/Avalanche/Source/AvalancheTransition/Public/AvaTransitionDefaultLayerTags.h
@@ -0,0 +1,21 @@
+// Copyright Epic Games, Inc. All Rights Reserved.
+
+#pragma once
+
+#include "AvaTagSoftHandle.h"
+
+namespace UE::AvaTransition
+{
+
+class FDefaultTags
+{
+	FDefaultTags();
+
+public:
+	AVALANCHETRANSITION_API static const FDefaultTags& Get();
+
+	/** Soft handle to the default layer tag */
+	FAvaTagSoftHandle DefaultLayer;
+};
+
+} // UE::AvaTransition
diff --git a/Engine/Plugins/VirtualProduction/Avalanche/Source/AvalancheTransitionEditor/Private/AvaTransitionTreeEditorData.cpp b/Engine/Plugins/VirtualProduction/Avalanche/Source/AvalancheTransitionEditor/Private/AvaTransitionTreeEditorData.cpp
--- a/Engine/Plugins/VirtualProduction/Avalanche/Source/AvalancheTransitionEditor/Private/AvaTransitionTreeEditorData.cpp
+++ b/Engine/Plugins/VirtualProduction/Avalanche/Source/AvalancheTransitionEditor/Private/AvaTransitionTreeEditorData.cpp
@@ -1,9 +1,13 @@
 // Copyright Epic Games, Inc. All Rights Reserved.
 
 #include "AvaTransitionTreeEditorData.h"
+#include "AvaTransitionDefaultLayerTags.h"
 
 UAvaTransitionTreeEditorData::UAvaTransitionTreeEditorData()
 {
+	const UE::AvaTransition::FDefaultTags& DefaultTags = UE::AvaTransition::FDefaultTags::Get();
+	TransitionLayer = DefaultTags.DefaultLayer.MakeTagHandle();
+
 	FStateTreeEditorColor DefaultColor;
 	DefaultColor.ColorRef = FStateTreeEditorColorRef(UE::AvaTransitionEditor::ColorId_Default);
 	DefaultColor.Color = FLinearColor(0, 0.15, 0.2);
