Motion Design: small fixes from previous sequence preset change

diff --git a/Engine/Plugins/VirtualProduction/Avalanche/Source/Avalanche/Private/AvaScene.cpp b/Engine/Plugins/VirtualProduction/Avalanche/Source/Avalanche/Private/AvaScene.cpp
--- a/Engine/Plugins/VirtualProduction/Avalanche/Source/Avalanche/Private/AvaScene.cpp
+++ b/Engine/Plugins/VirtualProduction/Avalanche/Source/Avalanche/Private/AvaScene.cpp
@@ -211,7 +211,7 @@ uint32 AAvaScene::AddSequences(TConstArrayView<UAvaSequence*> InSequences)
 
 	if (AddedCount == 0)
 	{
-		return false;
+		return 0;
 	}
 
 	ScheduleRebuildSequenceTree();
@@ -222,7 +222,7 @@ uint32 AAvaScene::AddSequences(TConstArrayView<UAvaSequence*> InSequences)
 		FEngineAnalytics::GetProvider().RecordEvent(TEXT("Editor.Usage.MotionDesign.Sequence"), TEXT("Action"), TEXT("AddedSequence"));
 	}
 #endif
-	return true;
+	return AddedCount;
 }
 
 void AAvaScene::RemoveSequence(UAvaSequence* InSequence)
diff --git a/Engine/Plugins/VirtualProduction/Avalanche/Source/AvalancheSequencer/Private/SequenceTree/Widgets/SAvaSequenceTree.cpp b/Engine/Plugins/VirtualProduction/Avalanche/Source/AvalancheSequencer/Private/SequenceTree/Widgets/SAvaSequenceTree.cpp
--- a/Engine/Plugins/VirtualProduction/Avalanche/Source/AvalancheSequencer/Private/SequenceTree/Widgets/SAvaSequenceTree.cpp
+++ b/Engine/Plugins/VirtualProduction/Avalanche/Source/AvalancheSequencer/Private/SequenceTree/Widgets/SAvaSequenceTree.cpp
@@ -369,7 +369,7 @@ void SAvaSequenceTree::CreateAddSequencePresetMenu(UToolMenu* InMenu)
 	if (!SequencerSettings->GetCustomPresetGroups().IsEmpty())
 	{
 		FToolMenuSection& CustomSection = InMenu->FindOrAddSection(TEXT("CustomGroups"));
-		TransitionLogicSection.Label = LOCTEXT("CustomGroupsLabel", "Custom");
+		CustomSection.Label = LOCTEXT("CustomGroupsLabel", "Custom");
 		for (const FAvaSequencePresetGroup& Group : SequencerSettings->GetCustomPresetGroups())
 		{
 			AddGroup(CustomSection, Group);
diff --git a/Engine/Plugins/VirtualProduction/Avalanche/Source/AvalancheSequencer/Public/Settings/AvaSequencePreset.h b/Engine/Plugins/VirtualProduction/Avalanche/Source/AvalancheSequencer/Public/Settings/AvaSequencePreset.h
--- a/Engine/Plugins/VirtualProduction/Avalanche/Source/AvalancheSequencer/Public/Settings/AvaSequencePreset.h
+++ b/Engine/Plugins/VirtualProduction/Avalanche/Source/AvalancheSequencer/Public/Settings/AvaSequencePreset.h
@@ -103,7 +103,7 @@ struct FAvaSequencePresetGroup
 	UPROPERTY(EditAnywhere, Category="Sequence")
 	FName GroupName;
 
-	/** The preset that this group represents */
+	/** The presets that this group represents */
 	UPROPERTY(EditAnywhere, Category="Sequence")
 	TArray<FName> PresetNames;
 };
