Motion Design: fix issue with changing sequence not being evaluated. A NotifyViewedSequenceChanged is used so that it requests an evaluation of the frame in next tick.
Sequence assets don't seem to have this issue since everytime a new asset is opened a sequencer instance is re-created (via the asset editor open execution flow).

diff --git a/Engine/Plugins/VirtualProduction/Avalanche/Source/AvalancheSequencer/Private/AvaSequencer.cpp b/Engine/Plugins/VirtualProduction/Avalanche/Source/AvalancheSequencer/Private/AvaSequencer.cpp
--- a/Engine/Plugins/VirtualProduction/Avalanche/Source/AvalancheSequencer/Private/AvaSequencer.cpp
+++ b/Engine/Plugins/VirtualProduction/Avalanche/Source/AvalancheSequencer/Private/AvaSequencer.cpp
@@ -566,6 +566,7 @@ void FAvaSequencer::NotifyViewedSequenceChanged(UAvaSequence* InOldSequence)
 	if (Sequencer.IsValid() && ViewedSequence)
 	{
 		Sequencer->ResetToNewRootSequence(*ViewedSequence);
+		Sequencer->NotifyMovieSceneDataChanged(EMovieSceneDataChangeType::ActiveMovieSceneChanged);
 	}
 
 	if (SequenceTree.IsValid())
