Motion Design: fixed regression on scrollbar not appearing for the playback and settings tabs in md sequencer

diff --git a/Engine/Plugins/VirtualProduction/Avalanche/Source/AvalancheSequencer/Private/DetailsView/Section/AvaSequencePlaybackDetails.h b/Engine/Plugins/VirtualProduction/Avalanche/Source/AvalancheSequencer/Private/DetailsView/Section/AvaSequencePlaybackDetails.h
--- a/Engine/Plugins/VirtualProduction/Avalanche/Source/AvalancheSequencer/Private/DetailsView/Section/AvaSequencePlaybackDetails.h
+++ b/Engine/Plugins/VirtualProduction/Avalanche/Source/AvalancheSequencer/Private/DetailsView/Section/AvaSequencePlaybackDetails.h
@@ -22,6 +22,7 @@ public:
 	virtual FName GetSectionId() const override;
 	virtual FText GetSectionDisplayText() const override;
 	virtual bool ShouldShowSection() const override;
+	virtual TOptional<float> GetSectionFillHeight() const override { return 1.f; }
 	virtual int32 GetSortOrder() const override;
 	virtual TSharedRef<SWidget> CreateContentWidget() override;
 	//~ End ISidebarDrawerContent
diff --git a/Engine/Plugins/VirtualProduction/Avalanche/Source/AvalancheSequencer/Private/DetailsView/Section/AvaSequenceSettingsDetails.h b/Engine/Plugins/VirtualProduction/Avalanche/Source/AvalancheSequencer/Private/DetailsView/Section/AvaSequenceSettingsDetails.h
--- a/Engine/Plugins/VirtualProduction/Avalanche/Source/AvalancheSequencer/Private/DetailsView/Section/AvaSequenceSettingsDetails.h
+++ b/Engine/Plugins/VirtualProduction/Avalanche/Source/AvalancheSequencer/Private/DetailsView/Section/AvaSequenceSettingsDetails.h
@@ -21,6 +21,7 @@ public:
 	virtual FName GetSectionId() const override;
 	virtual FText GetSectionDisplayText() const override;
 	virtual bool ShouldShowSection() const override;
+	virtual TOptional<float> GetSectionFillHeight() const override { return 1.f; }
 	virtual int32 GetSortOrder() const override;
 	virtual TSharedRef<SWidget> CreateContentWidget() override;
 	//~ End ISidebarDrawerContent
