Motion Design: added create default TL scene option to automatically set up a new motion design level with an empty transition logic set up. This option enables TL, sets it to reuse, and creates 4 sequences (In, Out, Layer1 Change In, Layer 1 Change Out). Additionally, a "layer 1 content" null actor is spawned in with sub layer modifier set to the layer 1 change in/out sequences set.

#virtualized

diff --git a/Engine/Plugins/VirtualProduction/Avalanche/Source/AvalancheEditor/Private/Scene/AvaSceneDefaults.cpp b/Engine/Plugins/VirtualProduction/Avalanche/Source/AvalancheEditor/Private/Scene/AvaSceneDefaults.cpp
--- a/Engine/Plugins/VirtualProduction/Avalanche/Source/AvalancheEditor/Private/Scene/AvaSceneDefaults.cpp
+++ b/Engine/Plugins/VirtualProduction/Avalanche/Source/AvalancheEditor/Private/Scene/AvaSceneDefaults.cpp
@@ -38,9 +38,14 @@ namespace UE::AvaEditor::Private
 
 	struct FAvaSceneDefaultActor
 	{
+		/** The class of the actor to create/replace/update*/
 		TWeakObjectPtr<UClass> ClassWeak;
+		/** Describes the actor to create/replace/update */
 		FText Description;
+		/** Delegate called to initialize the actor */
 		FAvaSceneDefaultActorInitDelegate InitDelegate;
+		/** Response to apply when there are available actors. */
+		EAvaSceneDefaultActorResponse DesiredResponse = EAvaSceneDefaultActorResponse::ReplaceActor;
 	};
 
 	namespace DefaultSceneActorNames
@@ -64,7 +69,8 @@ namespace UE::AvaEditor::Private
 			{
 				AAvaNullActor::StaticClass(),
 				LOCTEXT("DefaultRoot", "Default Scene Root"),
-				FAvaSceneDefaultActorInitDelegate::CreateStatic(&InitDefaultSceneActor_DefaultRoot)
+				FAvaSceneDefaultActorInitDelegate::CreateStatic(&InitDefaultSceneActor_DefaultRoot),
+				EAvaSceneDefaultActorResponse::CreateNewActor
 			}
 		},
 		{
@@ -256,16 +262,17 @@ namespace UE::AvaEditor::Private
 
 		for (const TPair<FName, FAvaSceneDefaultActor>& DefaultPair : DefaultSceneActors)
 		{
-			if (!DefaultPair.Value.ClassWeak.Get())
+			UClass* const Class = DefaultPair.Value.ClassWeak.Get();
+			if (!Class)
 			{
 				continue;
 			}
 
 			TSharedRef<FAvaSceneDefaultActorResponse> ActorResponse = MakeShared<FAvaSceneDefaultActorResponse>();
-			ActorResponse->ActorClassWeak = DefaultPair.Value.ClassWeak;
+			ActorResponse->ActorClassWeak = Class;
 			ActorResponse->Description = DefaultPair.Value.Description;
 
-			const FObjectKey ClassKey = FObjectKey(DefaultPair.Value.ClassWeak.Get());
+			const FObjectKey ClassKey = FObjectKey(Class);
 
 			if (TTuple<TArray<TWeakObjectPtr<AActor>>, int32>* const FoundActorsPtr = FoundActors.Find(ClassKey))
 			{
@@ -283,12 +290,12 @@ namespace UE::AvaEditor::Private
 			else
 			{
 				// Could be any number of actors.
-				ActorResponse->AvailableActors.Empty();
+				ActorResponse->AvailableActors.Reset();
 
 				for (AActor* Actor : FActorRange(InWorld))
 				{
 					// Ensure exact class match
-					if (Actor->GetClass() == DefaultPair.Value.ClassWeak)
+					if (Actor->GetClass() == Class)
 					{
 						ActorResponse->AvailableActors.Add(Actor);
 					}
@@ -297,7 +304,7 @@ namespace UE::AvaEditor::Private
 				if (!ActorResponse->AvailableActors.IsEmpty())
 				{
 					ActorResponse->SelectedActor = ActorResponse->AvailableActors[0];
-					ActorResponse->Response = EAvaSceneDefaultActorResponse::ReplaceActor;
+					ActorResponse->Response = DefaultPair.Value.DesiredResponse;
 				}
 
 				// The next index to be used in 1. 0 was just used.
diff --git a/Engine/Plugins/VirtualProduction/Avalanche/Source/AvalancheEditor/Private/Sequencer/AvaSequencerExtension.cpp b/Engine/Plugins/VirtualProduction/Avalanche/Source/AvalancheEditor/Private/Sequencer/AvaSequencerExtension.cpp
--- a/Engine/Plugins/VirtualProduction/Avalanche/Source/AvalancheEditor/Private/Sequencer/AvaSequencerExtension.cpp
+++ b/Engine/Plugins/VirtualProduction/Avalanche/Source/AvalancheEditor/Private/Sequencer/AvaSequencerExtension.cpp
@@ -15,6 +15,7 @@
 #include "Misc/PackageName.h"
 #include "MovieScene.h"
 #include "Settings/AvaSequencerSettings.h"
+#include "Settings/AvaTransitionEditorSettings.h"
 #include "StatusBarSubsystem.h"
 #include "UObject/Package.h"
 #include "Viewport/AvaViewportExtension.h"
@@ -326,8 +327,14 @@ void FAvaSequencerExtension::ConditionallyCreateDefaultSequences()
 		return;
 	}
 
-	const TArray<const FAvaSequencePreset*> SequencePresets = SequencerSettings->GatherPresetsFromGroup(UAvaSequencerSettings::TransitionLogicStarterGroup);
-	AvaSequencer->AddSequenceFromPresets(SequencePresets);
+	// Transition Logic Default Sequences
+	const UAvaTransitionEditorSettings* TransitionEditorSettings = GetDefault<UAvaTransitionEditorSettings>();
+	if (TransitionEditorSettings && TransitionEditorSettings->ShouldCreateTransitionLogicDefaultScene())
+	{
+		const TArray<const FAvaSequencePreset*> SequencePresets = SequencerSettings->GatherPresetsFromGroup(UAvaSequencerSettings::TransitionLogicStarterGroup);
+		AvaSequencer->AddSequenceFromPresets(SequencePresets);
+	}
+
 	bCreateDefaultSequences = false;
 }
 
diff --git a/Engine/Plugins/VirtualProduction/Avalanche/Source/AvalancheEditor/Private/Transition/AvaTransitionExtension.cpp b/Engine/Plugins/VirtualProduction/Avalanche/Source/AvalancheEditor/Private/Transition/AvaTransitionExtension.cpp
--- a/Engine/Plugins/VirtualProduction/Avalanche/Source/AvalancheEditor/Private/Transition/AvaTransitionExtension.cpp
+++ b/Engine/Plugins/VirtualProduction/Avalanche/Source/AvalancheEditor/Private/Transition/AvaTransitionExtension.cpp
@@ -12,9 +12,15 @@
 #include "Conditions/AvaTransitionTypeMatchCondition.h"
 #include "DetailCategoryBuilder.h"
 #include "Editor.h"
+#include "Framework/AvaNullActor.h"
 #include "IAvaTransitionEditorModule.h"
 #include "IRemoteControlUIModule.h"
+#include "Modifiers/ActorModifierCoreBase.h"
+#include "Modifiers/Utilities/ActorModifierCoreLibrary.h"
 #include "PropertyHandle.h"
+#include "ScopedTransaction.h"
+#include "Settings/AvaSequencerSettings.h"
+#include "Settings/AvaTransitionEditorSettings.h"
 #include "Styling/SlateIconFinder.h"
 #include "Subsystems/AssetEditorSubsystem.h"
 #include "Tasks/AvaTransitionWaitForLayerTask.h"
@@ -138,6 +144,72 @@ void FAvaTransitionExtension::ExtendSettingsCategory(IDetailCategoryBuilder& InC
 	}
 }
 
+void FAvaTransitionExtension::OnSceneObjectCreated(UObject* InSceneObject)
+{
+	const UAvaTransitionEditorSettings* TransitionEditorSettings = GetDefault<UAvaTransitionEditorSettings>();
+	if (!TransitionEditorSettings || !TransitionEditorSettings->ShouldCreateTransitionLogicDefaultScene())
+	{
+		return;
+	}
+
+	UWorld* const World = GetWorld();
+	if (!World)
+	{
+		return;
+	}
+
+	UAvaTransitionSubsystem* const TransitionSubsystem = World->GetSubsystem<UAvaTransitionSubsystem>();
+	if (!TransitionSubsystem)
+	{
+		return;
+	}
+
+	const IAvaTransitionBehavior* const TransitionBehavior = TransitionSubsystem->GetOrCreateTransitionBehavior();
+	if (!TransitionBehavior)
+	{
+		return;
+	}
+
+	// Enable TL if creating the TL default scene
+	if (UAvaTransitionTree* const TransitionTree = TransitionBehavior->GetTransitionTree())
+	{
+		TransitionTree->SetEnabled(true);
+		TransitionTree->SetInstancingMode(EAvaTransitionInstancingMode::Reuse);
+	}
+
+	// Spawn Layer 1 Content null
+	FActorSpawnParameters SpawnParameters;
+	SpawnParameters.InitialActorLabel = TEXT("Layer 1 Content");
+	AAvaNullActor* const Layer1Actor = World->SpawnActor<AAvaNullActor>(SpawnParameters);
+
+	// Add the Sub Layer Modifier (BP) to this Layer 1 Content null
+	UActorModifierCoreStack* ModifierStack;
+	if (UActorModifierCoreLibrary::FindModifierStack(Layer1Actor, ModifierStack, /*bCreateIfNone*/true))
+	{
+		constexpr const TCHAR* SubLayerModifierClassPath = TEXT("/Avalanche/TransitionLogic/SubLayer/BP_SubLayerModifier.BP_SubLayerModifier_C";
+
+		FActorModifierCoreInsertOperation InsertOperation;
+		InsertOperation.ModifierClass = StaticLoadClass(UActorModifierCoreBase::StaticClass(), nullptr, SubLayerModifierClassPath));
+
+		UActorModifierCoreBase* Modifier;
+		if (UActorModifierCoreLibrary::InsertModifier(ModifierStack, InsertOperation, Modifier))
+		{
+			auto SetSequenceName =
+				[Modifier, ModifierClass = InsertOperation.ModifierClass](FName InPropertyName, FAvaSequenceName InSequenceName)
+				{
+					FStructProperty* const StructProperty = CastField<FStructProperty>(ModifierClass->FindPropertyByName(InPropertyName));
+					if (StructProperty && StructProperty->Struct == FAvaSequenceName::StaticStruct())
+					{
+						StructProperty->SetValue_InContainer(Modifier, &InSequenceName);
+					}
+				};
+
+			SetSequenceName(TEXT("Change In Sequence"), FAvaSequenceName(UAvaSequencerSettings::DefaultChangeInName));
+			SetSequenceName(TEXT("Change Out Sequence"), FAvaSequenceName(UAvaSequencerSettings::DefaultChangeOutName));
+		}
+	}
+}
+
 void FAvaTransitionExtension::ExtendToolbarMenu(UToolMenu& InMenu)
 {
 	FToolMenuSection& Section = InMenu.FindOrAddSection(DefaultSectionName);
@@ -161,6 +233,37 @@ void FAvaTransitionExtension::ExtendToolbarMenu(UToolMenu& InMenu)
 		, /*bSimpleComboBox*/true));
 }
 
+void FAvaTransitionExtension::ExtendCreateSceneOptions(UToolMenu* InMenu)
+{
+	auto ToggleCreateTLDefaultScene = []()->void
+		{
+			if (UAvaTransitionEditorSettings* TransitionEditorSettings = GetMutableDefault<UAvaTransitionEditorSettings>())
+			{
+				FScopedTransaction Transaction(LOCTEXT("ToggleTransitionTreeEnabled", "Toggle Transition Tree Enabled"));
+				TransitionEditorSettings->Modify();
+				TransitionEditorSettings->ToggleCreateTransitionLogicDefaultScene();
+				TransitionEditorSettings->SaveConfig();
+			}
+		};
+
+	auto ShouldCreateTLDefaultScene = []()->bool
+		{
+			const UAvaTransitionEditorSettings* TransitionEditorSettings = GetDefault<UAvaTransitionEditorSettings>();
+			return TransitionEditorSettings && TransitionEditorSettings->ShouldCreateTransitionLogicDefaultScene();
+		};
+
+	FToolMenuSection& Section = InMenu->FindOrAddSection(DefaultSectionName);
+
+	Section.AddMenuEntry(TEXT("CreateTransitionLogicDefaultScene")
+		, LOCTEXT("CreateTransitionLogicDefaultSceneLabel", "Create Transition Logic Default Scene")
+		, LOCTEXT("CreateTransitionLogicDefaultSceneTooltip", "Determines whether to create the default setup for Transition Logic when creating the Motion Design Scene")
+		, TAttribute<FSlateIcon>()
+		, FUIAction(FExecuteAction::CreateLambda(ToggleCreateTLDefaultScene)
+			, FCanExecuteAction()
+			, FIsActionChecked::CreateLambda(ShouldCreateTLDefaultScene))
+		, EUserInterfaceActionType::ToggleButton);
+}
+
 IAvaTransitionBehavior* FAvaTransitionExtension::GetTransitionBehavior() const
 {
 	if (UAvaTransitionSubsystem* TransitionSubsystem = FAvaWorldSubsystemUtils::GetWorldSubsystem<UAvaTransitionSubsystem>(this))
diff --git a/Engine/Plugins/VirtualProduction/Avalanche/Source/AvalancheEditor/Private/Transition/AvaTransitionExtension.h b/Engine/Plugins/VirtualProduction/Avalanche/Source/AvalancheEditor/Private/Transition/AvaTransitionExtension.h
--- a/Engine/Plugins/VirtualProduction/Avalanche/Source/AvalancheEditor/Private/Transition/AvaTransitionExtension.h
+++ b/Engine/Plugins/VirtualProduction/Avalanche/Source/AvalancheEditor/Private/Transition/AvaTransitionExtension.h
@@ -23,7 +23,9 @@ public:
 	virtual void Deactivate() override;
 	virtual FName GetCategoryName() const override;
 	virtual void ExtendSettingsCategory(IDetailCategoryBuilder& InCategoryBuilder) override;
+	virtual void OnSceneObjectCreated(UObject* InSceneObject) override;
 	virtual void ExtendToolbarMenu(UToolMenu& InMenu) override;
+	virtual void ExtendCreateSceneOptions(UToolMenu* InMenu) override;
 	//~ End IAvaEditorExtension
 
 private:
diff --git a/Engine/Plugins/VirtualProduction/Avalanche/Source/AvalancheEditorCore/Private/AvaEditor.cpp b/Engine/Plugins/VirtualProduction/Avalanche/Source/AvalancheEditorCore/Private/AvaEditor.cpp
--- a/Engine/Plugins/VirtualProduction/Avalanche/Source/AvalancheEditorCore/Private/AvaEditor.cpp
+++ b/Engine/Plugins/VirtualProduction/Avalanche/Source/AvalancheEditorCore/Private/AvaEditor.cpp
@@ -243,6 +243,15 @@ void FAvaEditor::NotifySceneActivationEvent()
 	}
 }
 
+void FAvaEditor::Construct()
+{
+	ForEachExtension(
+		[This = SharedThis(this)](const TSharedRef<IAvaEditorExtension>& InExtension)
+		{
+			InExtension->Construct(This);
+		});
+}
+
 void FAvaEditor::Activate(TSharedPtr<IToolkitHost> InOverrideToolkitHost)
 {
 	if (!CanActivate())
@@ -554,6 +563,17 @@ void FAvaEditor::ExtendToolbarMenu(UToolMenu* InMenu)
 	});
 }
 
+void FAvaEditor::ExtendCreateSceneOptions(UToolMenu* InMenu)
+{
+	check(InMenu);
+
+	ForEachExtension(
+		[InMenu](const TSharedRef<IAvaEditorExtension>& InExtension)
+		{
+			InExtension->ExtendCreateSceneOptions(InMenu);
+		});
+}
+
 void FAvaEditor::CloseTabs()
 {
 	TSharedPtr<FTabManager> TabManager = GetTabManager();
diff --git a/Engine/Plugins/VirtualProduction/Avalanche/Source/AvalancheEditorCore/Private/AvaEditor.h b/Engine/Plugins/VirtualProduction/Avalanche/Source/AvalancheEditorCore/Private/AvaEditor.h
--- a/Engine/Plugins/VirtualProduction/Avalanche/Source/AvalancheEditorCore/Private/AvaEditor.h
+++ b/Engine/Plugins/VirtualProduction/Avalanche/Source/AvalancheEditorCore/Private/AvaEditor.h
@@ -38,6 +38,7 @@ protected:
 	void NotifySceneActivationEvent();
 
 	//~ Begin IAvaEditor
+	virtual void Construct() override;
 	virtual void Activate(TSharedPtr<IToolkitHost> InOverrideToolkitHost = nullptr) override;
 	virtual void Deactivate() override;
 	virtual void Cleanup() override;
@@ -60,6 +61,7 @@ protected:
 	virtual void RegisterTabSpawners() override;
 	virtual void UnregisterTabSpawners() override;
 	virtual void ExtendToolbarMenu(UToolMenu* InMenu) override;
+	virtual void ExtendCreateSceneOptions(UToolMenu* InMenu) override;
 	virtual void CloseTabs() override;
 	virtual FReply DockInLayout(FName InTabId) override;
 	virtual TArray<TSharedRef<IAvaEditorExtension>> GetExtensions() const override;
diff --git a/Engine/Plugins/VirtualProduction/Avalanche/Source/AvalancheEditorCore/Private/AvaEditorBuilder.cpp b/Engine/Plugins/VirtualProduction/Avalanche/Source/AvalancheEditorCore/Private/AvaEditorBuilder.cpp
--- a/Engine/Plugins/VirtualProduction/Avalanche/Source/AvalancheEditorCore/Private/AvaEditorBuilder.cpp
+++ b/Engine/Plugins/VirtualProduction/Avalanche/Source/AvalancheEditorCore/Private/AvaEditorBuilder.cpp
@@ -34,15 +34,6 @@ TSharedRef<IAvaEditor> FAvaEditorBuilder::Build()
 
 	TSharedRef<IAvaEditor> Editor = CreateEditor();
 	Editor->Construct();
-
-	// NOTE: Extensions and Provider are expected to have been moved into and possibly filtered out in IAvaEditor
-	// so from here on will only use the Extensions and Provider given by IAvaEditor
-
-	for (const TSharedRef<IAvaEditorExtension>& Extension : Editor->GetExtensions())
-	{
-		Extension->Construct(Editor);
-	}
-
 	return Editor;
 }
 
diff --git a/Engine/Plugins/VirtualProduction/Avalanche/Source/AvalancheEditorCore/Private/LevelEditor/AvaLevelEditorToolbar.cpp b/Engine/Plugins/VirtualProduction/Avalanche/Source/AvalancheEditorCore/Private/LevelEditor/AvaLevelEditorToolbar.cpp
--- a/Engine/Plugins/VirtualProduction/Avalanche/Source/AvalancheEditorCore/Private/LevelEditor/AvaLevelEditorToolbar.cpp
+++ b/Engine/Plugins/VirtualProduction/Avalanche/Source/AvalancheEditorCore/Private/LevelEditor/AvaLevelEditorToolbar.cpp
@@ -143,6 +143,22 @@ void FAvaLevelEditorToolbar::ExtendLevelEditorToolbar(const TSharedRef<IAvaEdito
 		};
 
 		AddEntry(AvaLevelEditorCommands.CreateScene);
+
+		Section.AddEntry(FToolMenuEntry::InitComboButton(TEXT("CreateSceneOptions")
+			, FUIAction()
+			, FNewToolMenuDelegate::CreateSPLambda(this,
+				[this](UToolMenu* InMenu)
+				{
+					if (const TSharedPtr<IAvaEditor> LevelEditor = LevelEditorWeak.Pin())
+					{
+						LevelEditor->ExtendCreateSceneOptions(InMenu);
+					}
+				})
+			, LOCTEXT("CreateSceneOptionsLabel", "Create Scene Options")
+			, LOCTEXT("CreateSceneOptionsTooltip", "Options for creating a Motion Design scene")
+			, AvaLevelEditorCommands.CreateScene->GetIcon()
+			, /*bSimpleComboBox*/true));
+
 		AddEntry(AvaLevelEditorCommands.ActivateScene);
 		AddEntry(AvaLevelEditorCommands.DeactivateScene);
 	}
diff --git a/Engine/Plugins/VirtualProduction/Avalanche/Source/AvalancheEditorCore/Public/IAvaEditor.h b/Engine/Plugins/VirtualProduction/Avalanche/Source/AvalancheEditorCore/Public/IAvaEditor.h
--- a/Engine/Plugins/VirtualProduction/Avalanche/Source/AvalancheEditorCore/Public/IAvaEditor.h
+++ b/Engine/Plugins/VirtualProduction/Avalanche/Source/AvalancheEditorCore/Public/IAvaEditor.h
@@ -88,6 +88,9 @@ public:
 
 	virtual void ExtendToolbarMenu(UToolMenu* InMenu) {}
 
+	/** Adds entries to the create scene options menu */
+	virtual void ExtendCreateSceneOptions(UToolMenu* InMenu) {}
+
 	virtual void CloseTabs() = 0;
 
 	virtual FReply DockInLayout(FName InTabId) = 0;
diff --git a/Engine/Plugins/VirtualProduction/Avalanche/Source/AvalancheEditorCore/Public/IAvaEditorExtension.h b/Engine/Plugins/VirtualProduction/Avalanche/Source/AvalancheEditorCore/Public/IAvaEditorExtension.h
--- a/Engine/Plugins/VirtualProduction/Avalanche/Source/AvalancheEditorCore/Public/IAvaEditorExtension.h
+++ b/Engine/Plugins/VirtualProduction/Avalanche/Source/AvalancheEditorCore/Public/IAvaEditorExtension.h
@@ -65,6 +65,9 @@ public:
 	/** Extend the Level Editor Layout (only called when instancing for Level Editor) */
 	virtual void ExtendLevelEditorLayout(FLayoutExtender& InExtender) const {}
 
+	/** Opportunity for an extension to add to the create scene options menu */
+	virtual void ExtendCreateSceneOptions(UToolMenu* InMenu) {}
+
 	/**
 	 * Name of the category this Extension belongs in.
 	 * Used in places like Setting Details Category
diff --git a/Engine/Plugins/VirtualProduction/Avalanche/Source/AvalancheModifiers/AvalancheModifiers.Build.cs b/Engine/Plugins/VirtualProduction/Avalanche/Source/AvalancheModifiers/AvalancheModifiers.Build.cs
--- a/Engine/Plugins/VirtualProduction/Avalanche/Source/AvalancheModifiers/AvalancheModifiers.Build.cs
+++ b/Engine/Plugins/VirtualProduction/Avalanche/Source/AvalancheModifiers/AvalancheModifiers.Build.cs
@@ -15,6 +15,7 @@ public class AvalancheModifiers : ModuleRules
 				"ActorModifierCore",
 				"ActorModifierLayout",
 				"Avalanche",
+				"AvalancheSequence",
 				"Core",
 				"CoreUObject",
 				"DynamicMesh",
diff --git a/Engine/Plugins/VirtualProduction/Avalanche/Source/AvalancheSequencer/Private/Settings/AvaSequencerSettings.cpp b/Engine/Plugins/VirtualProduction/Avalanche/Source/AvalancheSequencer/Private/Settings/AvaSequencerSettings.cpp
--- a/Engine/Plugins/VirtualProduction/Avalanche/Source/AvalancheSequencer/Private/Settings/AvaSequencerSettings.cpp
+++ b/Engine/Plugins/VirtualProduction/Avalanche/Source/AvalancheSequencer/Private/Settings/AvaSequencerSettings.cpp
@@ -34,8 +34,8 @@ namespace UE::AvaSequencer::Private
 		{
 			CreateDefaultPreset(TEXT("In"), DefaultTags.In),
 			CreateDefaultPreset(TEXT("Out"), DefaultTags.Out),
-			CreateDefaultPreset(TEXT("NewLayer_ChangeIn")),
-			CreateDefaultPreset(TEXT("NewLayer_ChangeOut")),
+			CreateDefaultPreset(UAvaSequencerSettings::DefaultChangeInName),
+			CreateDefaultPreset(UAvaSequencerSettings::DefaultChangeOutName),
 		};
 	}
 
@@ -51,8 +51,8 @@ namespace UE::AvaSequencer::Private
 			PresetGroup.PresetNames.SetNum(4);
 			PresetGroup.PresetNames[0] = TEXT("In");
 			PresetGroup.PresetNames[1] = TEXT("Out");
-			PresetGroup.PresetNames[2] = TEXT("NewLayer_ChangeIn");
-			PresetGroup.PresetNames[3] = TEXT("NewLayer_ChangeOut");
+			PresetGroup.PresetNames[2] = UAvaSequencerSettings::DefaultChangeInName;
+			PresetGroup.PresetNames[3] = UAvaSequencerSettings::DefaultChangeOutName;
 		}
 
 		// Transition Logic Sub Layer. 2 sequences: a sub layer's in/out pair
@@ -61,8 +61,8 @@ namespace UE::AvaSequencer::Private
 			PresetGroup.GroupName = UAvaSequencerSettings::TransitionLogicSubLayerGroup;
 
 			PresetGroup.PresetNames.SetNum(2);
-			PresetGroup.PresetNames[0] = TEXT("NewLayer_ChangeIn");
-			PresetGroup.PresetNames[1] = TEXT("NewLayer_ChangeOut");
+			PresetGroup.PresetNames[0] = UAvaSequencerSettings::DefaultChangeInName;
+			PresetGroup.PresetNames[1] = UAvaSequencerSettings::DefaultChangeOutName;
 		}
 		return PresetGroups;
 	}
@@ -70,6 +70,8 @@ namespace UE::AvaSequencer::Private
 
 const FLazyName UAvaSequencerSettings::TransitionLogicStarterGroup = TEXT("Transition Logic Starter");
 const FLazyName UAvaSequencerSettings::TransitionLogicSubLayerGroup = TEXT("Transition Logic Sub Layer");
+const FLazyName UAvaSequencerSettings::DefaultChangeInName = TEXT("Layer1_ChangeIn");
+const FLazyName UAvaSequencerSettings::DefaultChangeOutName = TEXT("Layer1_ChangeOut");
 
 UAvaSequencerSettings::UAvaSequencerSettings()
 {
diff --git a/Engine/Plugins/VirtualProduction/Avalanche/Source/AvalancheSequencer/Public/Settings/AvaSequencerSettings.h b/Engine/Plugins/VirtualProduction/Avalanche/Source/AvalancheSequencer/Public/Settings/AvaSequencerSettings.h
--- a/Engine/Plugins/VirtualProduction/Avalanche/Source/AvalancheSequencer/Public/Settings/AvaSequencerSettings.h
+++ b/Engine/Plugins/VirtualProduction/Avalanche/Source/AvalancheSequencer/Public/Settings/AvaSequencerSettings.h
@@ -9,6 +9,8 @@
 #include "UObject/Object.h"
 #include "AvaSequencerSettings.generated.h"
 
+#define UE_API AVALANCHESEQUENCER_API
+
 UCLASS(MinimalAPI, config=EditorPerProjectUserSettings, meta = (DisplayName = "Sequencer"))
 class UAvaSequencerSettings : public UDeveloperSettings
 {
@@ -19,8 +21,12 @@ public:
 	static constexpr const TCHAR* SettingsSection  = TEXT("Sequencer");
 
 	/** Group name for the built-in default groups */
-	AVALANCHESEQUENCER_API static const FLazyName TransitionLogicStarterGroup;
-	AVALANCHESEQUENCER_API static const FLazyName TransitionLogicSubLayerGroup;
+	UE_API static const FLazyName TransitionLogicStarterGroup;
+	UE_API static const FLazyName TransitionLogicSubLayerGroup;
+
+	/** Default change sequence names */
+	UE_API static const FLazyName DefaultChangeInName;
+	UE_API static const FLazyName DefaultChangeOutName;
 
 	UAvaSequencerSettings();
 
@@ -46,11 +52,11 @@ public:
 	const FAvaSequencePresetGroup* FindPresetGroup(FName InGroupName) const;
 
 	/** Finds all the presets within a preset group */
-	AVALANCHESEQUENCER_API TArray<const FAvaSequencePreset*> GatherPresetsFromGroup(FName InGroupName) const;
+	UE_API TArray<const FAvaSequencePreset*> GatherPresetsFromGroup(FName InGroupName) const;
 
-	AVALANCHESEQUENCER_API TConstArrayView<FAvaSequencePreset> GetDefaultSequencePresets() const;
+	UE_API TConstArrayView<FAvaSequencePreset> GetDefaultSequencePresets() const;
 
-	AVALANCHESEQUENCER_API TConstArrayView<FAvaSequencePresetGroup> GetTransitionLogicPresetGroups() const;
+	UE_API TConstArrayView<FAvaSequencePresetGroup> GetTransitionLogicPresetGroups() const;
 
 	const TSet<FAvaSequencePreset>& GetCustomSequencePresets() const
 	{
@@ -99,3 +105,5 @@ private:
 	UPROPERTY(Config, meta = (DeprecatedProperty, DeprecationMessage = "The Motion Design sidebar drawer was moved to the Sequencer sidebar and now stores its state in USequencerSettings"))
 	FSidebarState SidebarState_DEPRECATED;
 };
+
+#undef UE_API
diff --git a/Engine/Plugins/VirtualProduction/Avalanche/Source/AvalancheTransitionEditor/Private/Settings/AvaTransitionEditorSettings.cpp b/Engine/Plugins/VirtualProduction/Avalanche/Source/AvalancheTransitionEditor/Private/Settings/AvaTransitionEditorSettings.cpp
--- a/Engine/Plugins/VirtualProduction/Avalanche/Source/AvalancheTransitionEditor/Private/Settings/AvaTransitionEditorSettings.cpp
+++ b/Engine/Plugins/VirtualProduction/Avalanche/Source/AvalancheTransitionEditor/Private/Settings/AvaTransitionEditorSettings.cpp
@@ -1,6 +1,6 @@
 // Copyright Epic Games, Inc. All Rights Reserved.
 
-#include "AvaTransitionEditorSettings.h"
+#include "Settings/AvaTransitionEditorSettings.h"
 #include "AvaTransitionTree.h"
 #include "AvaTransitionTreeEditorData.h"
 
@@ -19,3 +19,8 @@ UAvaTransitionTreeEditorData* UAvaTransitionEditorSettings::LoadDefaultTemplateE
 	}
 	return nullptr;
 }
+
+void UAvaTransitionEditorSettings::ToggleCreateTransitionLogicDefaultScene()
+{
+	bCreateTransitionLogicDefaultScene = !bCreateTransitionLogicDefaultScene;
+}
diff --git a/Engine/Plugins/VirtualProduction/Avalanche/Source/AvalancheTransitionEditor/Private/Settings/AvaTransitionEditorSettings.h b/Engine/Plugins/VirtualProduction/Avalanche/Source/AvalancheTransitionEditor/Private/Settings/AvaTransitionEditorSettings.h
deleted file mode 100644
diff --git a/Engine/Plugins/VirtualProduction/Avalanche/Source/AvalancheTransitionEditor/Public/Settings/AvaTransitionEditorSettings.h b/Engine/Plugins/VirtualProduction/Avalanche/Source/AvalancheTransitionEditor/Public/Settings/AvaTransitionEditorSettings.h
new file mode 100644
--- /dev/null
+++ b/Engine/Plugins/VirtualProduction/Avalanche/Source/AvalancheTransitionEditor/Public/Settings/AvaTransitionEditorSettings.h
@@ -0,0 +1,41 @@
+// Copyright Epic Games, Inc. All Rights Reserved.
+
+#pragma once
+
+#include "Engine/DeveloperSettings.h"
+#include "UObject/SoftObjectPtr.h"
+#include "AvaTransitionEditorSettings.generated.h"
+
+#define UE_API AVALANCHETRANSITIONEDITOR_API
+
+class UAvaTransitionTree;
+class UAvaTransitionTreeEditorData;
+
+UCLASS(MinimalAPI, config=EditorPerProjectUserSettings, meta=(DisplayName="Transition Logic"))
+class UAvaTransitionEditorSettings : public UDeveloperSettings
+{
+	GENERATED_BODY()
+
+public:
+	UAvaTransitionEditorSettings();
+
+	UAvaTransitionTreeEditorData* LoadDefaultTemplateEditorData() const;
+
+	bool ShouldCreateTransitionLogicDefaultScene() const
+	{
+		return bCreateTransitionLogicDefaultScene;
+	}
+
+	UE_API void ToggleCreateTransitionLogicDefaultScene();
+
+private:
+	/** The template to use when building new transition trees */
+	UPROPERTY(Config, EditAnywhere, Category="Motion Design")
+	TSoftObjectPtr<UAvaTransitionTree> DefaultTemplate;
+
+	/** Whether to create the default scene when creating a Motion Design scene for the first time */
+	UPROPERTY(Config, EditAnywhere, Category="Motion Design")
+	bool bCreateTransitionLogicDefaultScene = true;
+};
+
+#undef UE_API
