Motion Design: added sequence behavior for remote control

diff --git a/Engine/Plugins/VirtualProduction/Avalanche/Source/Avalanche/Avalanche.Build.cs b/Engine/Plugins/VirtualProduction/Avalanche/Source/Avalanche/Avalanche.Build.cs
--- a/Engine/Plugins/VirtualProduction/Avalanche/Source/Avalanche/Avalanche.Build.cs
+++ b/Engine/Plugins/VirtualProduction/Avalanche/Source/Avalanche/Avalanche.Build.cs
@@ -26,6 +26,8 @@ public class Avalanche : ModuleRules
 				"GeometryScriptingCore",
 				"Niagara",
 				"RawMesh",
+				"RemoteControl",
+				"RemoteControlLogic",
 				"Text3D",
 			}
 		);
@@ -54,7 +56,6 @@ public class Avalanche : ModuleRules
 				"Projects",
 				"RHI",
 				"RawMesh",
-				"RemoteControl",
 				"RemoteControlComponents",
 				"RenderCore",
 				"Slate",
diff --git a/Engine/Plugins/VirtualProduction/Avalanche/Source/Avalanche/Private/RemoteControl/AvaRCSequenceBehaviorNode.cpp b/Engine/Plugins/VirtualProduction/Avalanche/Source/Avalanche/Private/RemoteControl/AvaRCSequenceBehaviorNode.cpp
new file mode 100644
--- /dev/null
+++ b/Engine/Plugins/VirtualProduction/Avalanche/Source/Avalanche/Private/RemoteControl/AvaRCSequenceBehaviorNode.cpp
@@ -0,0 +1,84 @@
+// Copyright Epic Games, Inc. All Rights Reserved.
+
+#include "RemoteControl/AvaRCSequenceBehaviorNode.h"
+#include "AvaLog.h"
+#include "AvaSequencePlaybackObject.h"
+#include "AvaSequenceSubsystem.h"
+#include "Behaviour/RCBehaviour.h"
+#include "Engine/Level.h"
+#include "RemoteControl/AvaRCSequenceBehavior.h"
+
+#define LOCTEXT_NAMESPACE "AvaRCPlaySequenceBehaviorNode"
+
+UAvaRCSequenceBehaviorNode::UAvaRCSequenceBehaviorNode()
+{
+	DisplayName = LOCTEXT("DisplayName", "Motion Design Sequence");
+	BehaviorDescription = LOCTEXT("Description", "Performs a specified action to Motion Design Sequences (Play, Stop, etc)");
+}
+
+bool UAvaRCSequenceBehaviorNode::Execute(URCBehaviour* InBehavior) const
+{
+	return true;
+}
+
+bool UAvaRCSequenceBehaviorNode::IsSupported(URCBehaviour* InBehavior) const
+{
+	// Only allow behaviors that are embedded in levels (i.e. in embedded presets)
+	return InBehavior
+		&& InBehavior->IsA<UAvaRCSequenceBehavior>()
+		&& InBehavior->GetTypedOuter<ULevel>();
+}
+
+void UAvaRCSequenceBehaviorNode::OnPassed(URCBehaviour* InBehavior) const
+{
+	ULevel* const Level = InBehavior ? InBehavior->GetTypedOuter<ULevel>() : nullptr;
+	if (!Level)
+	{
+		UE_LOG(LogAva, Error, TEXT("Sequence behavior failed. Behavior node %s could not find a valid outer level."), *GetFullName())
+		return;
+	}
+
+	const UAvaSequenceSubsystem* const SequenceSubsystem = UAvaSequenceSubsystem::Get(Level);
+	if (!SequenceSubsystem)
+	{
+		UE_LOG(LogAva, Error, TEXT("Sequence behavior failed. Behavior node %s could not find a valid sequence subsystem."), *GetFullName())
+		return;
+	}
+
+	IAvaSequencePlaybackObject* const PlaybackObject = SequenceSubsystem->FindPlaybackObject(Level);
+	if (!PlaybackObject)
+	{
+		UE_LOG(LogAva, Error, TEXT("Sequence behavior failed. Behavior node %s could not find a valid playback object."), *GetFullName());
+		return;
+	}
+
+	switch (SequenceAction)
+	{
+	case EAvaSequenceActionType::Play:
+		PlaybackObject->PlaySequencesByLabel(SequenceName.Name, PlaySettings);
+		break;
+
+	case EAvaSequenceActionType::Continue:
+		PlaybackObject->ContinueSequencesByLabel(SequenceName.Name);
+		break;
+
+	case EAvaSequenceActionType::Pause:
+		PlaybackObject->PauseSequencesByLabel(SequenceName.Name);
+		break;
+
+	case EAvaSequenceActionType::Stop:
+		PlaybackObject->StopSequencesByLabel(SequenceName.Name);
+		break;
+
+	default:
+		UE_LOG(LogAva, Error, TEXT("Sequence behavior failed. Behavior node %s does not have a valid sequence action."), *GetFullName())
+		break;
+	}
+}
+
+UClass* UAvaRCSequenceBehaviorNode::GetBehaviourClass() const
+{
+	return UAvaRCSequenceBehavior::StaticClass();
+}
+
+#undef LOCTEXT_NAMESPACE
diff --git a/Engine/Plugins/VirtualProduction/Avalanche/Source/Avalanche/Public/RemoteControl/AvaRCSequenceBehavior.h b/Engine/Plugins/VirtualProduction/Avalanche/Source/Avalanche/Public/RemoteControl/AvaRCSequenceBehavior.h
new file mode 100644
--- /dev/null
+++ b/Engine/Plugins/VirtualProduction/Avalanche/Source/Avalanche/Public/RemoteControl/AvaRCSequenceBehavior.h
@@ -0,0 +1,12 @@
+// Copyright Epic Games, Inc. All Rights Reserved.
+
+#pragma once
+
+#include "Behaviour/RCBehaviour.h"
+#include "AvaRCSequenceBehavior.generated.h"
+
+UCLASS(MinimalAPI)
+class UAvaRCSequenceBehavior : public URCBehaviour
+{
+	GENERATED_BODY()
+};
diff --git a/Engine/Plugins/VirtualProduction/Avalanche/Source/Avalanche/Public/RemoteControl/AvaRCSequenceBehaviorNode.h b/Engine/Plugins/VirtualProduction/Avalanche/Source/Avalanche/Public/RemoteControl/AvaRCSequenceBehaviorNode.h
new file mode 100644
--- /dev/null
+++ b/Engine/Plugins/VirtualProduction/Avalanche/Source/Avalanche/Public/RemoteControl/AvaRCSequenceBehaviorNode.h
@@ -0,0 +1,39 @@
+// Copyright Epic Games, Inc. All Rights Reserved.
+
+#pragma once
+
+#include "AvaSequenceName.h"
+#include "AvaSequenceShared.h"
+#include "Behaviour/RCBehaviourNode.h"
+#include "AvaRCSequenceBehaviorNode.generated.h"
+
+/** Performs an action (e.g. Play, Stop) for a defined Motion Design sequence */
+UCLASS(MinimalAPI)
+class UAvaRCSequenceBehaviorNode : public URCBehaviourNode
+{
+	GENERATED_BODY()
+
+public:
+	UAvaRCSequenceBehaviorNode();
+
+protected:
+	//~ Begin URCBehaviourNode
+	bool Execute(URCBehaviour* InBehavior) const override;
+	bool IsSupported(URCBehaviour* InBehavior) const override;
+	void OnPassed(URCBehaviour* InBehavior) const;
+	UClass* GetBehaviourClass() const override;
+	//~ End URCBehaviourNode
+
+private:
+	/** Sequence name to play */
+	UPROPERTY(EditAnywhere, Category="Sequence")
+	FAvaSequenceName SequenceName;
+
+	/** Action to perform for the given sequence */
+	UPROPERTY(EditAnywhere, Category="Sequence")
+	EAvaSequenceActionType SequenceAction = EAvaSequenceActionType::Play;
+
+	/** Sequence Play Settings */
+	UPROPERTY(EditAnywhere, Category="Sequence", meta=(EditCondition="SequenceAction==EAvaSequenceActionType::Play", EditConditionHides))
+	FAvaSequencePlayParams PlaySettings;
+};
diff --git a/Engine/Plugins/VirtualProduction/Avalanche/Source/AvalancheSequence/AvalancheSequence.Build.cs b/Engine/Plugins/VirtualProduction/Avalanche/Source/AvalancheSequence/AvalancheSequence.Build.cs
--- a/Engine/Plugins/VirtualProduction/Avalanche/Source/AvalancheSequence/AvalancheSequence.Build.cs
+++ b/Engine/Plugins/VirtualProduction/Avalanche/Source/AvalancheSequence/AvalancheSequence.Build.cs
@@ -26,8 +26,6 @@ public class AvalancheSequence : ModuleRules
             new string[]
             {
                 "AvalancheCore",
-                "RemoteControl",
-                "RemoteControlLogic",
                 "StateTreeModule",
             });
 
diff --git a/Engine/Plugins/VirtualProduction/Avalanche/Source/AvalancheSequence/Private/AvaSequencePlaybackActor.cpp b/Engine/Plugins/VirtualProduction/Avalanche/Source/AvalancheSequence/Private/AvaSequencePlaybackActor.cpp
--- a/Engine/Plugins/VirtualProduction/Avalanche/Source/AvalancheSequence/Private/AvaSequencePlaybackActor.cpp
+++ b/Engine/Plugins/VirtualProduction/Avalanche/Source/AvalancheSequence/Private/AvaSequencePlaybackActor.cpp
@@ -223,6 +223,37 @@ TArray<UAvaSequencePlayer*> AAvaSequencePlaybackActor::ContinueSequencesByTag(co
 	return SequencePlayers;
 }
 
+void AAvaSequencePlaybackActor::PauseSequence(UAvaSequence* InSequence)
+{
+	if (UAvaSequencePlayer* const Player = GetSequencePlayer(InSequence))
+	{
+		Player->Pause();
+	}
+}
+
+void AAvaSequencePlaybackActor::PauseSequencesByLabel(FName InSequenceLabel)
+{
+	PauseSequencesByLabels({ InSequenceLabel });
+}
+
+void AAvaSequencePlaybackActor::PauseSequencesByLabels(const TArray<FName>& InSequenceLabels)
+{
+	const TArray<UAvaSequencePlayer*> SequencePlayers = GetSequencePlayersByLabels(InSequenceLabels);
+	for (UAvaSequencePlayer* SequencePlayer : SequencePlayers)
+	{
+		SequencePlayer->Pause();
+	}
+}
+
+void AAvaSequencePlaybackActor::PauseSequencesByTag(const FAvaTagHandle& InTagHandle, bool bInExactMatch)
+{
+	const TArray<UAvaSequencePlayer*> SequencePlayers = GetSequencePlayersByTag(InTagHandle, bInExactMatch);
+	for (UAvaSequencePlayer* SequencePlayer : SequencePlayers)
+	{
+		SequencePlayer->Pause();
+	}
+}
+
 TArray<UAvaSequencePlayer*> AAvaSequencePlaybackActor::ContinueSequencesByLabels(const TArray<FName>& InSequenceLabels)
 {
 	TArray<UAvaSequencePlayer*> SequencePlayers;
@@ -245,6 +276,29 @@ void AAvaSequencePlaybackActor::StopSequence(UAvaSequence* InSequence)
 	}
 }
 
+void AAvaSequencePlaybackActor::StopSequencesByLabel(FName InSequenceLabel)
+{
+	StopSequencesByLabels({ InSequenceLabel });
+}
+
+void AAvaSequencePlaybackActor::StopSequencesByLabels(const TArray<FName>& InSequenceLabels)
+{
+	const TArray<UAvaSequencePlayer*> SequencePlayers = GetSequencePlayersByLabels(InSequenceLabels);
+	for (UAvaSequencePlayer* SequencePlayer : SequencePlayers)
+	{
+		SequencePlayer->Stop();
+	}
+}
+
+void AAvaSequencePlaybackActor::StopSequencesByTag(const FAvaTagHandle& InTagHandle, bool bInExactMatch)
+{
+	const TArray<UAvaSequencePlayer*> SequencePlayers = GetSequencePlayersByTag(InTagHandle, bInExactMatch);
+	for (UAvaSequencePlayer* SequencePlayer : SequencePlayers)
+	{
+		SequencePlayer->Stop();
+	}
+}
+
 void AAvaSequencePlaybackActor::UpdateCameraCut(const UE::MovieScene::FOnCameraCutUpdatedParams& InCameraCutParams)
 {
 	OnCameraCut.Broadcast(InCameraCutParams.ViewTarget, InCameraCutParams.bIsJumpCut);
@@ -354,7 +408,22 @@ TArray<UAvaSequence*> AAvaSequencePlaybackActor::GetSequencesByTag(const FAvaTag
 
 	for (UAvaSequence* Sequence : AllSequences)
 	{
-		if (Sequence && Sequence->GetSequenceTag().Overlaps(InTagHandle))
+		if (!Sequence)
+		{
+			continue;
+		}
+
+		bool bTagMatches;
+		if (bInExactMatch)
+		{
+			bTagMatches = Sequence->GetSequenceTag().MatchesExact(InTagHandle);
+		}
+		else
+		{
+			bTagMatches = Sequence->GetSequenceTag().Overlaps(InTagHandle);
+		}
+
+		if (bTagMatches)
 		{
 			OutSequences.Add(Sequence);
 		}
@@ -396,26 +465,32 @@ UAvaSequencePlayer* AAvaSequencePlaybackActor::GetSequencePlayer(const UAvaSeque
 	return nullptr;
 }
 
-TArray<UAvaSequencePlayer*> AAvaSequencePlaybackActor::GetSequencePlayersByLabel(FName InSequenceLabel) const
+TArray<UAvaSequencePlayer*> AAvaSequencePlaybackActor::GetSequencePlayers(TConstArrayView<const UAvaSequence*> InSequences) const
 {
-	return GetSequencePlayersByLabels({ InSequenceLabel });
-}
-
-TArray<UAvaSequencePlayer*> AAvaSequencePlaybackActor::GetSequencePlayersByTag(const FAvaTagHandle& InTagHandle, bool bInExactMatch) const
-{
-	const TArray<UAvaSequence*> Sequences = GetSequencesByTag(InTagHandle, bInExactMatch);
-
 	TArray<UAvaSequencePlayer*> SequencePlayers;
-	SequencePlayers.Reserve(Sequences.Num());
+	SequencePlayers.Reserve(InSequences.Num());
 
-	for (UAvaSequence* Sequence : Sequences)
+	for (UAvaSequencePlayer* SequencePlayer : ActiveSequencePlayers)
 	{
-		SequencePlayers.Add(GetSequencePlayer(Sequence));
+		if (SequencePlayer && InSequences.Contains(SequencePlayer->GetSequence()))
+		{
+			SequencePlayers.Add(SequencePlayer);
+		}
 	}
 
 	return SequencePlayers;
 }
 
+TArray<UAvaSequencePlayer*> AAvaSequencePlaybackActor::GetSequencePlayersByLabel(FName InSequenceLabel) const
+{
+	return GetSequencePlayersByLabels({ InSequenceLabel });
+}
+
+TArray<UAvaSequencePlayer*> AAvaSequencePlaybackActor::GetSequencePlayersByTag(const FAvaTagHandle& InTagHandle, bool bInExactMatch) const
+{
+	return GetSequencePlayers(GetSequencesByTag(InTagHandle, bInExactMatch));
+}
+
 TArray<UAvaSequencePlayer*> AAvaSequencePlaybackActor::GetAllSequencePlayers() const
 {
 	TArray<UAvaSequencePlayer*> SequencePlayers;
@@ -437,17 +512,7 @@ bool AAvaSequencePlaybackActor::HasActiveSequencePlayers() const
 
 TArray<UAvaSequencePlayer*> AAvaSequencePlaybackActor::GetSequencePlayersByLabels(const TArray<FName>& InSequenceLabels) const
 {
-	const TArray<UAvaSequence*> Sequences = GetSequencesByLabel(InSequenceLabels);
-
-	TArray<UAvaSequencePlayer*> SequencePlayers;
-	SequencePlayers.Reserve(Sequences.Num());
-
-	for (UAvaSequence* Sequence : Sequences)
-	{
-		SequencePlayers.Add(GetSequencePlayer(Sequence));
-	}
-
-	return SequencePlayers;
+	return GetSequencePlayers(GetSequencesByLabel(InSequenceLabels));
 }
 
 UAvaSequencePlayer* AAvaSequencePlaybackActor::GetOrAddSequencePlayer(UAvaSequence* InSequence)
diff --git a/Engine/Plugins/VirtualProduction/Avalanche/Source/AvalancheSequence/Private/AvaSequencePlaybackActor.h b/Engine/Plugins/VirtualProduction/Avalanche/Source/AvalancheSequence/Private/AvaSequencePlaybackActor.h
--- a/Engine/Plugins/VirtualProduction/Avalanche/Source/AvalancheSequence/Private/AvaSequencePlaybackActor.h
+++ b/Engine/Plugins/VirtualProduction/Avalanche/Source/AvalancheSequence/Private/AvaSequencePlaybackActor.h
@@ -40,9 +40,17 @@ protected:
 	virtual TArray<UAvaSequencePlayer*> ContinueSequencesByLabel(FName InSequenceLabel) override;
 	virtual TArray<UAvaSequencePlayer*> ContinueSequencesByLabels(const TArray<FName>& InSequenceLabels) override;
 	virtual TArray<UAvaSequencePlayer*> ContinueSequencesByTag(const FAvaTagHandle& InTagHandle, bool bInExactMatch) override;
+	virtual void PauseSequence(UAvaSequence* InSequence) override;
+	virtual void PauseSequencesByLabel(FName InSequenceLabel) override;
+	virtual void PauseSequencesByLabels(const TArray<FName>& InSequenceLabels) override;
+	virtual void PauseSequencesByTag(const FAvaTagHandle& InTagHandle, bool bInExactMatch) override;
 	virtual void StopSequence(UAvaSequence* InSequence) override;
+	virtual void StopSequencesByLabel(FName InSequenceLabel) override;
+	virtual void StopSequencesByLabels(const TArray<FName>& InSequenceLabels) override;
+	virtual void StopSequencesByTag(const FAvaTagHandle& InTagHandle, bool bInExactMatch) override;
 	virtual UObject* GetPlaybackContext() const override;
 	virtual UAvaSequencePlayer* GetSequencePlayer(const UAvaSequence* InSequence) const override;
+	virtual TArray<UAvaSequencePlayer*> GetSequencePlayers(TConstArrayView<const UAvaSequence*> InSequences) const override;
 	virtual TArray<UAvaSequencePlayer*> GetSequencePlayersByLabel(FName InSequenceLabel) const override;
 	virtual TArray<UAvaSequencePlayer*> GetSequencePlayersByLabels(const TArray<FName>& InSequenceLabels) const override;
 	virtual TArray<UAvaSequencePlayer*> GetSequencePlayersByTag(const FAvaTagHandle& InTagHandle, bool bInExactMatch) const override;
diff --git a/Engine/Plugins/VirtualProduction/Avalanche/Source/AvalancheSequence/Public/AvaSequencePlaybackObject.h b/Engine/Plugins/VirtualProduction/Avalanche/Source/AvalancheSequence/Public/AvaSequencePlaybackObject.h
--- a/Engine/Plugins/VirtualProduction/Avalanche/Source/AvalancheSequence/Public/AvaSequencePlaybackObject.h
+++ b/Engine/Plugins/VirtualProduction/Avalanche/Source/AvalancheSequence/Public/AvaSequencePlaybackObject.h
@@ -135,8 +135,62 @@ public:
 	UFUNCTION(BlueprintCallable, DisplayName = "Continue Sequences (by Tag)", Category = "Playback")
 	virtual TArray<UAvaSequencePlayer*> ContinueSequencesByTag(const FAvaTagHandle& InTagHandle, bool bInExactMatch) = 0;
 
+	/**
+	 * Pauses all players for a given sequence
+	 * @param InSequence the sequence to pause
+	 */
+	virtual void PauseSequence(UAvaSequence* InSequence) = 0;
+
+	/**
+	 * Pauses all players for sequences that match the given label
+	 * @param InSequenceLabel the label of the sequences to pause
+	 */
+	UFUNCTION(BlueprintCallable, DisplayName = "Pause Sequences (by Label)", Category = "Playback")
+	virtual void PauseSequencesByLabel(FName InSequenceLabel) = 0;
+
+	/**
+	 * Pauses all players for sequences that match the given labels
+	 * @param InSequenceLabels the array of sequence labels to pause (must be an active sequence playing)
+	 */
+	UFUNCTION(BlueprintCallable, DisplayName = "Pause Sequences (by Labels)", Category = "Playback")
+	virtual void PauseSequencesByLabels(const TArray<FName>& InSequenceLabels) = 0;
+
+	/**
+	 * Pauses all players for sequences that match the provided tag
+	 * @param InTagHandle the tag to match
+	 * @param bInExactMatch whether to only consider sequences that match the tag exactly
+	 */
+	UFUNCTION(BlueprintCallable, DisplayName = "Pause Sequences (by Tag)", Category = "Playback")
+	virtual void PauseSequencesByTag(const FAvaTagHandle& InTagHandle, bool bInExactMatch) = 0;
+
+	/**
+	 * Stops all players for a given sequence
+	 * @param InSequence the sequence to stop
+	 */
 	virtual void StopSequence(UAvaSequence* InSequence) = 0;
 
+	/**
+	 * Stops all players for sequences that match the given label
+	 * @param InSequenceLabel the label of the sequences to stop
+	 */
+	UFUNCTION(BlueprintCallable, DisplayName = "Stop Sequences (by Label)", Category = "Playback")
+	virtual void StopSequencesByLabel(FName InSequenceLabel) = 0;
+
+	/**
+	 * Stops all players for sequences that match the given labels
+	 * @param InSequenceLabels the array of sequence labels to stop (must be an active sequence playing)
+	 */
+	UFUNCTION(BlueprintCallable, DisplayName = "Stop Sequences (by Labels)", Category = "Playback")
+	virtual void StopSequencesByLabels(const TArray<FName>& InSequenceLabels) = 0;
+
+	/**
+	 * Stops all players for sequences that match the provided tag
+	 * @param InTagHandle the tag to match
+	 * @param bInExactMatch whether to only consider sequences that match the tag exactly
+	 */
+	UFUNCTION(BlueprintCallable, DisplayName = "Stop Sequences (by Tag)", Category = "Playback")
+	virtual void StopSequencesByTag(const FAvaTagHandle& InTagHandle, bool bInExactMatch) = 0;
+
 	virtual void UpdateCameraCut(const UE::MovieScene::FOnCameraCutUpdatedParams& InCameraCutParams) = 0;
 
 	DECLARE_MULTICAST_DELEGATE_TwoParams(FOnCameraCut, UObject* /*CameraObject*/, bool /*bJump*/)
@@ -148,10 +202,16 @@ public:
 
 	virtual UAvaSequencePlayer* GetSequencePlayer(const UAvaSequence* InSequence) const = 0;
 
+	/** Retrieves the active sequence players for the given sequences, if any */
+	virtual TArray<UAvaSequencePlayer*> GetSequencePlayers(TConstArrayView<const UAvaSequence*> InSequences) const = 0;
+
+	/** Retrieves the active sequence players for the given sequences that match the given label, if any */
 	virtual TArray<UAvaSequencePlayer*> GetSequencePlayersByLabel(FName InSequenceLabel) const = 0;
 
+	/** Retrieves the active sequence players for the given sequences that match the given labels, if any */
 	virtual TArray<UAvaSequencePlayer*> GetSequencePlayersByLabels(const TArray<FName>& InSequenceLabels) const = 0;
 
+	/** Retrieves the active sequence players for the given sequences that match the given tag, if any */
 	virtual TArray<UAvaSequencePlayer*> GetSequencePlayersByTag(const FAvaTagHandle& InTagHandle, bool bInExactMatch) const = 0;
 
 	/** Retrieves all Active Sequence Players */
diff --git a/Engine/Plugins/VirtualProduction/Avalanche/Source/AvalancheSequence/Public/AvaSequenceShared.h b/Engine/Plugins/VirtualProduction/Avalanche/Source/AvalancheSequence/Public/AvaSequenceShared.h
--- a/Engine/Plugins/VirtualProduction/Avalanche/Source/AvalancheSequence/Public/AvaSequenceShared.h
+++ b/Engine/Plugins/VirtualProduction/Avalanche/Source/AvalancheSequence/Public/AvaSequenceShared.h
@@ -16,6 +16,17 @@ namespace UE::AvaSequence
 	constexpr double SmallSubFrame = 0.00000005;
 }
 
+/** Enumeration of actions that can be performed on a sequence */
+UENUM()
+enum class EAvaSequenceActionType : uint8
+{
+	None UMETA(Hidden),
+	Play,
+	Continue,
+	Pause,
+	Stop,
+};
+
 UENUM(BlueprintType, DisplayName="Motion Design Sequence Play Mode")
 enum class EAvaSequencePlayMode : uint8
 {
