Adding GCObject.h include explicitly. It will soon be removed from StrongObjectPtr.h and therefor many other UObject headers.

diff --git a/Engine/Plugins/VirtualProduction/Avalanche/Source/AvalancheMedia/Public/Playback/AvaPlaybackManager.h b/Engine/Plugins/VirtualProduction/Avalanche/Source/AvalancheMedia/Public/Playback/AvaPlaybackManager.h
--- a/Engine/Plugins/VirtualProduction/Avalanche/Source/AvalancheMedia/Public/Playback/AvaPlaybackManager.h
+++ b/Engine/Plugins/VirtualProduction/Avalanche/Source/AvalancheMedia/Public/Playback/AvaPlaybackManager.h
@@ -3,6 +3,7 @@
 #pragma once
 
 #include "Playback/AvaPlaybackGraph.h"
+#include "UObject/GCObject.h"
 #include "UObject/StrongObjectPtr.h"
 #include "AvaPlaybackManager.generated.h"
 
diff --git a/Engine/Source/Runtime/CoreUObject/Tests/ObjectPinningTest.cpp b/Engine/Source/Runtime/CoreUObject/Tests/ObjectPinningTest.cpp
--- a/Engine/Source/Runtime/CoreUObject/Tests/ObjectPinningTest.cpp
+++ b/Engine/Source/Runtime/CoreUObject/Tests/ObjectPinningTest.cpp
@@ -6,6 +6,7 @@
 #include "TestMacros/Assertions.h"
 #include "Tests/Benchmark.h"
 #include "Tasks/Task.h"
+#include "UObject/GCObject.h"
 #include "UObject/Package.h"
 #include "UObject/UObjectGlobals.h"
 
