MotionDesign : GeometryMask
- Add Text3D support for multiple styles and their materials

#DisabledValidations NonUnityAndStaticAnalysis

diff --git a/Engine/Plugins/Experimental/Text3D/Source/Text3D/Private/Extensions/Text3DDefaultMaterialExtension.cpp b/Engine/Plugins/Experimental/Text3D/Source/Text3D/Private/Extensions/Text3DDefaultMaterialExtension.cpp
--- a/Engine/Plugins/Experimental/Text3D/Source/Text3D/Private/Extensions/Text3DDefaultMaterialExtension.cpp
+++ b/Engine/Plugins/Experimental/Text3D/Source/Text3D/Private/Extensions/Text3DDefaultMaterialExtension.cpp
@@ -386,30 +386,21 @@ EText3DExtensionResult UText3DDefaultMaterialExtension::PreRendererUpdate(const
 		{
 			if (const UText3DStyleBase* StyleRule = StyleExtension->GetStyle(MaterialOverride.Tag))
 			{
+				UMaterialInterface* Material = FrontMaterial;
+
 				if (StyleRule->GetOverrideFrontColor())
 				{
-					UMaterialInterface* Material = MaterialOverride.Materials[static_cast<int32>(EText3DGroupType::Front)];
-
-					if (!Material)
-					{
-						Material = FindOrAdd(EText3DGroupType::Front, MaterialOverride.Tag);
-					}
-
-					if (UMaterialInstanceDynamic* DynamicMaterial = Cast<UMaterialInstanceDynamic>(Material))
+					if (UMaterialInstanceDynamic* DynamicMaterial = FindOrAdd(EText3DGroupType::Front, MaterialOverride.Tag))
 					{
 						const FLinearColor SolidColor = StyleRule->GetFrontColor();
 						DynamicMaterial->SetVectorParameterValue(UText3DProjectSettings::FMaterialParameters::SolidColor, SolidColor);
 						DynamicMaterial->SetScalarParameterValue(UText3DProjectSettings::FMaterialParameters::Mode, 0);
 						DynamicMaterial->SetScalarParameterValue(UText3DProjectSettings::FMaterialParameters::Opacity, Opacity);
+						Material = DynamicMaterial;
 					}
-
-					MaterialOverride.Materials[static_cast<int32>(EText3DGroupType::Front)] = Material;
-				}
-				else
-				{
-					MaterialOverride.Materials[static_cast<int32>(EText3DGroupType::Front)] = FrontMaterial;
 				}
 
+				MaterialOverride.Materials[static_cast<int32>(EText3DGroupType::Front)] = Material;
 				MaterialOverride.Materials[static_cast<int32>(EText3DGroupType::Bevel)] = BevelMaterial;
 				MaterialOverride.Materials[static_cast<int32>(EText3DGroupType::Extrude)] = ExtrudeMaterial;
 				MaterialOverride.Materials[static_cast<int32>(EText3DGroupType::Back)] = BackMaterial;
@@ -479,10 +470,7 @@ UMaterialInterface* UText3DDefaultMaterialExtension::GetMaterial(const UE::Text3
 		const int32 MaterialIndex = static_cast<int32>(InParameters.Group);
 		if (MaterialOverride && MaterialOverride->Materials.IsValidIndex(MaterialIndex))
 		{
-			if (const TObjectPtr<UMaterialInterface> Material = MaterialOverride->Materials[MaterialIndex])
-			{
-				return Material;
-			}
+			return MaterialOverride->Materials[MaterialIndex];
 		}
 	}
 
@@ -532,22 +520,20 @@ UMaterialInstanceDynamic* UText3DDefaultMaterialExtension::FindOrAdd(EText3DGrou
 	const FText3DMaterialKey MaterialKey(BlendMode, bIsUnlit);
 	UMaterial* ParentMaterial = Text3DSettings->GetBaseMaterial(MaterialKey);
 
-	if (InTag.IsNone())
+	FMaterialParameters Parameters;
+	Parameters.Group = InGroup;
+	Parameters.Tag = InTag;
+	if (UMaterialInstanceDynamic* DynMat = Cast<UMaterialInstanceDynamic>(GetMaterial(Parameters)))
 	{
-		FMaterialParameters Parameters;
-		Parameters.Group = InGroup;
-		if (UMaterialInstanceDynamic* DynMat = Cast<UMaterialInstanceDynamic>(GetMaterial(Parameters)))
+		if (DynMat->GetBaseMaterial() == ParentMaterial)
 		{
-			if (DynMat->GetBaseMaterial() == ParentMaterial)
-			{
-				return DynMat;
-			}
+			return DynMat;
 		}
+	}
 
-		if (Style == EText3DMaterialStyle::Custom || Style == EText3DMaterialStyle::Invalid)
-		{
-			return nullptr;
-		}
+	if (Style == EText3DMaterialStyle::Custom || Style == EText3DMaterialStyle::Invalid)
+	{
+		return nullptr;
 	}
 
 	FText3DMaterialGroupKey GroupKey(MaterialKey, InGroup, InTag);
@@ -620,11 +606,6 @@ void UText3DDefaultMaterialExtension::OnCustomMaterialChanged()
 	OnMaterialOptionsChanged();
 }
 
-void UText3DDefaultMaterialExtension::ClearMaterialOverrides()
-{
-	MaterialOverrides.Empty();
-}
-
 void UText3DDefaultMaterialExtension::RegisterMaterialOverride(FName InTag)
 {
 	if (!InTag.IsNone())
@@ -640,6 +621,17 @@ void UText3DDefaultMaterialExtension::RegisterMaterialOverride(FName InTag)
 	}
 }
 
+void UText3DDefaultMaterialExtension::UnregisterMaterialOverride(FName InTag)
+{
+	if (!InTag.IsNone())
+	{
+		MaterialOverrides.RemoveAll([InTag](const FText3DMaterialOverride& InMaterialOverride)
+		{
+			return InMaterialOverride.Tag.IsEqual(InTag);
+		});
+	}
+}
+
 void UText3DDefaultMaterialExtension::ForEachMaterial(TFunctionRef<bool(const UE::Text3D::Material::FMaterialParameters&, UMaterialInterface*)> InFunctor) const
 {
 	// Base materials
diff --git a/Engine/Plugins/Experimental/Text3D/Source/Text3D/Private/Extensions/Text3DDefaultStyleExtension.cpp b/Engine/Plugins/Experimental/Text3D/Source/Text3D/Private/Extensions/Text3DDefaultStyleExtension.cpp
--- a/Engine/Plugins/Experimental/Text3D/Source/Text3D/Private/Extensions/Text3DDefaultStyleExtension.cpp
+++ b/Engine/Plugins/Experimental/Text3D/Source/Text3D/Private/Extensions/Text3DDefaultStyleExtension.cpp
@@ -121,7 +121,8 @@ EText3DExtensionResult UText3DDefaultStyleExtension::PreRendererUpdate(const UE:
 		TextStyleSet = MakeShared<FSlateStyleSet>(TEXT("Text3DRichTextStyle"));
 		TextStyleSet->Set(FName(TEXT("default")), *DefaultFontStyle);
 
-		MaterialExtension->ClearMaterialOverrides();
+		TArray<FName> OutdatedMaterialNames;
+		MaterialExtension->GetMaterialNames(OutdatedMaterialNames);
 
 		// Process style set first
 		if (StyleSet)
@@ -137,6 +138,7 @@ EText3DExtensionResult UText3DDefaultStyleExtension::PreRendererUpdate(const UE:
 						FormatStyle->ConfigureStyle(TextStyle);
 						TextStyleSet->Set(StyleName, TextStyle);
 						MaterialExtension->RegisterMaterialOverride(StyleName);
+						OutdatedMaterialNames.Remove(StyleName);
 					}
 				}
 			}
@@ -154,9 +156,16 @@ EText3DExtensionResult UText3DDefaultStyleExtension::PreRendererUpdate(const UE:
 					FormatStyle->ConfigureStyle(TextStyle);
 					TextStyleSet->Set(StyleName, TextStyle);
 					MaterialExtension->RegisterMaterialOverride(StyleName);
+					OutdatedMaterialNames.Remove(StyleName);
 				}
 			}
 		}
+
+		// Unregister outdated material slot
+		for (const FName MaterialName : OutdatedMaterialNames)
+		{
+			MaterialExtension->UnregisterMaterialOverride(MaterialName);
+		}
 	}
 
 	return EText3DExtensionResult::Finished;
diff --git a/Engine/Plugins/Experimental/Text3D/Source/Text3D/Public/Extensions/Text3DDefaultMaterialExtension.h b/Engine/Plugins/Experimental/Text3D/Source/Text3D/Public/Extensions/Text3DDefaultMaterialExtension.h
--- a/Engine/Plugins/Experimental/Text3D/Source/Text3D/Public/Extensions/Text3DDefaultMaterialExtension.h
+++ b/Engine/Plugins/Experimental/Text3D/Source/Text3D/Public/Extensions/Text3DDefaultMaterialExtension.h
@@ -252,8 +252,8 @@ protected:
 	//~ Begin UText3DMaterialExtensionBase
 	virtual void SetMaterial(const UE::Text3D::Material::FMaterialParameters& InParameters, UMaterialInterface* InMaterial) override;
 	virtual UMaterialInterface* GetMaterial(const UE::Text3D::Material::FMaterialParameters& InParameters) const override;
-	virtual void ClearMaterialOverrides() override;
 	virtual void RegisterMaterialOverride(FName InTag) override;
+	virtual void UnregisterMaterialOverride(FName InTag) override;
 	virtual void ForEachMaterial(TFunctionRef<bool(const UE::Text3D::Material::FMaterialParameters&, UMaterialInterface*)> InFunctor) const override;
 	virtual int32 GetMaterialCount() const override;
 	virtual void GetMaterialNames(TArray<FName>& OutNames) const override;
diff --git a/Engine/Plugins/Experimental/Text3D/Source/Text3D/Public/Extensions/Text3DDefaultStyleExtension.h b/Engine/Plugins/Experimental/Text3D/Source/Text3D/Public/Extensions/Text3DDefaultStyleExtension.h
--- a/Engine/Plugins/Experimental/Text3D/Source/Text3D/Public/Extensions/Text3DDefaultStyleExtension.h
+++ b/Engine/Plugins/Experimental/Text3D/Source/Text3D/Public/Extensions/Text3DDefaultStyleExtension.h
@@ -36,11 +36,11 @@ protected:
 #endif
 	//~ End UObject
 
-	//~ Begin UText3DFormatStyleExtensionBase
+	//~ Begin UText3DStyleExtensionBase
 	virtual TSharedPtr<FTextBlockStyle> GetDefaultStyle() const override;
     virtual TSharedPtr<FSlateStyleSet> GetCustomStyles() const override;
 	virtual UText3DStyleBase* GetStyle(FName InStyle) const override;
-    //~ End UText3DFormatStyleExtensionBase
+    //~ End UText3DStyleExtensionBase
     
     //~ Begin UText3DExtensionBase
     virtual EText3DExtensionResult PreRendererUpdate(const UE::Text3D::Renderer::FUpdateParameters& InParameters) override;
diff --git a/Engine/Plugins/Experimental/Text3D/Source/Text3D/Public/Extensions/Text3DMaterialExtensionBase.h b/Engine/Plugins/Experimental/Text3D/Source/Text3D/Public/Extensions/Text3DMaterialExtensionBase.h
--- a/Engine/Plugins/Experimental/Text3D/Source/Text3D/Public/Extensions/Text3DMaterialExtensionBase.h
+++ b/Engine/Plugins/Experimental/Text3D/Source/Text3D/Public/Extensions/Text3DMaterialExtensionBase.h
@@ -49,13 +49,13 @@ public:
 		return nullptr;
 	}
 
-	/** Reset all custom materials override defined */
-	virtual void ClearMaterialOverrides()
+	/** Register a new material override */
+	virtual void RegisterMaterialOverride(FName InTag)
 	{
 	}
 
-	/** Register a new material override */
-	virtual void RegisterMaterialOverride(FName InTag)
+	/** Unregister a material override previously defined */
+	virtual void UnregisterMaterialOverride(FName InTag)
 	{
 	}
 
diff --git a/Engine/Plugins/VirtualProduction/Avalanche/Source/AvalancheMask/Private/Handling/AvaMaskMaterialInstanceHandle.cpp b/Engine/Plugins/VirtualProduction/Avalanche/Source/AvalancheMask/Private/Handling/AvaMaskMaterialInstanceHandle.cpp
--- a/Engine/Plugins/VirtualProduction/Avalanche/Source/AvalancheMask/Private/Handling/AvaMaskMaterialInstanceHandle.cpp
+++ b/Engine/Plugins/VirtualProduction/Avalanche/Source/AvalancheMask/Private/Handling/AvaMaskMaterialInstanceHandle.cpp
@@ -5,6 +5,7 @@
 #include "AvaMaskUtilities.h"
 #include "Engine/Texture.h"
 #include "Materials/AvaMaskMaterialInstanceSubsystem.h"
+#include "Materials/Material.h"
 #include "Materials/MaterialInstanceDynamic.h"
 #include "Materials/MaterialInterface.h"
 
@@ -109,8 +110,9 @@ bool FAvaMaskMaterialInstanceHandle::ApplyModifiedState(
 				InModifiedParameters.MaterialParameters.Texture
 				, InModifiedParameters.MaterialParameters.BaseOpacity
 				, InModifiedParameters.MaterialParameters.Channel
-				, InModifiedParameters.MaterialParameters.bInvert);
-			
+				, InModifiedParameters.MaterialParameters.bInvert
+			);
+
 			InMaterialSetter(MaterialInstance);
 		}
 
@@ -139,11 +141,20 @@ UMaterialInstanceDynamic* FAvaMaskMaterialInstanceHandle::GetOrCreateMaterialIns
 	const FName InMaskName
 	, const EBlendMode InBlendMode)
 {
+	const EBlendMode TargetBlendMode = UE::AvaMask::Internal::GetTargetBlendMode(InBlendMode, GetBlendMode());
+
+	if (UMaterialInstanceDynamic* MaterialInstance = WeakMaterialInstance.Get())
+	{
+		if ((MaterialInstance == GetParentMaterial() || MaterialInstance->GetBaseMaterial() == GetParentMaterial())
+			&& MaterialInstance->BlendMode == TargetBlendMode)
+		{
+			return MaterialInstance;
+		}
+	}
+
 	if (UMaterialInterface* ParentMaterial = GetParentMaterial())
 	{
-		const EBlendMode TargetBlendMode = UE::AvaMask::Internal::GetTargetBlendMode(InBlendMode, GetBlendMode());
 		const uint32 MaterialInstanceKey = UE::AvaMask::Internal::MakeMaterialInstanceKey(ParentMaterial, InMaskName, TargetBlendMode);
-
 		UMaterialInstanceDynamic* MaterialInstance = UE::AvaMask::Internal::GetMaterialInstanceSubsystem()->GetMaterialInstanceProvider()->FindOrAddMID(ParentMaterial, MaterialInstanceKey, TargetBlendMode);
 		WeakMaterialInstance = MaterialInstance;
 		return MaterialInstance;
diff --git a/Engine/Plugins/VirtualProduction/Avalanche/Source/AvalancheMask/Private/Handling/AvaMaskText3DActorMaterialCollectionHandle.cpp b/Engine/Plugins/VirtualProduction/Avalanche/Source/AvalancheMask/Private/Handling/AvaMaskText3DActorMaterialCollectionHandle.cpp
--- a/Engine/Plugins/VirtualProduction/Avalanche/Source/AvalancheMask/Private/Handling/AvaMaskText3DActorMaterialCollectionHandle.cpp
+++ b/Engine/Plugins/VirtualProduction/Avalanche/Source/AvalancheMask/Private/Handling/AvaMaskText3DActorMaterialCollectionHandle.cpp
@@ -6,7 +6,6 @@
 #include "AvaMaskMaterialReference.h"
 #include "AvaMaskUtilities.h"
 #include "AvaObjectHandleSubsystem.h"
-#include "AvaText3DComponent.h"
 #include "Components/PrimitiveComponent.h"
 #include "Extensions/Text3DDefaultMaterialExtension.h"
 #include "GameFramework/Actor.h"
@@ -57,14 +56,17 @@ FAvaMaskText3DActorMaterialCollectionHandle::~FAvaMaskText3DActorMaterialCollect
 TArray<TObjectPtr<UMaterialInterface>> FAvaMaskText3DActorMaterialCollectionHandle::GetMaterials()
 {
 	TArray<TObjectPtr<UMaterialInterface>> Materials;
-	Materials.Reserve(GetNumMaterials());
-	
+
 	if (const UText3DComponent* TextComponent = GetComponent())
 	{
-		Materials.Emplace(TextComponent->GetFrontMaterial());
-		Materials.Emplace(TextComponent->GetBevelMaterial());
-		Materials.Emplace(TextComponent->GetExtrudeMaterial());
-		Materials.Emplace(TextComponent->GetBackMaterial());
+		const UText3DMaterialExtensionBase* MaterialExtension = TextComponent->GetMaterialExtension();
+		Materials.Reserve(MaterialExtension->GetMaterialCount());
+
+		MaterialExtension->ForEachMaterial([&Materials](const UE::Text3D::Material::FMaterialParameters& InParameters, UMaterialInterface* InMaterial)->bool
+		{
+			Materials.Add(InMaterial);
+			return true; // continue
+		});
 	}
 
 	return Materials;
@@ -72,13 +74,17 @@ TArray<TObjectPtr<UMaterialInterface>> FAvaMaskText3DActorMaterialCollectionHand
 
 TArray<TSharedPtr<IAvaMaskMaterialHandle>> FAvaMaskText3DActorMaterialCollectionHandle::GetMaterialHandles()
 {
-	if (MaterialHandles.Num() != static_cast<int32>(EText3DGroupType::TypeCount)
-		|| MaterialHandles.ContainsByPredicate([](const TSharedPtr<IAvaMaskMaterialHandle>& InHandle)
-		{
-			return !UE::Ava::Internal::IsHandleValid(InHandle);
-		}))
+	if (const UText3DComponent* TextComponent = GetComponent())
 	{
-		if (const UText3DComponent* TextComponent = GetComponent())
+		const UText3DMaterialExtensionBase* MaterialExtension = TextComponent->GetMaterialExtension();
+
+		const int32 MaterialCount = MaterialExtension->GetMaterialCount();
+
+		if (MaterialHandles.Num() != MaterialCount
+			|| MaterialHandles.ContainsByPredicate([](const TSharedPtr<IAvaMaskMaterialHandle>& InHandle)
+			{
+				return !UE::Ava::Internal::IsHandleValid(InHandle);
+			}))
 		{
 			UAvaObjectHandleSubsystem* HandleSubsystem = UE::Ava::Internal::GetObjectHandleSubsystem();
 			if (!HandleSubsystem)
@@ -86,26 +92,19 @@ TArray<TSharedPtr<IAvaMaskMaterialHandle>> FAvaMaskText3DActorMaterialCollection
 				return MaterialHandles;
 			}
 
-			const UText3DMaterialExtensionBase* MaterialExtension = TextComponent->GetMaterialExtension();
-
-			UE::Text3D::Material::FMaterialParameters Parameters;
-			Parameters.Group = EText3DGroupType::Front;
-			UMaterialInterface* FrontMaterial = MaterialExtension->GetMaterial(Parameters);
-
-			Parameters.Group = EText3DGroupType::Back;
-			UMaterialInterface* BackMaterial = MaterialExtension->GetMaterial(Parameters);
-
-			Parameters.Group = EText3DGroupType::Extrude;
-			UMaterialInterface* ExtrudeMaterial = MaterialExtension->GetMaterial(Parameters);
+			MaterialHandles.SetNum(MaterialCount);
+			int32 Index = 0;
+			MaterialExtension->ForEachMaterial([this, HandleSubsystem, &Index](const UE::Text3D::Material::FMaterialParameters& InParameters, UMaterialInterface* InMaterial)->bool
+			{
+				if (!MaterialHandles[Index] || MaterialHandles[Index]->GetMaterial() != InMaterial)
+				{
+					MaterialHandles[Index] = HandleSubsystem->MakeHandle<IAvaMaskMaterialHandle>(InMaterial, UE::AvaMask::Internal::HandleTag);
+				}
 
-			Parameters.Group = EText3DGroupType::Bevel;
-			UMaterialInterface* BevelMaterial = MaterialExtension->GetMaterial(Parameters);
+				Index++;
 
-			MaterialHandles.Init(nullptr, static_cast<int32>(EText3DGroupType::TypeCount));
-			MaterialHandles[static_cast<int32>(EText3DGroupType::Front)] = FrontMaterial != nullptr ? HandleSubsystem->MakeHandle<IAvaMaskMaterialHandle>(FrontMaterial, UE::AvaMask::Internal::HandleTag) : nullptr;
-			MaterialHandles[static_cast<int32>(EText3DGroupType::Bevel)] = BevelMaterial != nullptr ? HandleSubsystem->MakeHandle<IAvaMaskMaterialHandle>(BevelMaterial, UE::AvaMask::Internal::HandleTag) : nullptr;
-			MaterialHandles[static_cast<int32>(EText3DGroupType::Extrude)] = ExtrudeMaterial != nullptr ? HandleSubsystem->MakeHandle<IAvaMaskMaterialHandle>(ExtrudeMaterial, UE::AvaMask::Internal::HandleTag) : nullptr;
-			MaterialHandles[static_cast<int32>(EText3DGroupType::Back)] = BackMaterial != nullptr ? HandleSubsystem->MakeHandle<IAvaMaskMaterialHandle>(BackMaterial, UE::AvaMask::Internal::HandleTag) : nullptr;
+				return true; // continue
+			});
 		}
 	}
 
@@ -117,48 +116,38 @@ void FAvaMaskText3DActorMaterialCollectionHandle::SetMaterial(
 	, const int32 InSlotIdx
 	, UMaterialInterface* InMaterial)
 {
-	if (!ensure(InSlotIdx < 4))
-	{
-		UE_LOG(LogAvaMask, Error, TEXT("SlotIdx out of range, should be from 0-3, was %u"), InSlotIdx);
-		return;
-	}
-
 	if (UText3DComponent* TextComponent = GetComponent())
 	{
-		bool bRefresh = false;
+		UText3DMaterialExtensionBase* MaterialExtension = TextComponent->GetMaterialExtension();
+		const int32 MaterialCount = MaterialExtension->GetMaterialCount();
+		
+		if (!ensure(InSlotIdx < MaterialCount))
+		{
+			UE_LOG(LogAvaMask, Error, TEXT("SlotIdx out of range, should be from 0-%i, was %u"), MaterialCount-1, InSlotIdx);
+			return;
+		}
+
+		const int32 MaterialIndex = InSlotIdx / 4;
+		const EText3DGroupType GroupType = static_cast<EText3DGroupType>(InSlotIdx % 4);
+
+		TArray<FName> OutNames;
+		MaterialExtension->GetMaterialNames(OutNames);
 
-		switch (InSlotIdx)
+		if (!ensure(OutNames.IsValidIndex(MaterialIndex)))
 		{
-		case 0:
-			{
-				TextComponent->SetFrontMaterial(InMaterial);
-				bRefresh = true;
-			}
-		break;
-		case 1:
-			{
-				TextComponent->SetBevelMaterial(InMaterial);
-				bRefresh = true;
-			}
-		break;
-		case 2:
-			{
-				TextComponent->SetExtrudeMaterial(InMaterial);
-				bRefresh = true;
-			}
-		break;
-		case 3:
-			{
-				TextComponent->SetBackMaterial(InMaterial);
-				bRefresh = true;
-			}
-		break;
-		default:
-		break;
+			UE_LOG(LogAvaMask, Error, TEXT("Material index out of range, should be from 0-%i, was %u"), OutNames.Num()-1, MaterialIndex);
+			return;
 		}
 
-		if (bRefresh)
+		UE::Text3D::Material::FMaterialParameters Parameters;
+		Parameters.Group = GroupType;
+		Parameters.Tag = OutNames[MaterialIndex];
+		const UMaterialInterface* CurrentMaterial = MaterialExtension->GetMaterial(Parameters);
+
+		if (CurrentMaterial != InMaterial)
 		{
+			MaterialExtension->SetMaterial(Parameters, InMaterial);
+
 			constexpr bool bImmediate = true;
 			TextComponent->RequestUpdate(EText3DRendererFlags::Material, bImmediate);
 		}
@@ -169,45 +158,47 @@ void FAvaMaskText3DActorMaterialCollectionHandle::SetMaterials(
 	const TArray<TObjectPtr<UMaterialInterface>>& InMaterials
 	, const TBitArray<>& InSetToggle)
 {
-	if (!ensure(InMaterials.Num() == GetNumMaterials()))
-	{
-		UE_LOG(LogAvaMask, Warning, TEXT("Expected 4 materials, got %u"), InMaterials.Num());
-		return;
-	}
-
-	if (!ensure(InSetToggle.Num() == GetNumMaterials()))
-	{
-		UE_LOG(LogAvaMask, Warning, TEXT("Expected 4 materials, got %u"), InSetToggle.Num());
-		return;
-	}
-
 	if (UText3DComponent* TextComponent = GetComponent())
 	{
-		bool bRefresh = false;
+		const int32 MaterialCount = GetNumMaterials();
 
-		if (InSetToggle[0])
+		if (!ensure(InMaterials.Num() == MaterialCount))
 		{
-			TextComponent->SetFrontMaterial(InMaterials[0]);
-			bRefresh = true;
+			UE_LOG(LogAvaMask, Warning, TEXT("Expected %i materials, got %u"), MaterialCount, InMaterials.Num());
+			return;
 		}
 
-		if (InSetToggle[1])
+		if (!ensure(InSetToggle.Num() == MaterialCount))
 		{
-			TextComponent->SetBevelMaterial(InMaterials[1]);
-			bRefresh = true;
+			UE_LOG(LogAvaMask, Warning, TEXT("Expected %i materials, got %u"), MaterialCount, InSetToggle.Num());
+			return;
 		}
 
-		if (InSetToggle[2])
-		{
-			TextComponent->SetExtrudeMaterial(InMaterials[2]);
-			bRefresh = true;
-		}
+		UText3DMaterialExtensionBase* MaterialExtension = TextComponent->GetMaterialExtension();
 
-		if (InSetToggle[3])
-		{
-			TextComponent->SetBackMaterial(InMaterials[3]);
-			bRefresh = true;
-		}
+		TArray<FName> OutNames;
+		MaterialExtension->GetMaterialNames(OutNames);
+
+		bool bRefresh = false;
+		for (int32 Index = 0; Index < MaterialCount; ++Index)
+        {
+			if (InSetToggle[Index])
+			{
+				const int32 MaterialIndex = Index / 4;
+				const EText3DGroupType GroupType = static_cast<EText3DGroupType>(Index % 4);
+
+				UE::Text3D::Material::FMaterialParameters Parameters;
+				Parameters.Group = GroupType;
+				Parameters.Tag = OutNames[MaterialIndex];
+				const UMaterialInterface* CurrentMaterial = MaterialExtension->GetMaterial(Parameters);
+
+				if (CurrentMaterial != InMaterials[Index])
+				{
+					MaterialExtension->SetMaterial(Parameters, InMaterials[Index]);
+					bRefresh = true;
+				}
+			}
+        }
 
 		if (bRefresh)
 		{
@@ -219,10 +210,13 @@ void FAvaMaskText3DActorMaterialCollectionHandle::SetMaterials(
 
 int32 FAvaMaskText3DActorMaterialCollectionHandle::GetNumMaterials() const
 {
-	constexpr int32 TextMaterialCount = static_cast<int32>(EText3DGroupType::TypeCount);
-	static_assert(TextMaterialCount == 4);
-	
-	return TextMaterialCount;
+	if (const UText3DComponent* TextComponent = GetComponent())
+	{
+		const UText3DMaterialExtensionBase* MaterialExtension = TextComponent->GetMaterialExtension();
+		return MaterialExtension->GetMaterialCount();
+	}
+
+	return 0;
 }
 
 void FAvaMaskText3DActorMaterialCollectionHandle::ForEachMaterial(
@@ -234,22 +228,13 @@ void FAvaMaskText3DActorMaterialCollectionHandle::ForEachMaterial(
 	if (const UText3DComponent* TextComponent = GetComponent())
 	{
 		const FSoftComponentReference ComponentReference = UE::AvaMask::Internal::MakeComponentReference(TextComponent->GetOwner(), TextComponent);
-		if (!InFunction(ComponentReference, static_cast<int32>(EText3DGroupType::Front), TextComponent->GetFrontMaterial()))
-		{
-			return;
-		}
 
-		if (!InFunction(ComponentReference, static_cast<int32>(EText3DGroupType::Bevel), TextComponent->GetBevelMaterial()))
+		const UText3DMaterialExtensionBase* MaterialExtension = TextComponent->GetMaterialExtension();
+		int32 Index = 0;
+		MaterialExtension->ForEachMaterial([&InFunction, &ComponentReference, &Index](const UE::Text3D::Material::FMaterialParameters& InParameters, UMaterialInterface* InMaterial)
 		{
-			return;
-		}
-
-		if (!InFunction(ComponentReference, static_cast<int32>(EText3DGroupType::Extrude), TextComponent->GetExtrudeMaterial()))
-		{
-			return;
-		}
-
-		InFunction(ComponentReference, static_cast<int32>(EText3DGroupType::Back), TextComponent->GetBackMaterial());
+			return InFunction(ComponentReference, Index++, InMaterial);
+		});
 	}
 }
 
@@ -260,32 +245,24 @@ void FAvaMaskText3DActorMaterialCollectionHandle::ForEachMaterialHandle(TFunctio
 		// Effectively refreshes material handles if necessary
 		GetMaterialHandles();
 
-		if (MaterialHandles.Num() != static_cast<int32>(EText3DGroupType::TypeCount))
-		{
-			return;
-		}
+		const int32 MaterialCount = GetNumMaterials();
 
-		const FSoftComponentReference ComponentReference = UE::AvaMask::Internal::MakeComponentReference(TextComponent->GetOwner(), TextComponent);
-		TSharedPtr<IAvaMaskMaterialHandle> MaterialHandle = MaterialHandles[static_cast<int32>(EText3DGroupType::Front)];
-		if (!InFunction(ComponentReference, static_cast<int32>(EText3DGroupType::Front), MaterialHandle.IsValid(), MaterialHandle))
+		if (MaterialHandles.Num() != MaterialCount)
 		{
 			return;
 		}
 
-		MaterialHandle = MaterialHandles[static_cast<int32>(EText3DGroupType::Bevel)];
-		if (!InFunction(ComponentReference, static_cast<int32>(EText3DGroupType::Bevel), MaterialHandle.IsValid(), MaterialHandle))
-		{
-			return;
-		}
+		const FSoftComponentReference ComponentReference = UE::AvaMask::Internal::MakeComponentReference(TextComponent->GetOwner(), TextComponent);
 
-		MaterialHandle = MaterialHandles[static_cast<int32>(EText3DGroupType::Extrude)];
-		if (!InFunction(ComponentReference, static_cast<int32>(EText3DGroupType::Extrude), MaterialHandle.IsValid(), MaterialHandle))
+		int32 Index = 0;
+		const TArray<TSharedPtr<IAvaMaskMaterialHandle>> MaterialHandlesCopy(MaterialHandles);
+		for (const TSharedPtr<IAvaMaskMaterialHandle>& MaterialHandle : MaterialHandlesCopy)
 		{
-			return;
+			if (!InFunction(ComponentReference, Index++, MaterialHandle.IsValid(), MaterialHandle))
+			{
+				return;
+			}
 		}
-
-		MaterialHandle = MaterialHandles[static_cast<int32>(EText3DGroupType::Back)];
-		InFunction(ComponentReference, static_cast<int32>(EText3DGroupType::Back), MaterialHandle.IsValid(), MaterialHandle);
 	}
 }
 
@@ -297,25 +274,34 @@ void FAvaMaskText3DActorMaterialCollectionHandle::MapEachMaterial(
 {	
 	if (const UText3DComponent* TextComponent = GetComponent())
 	{
+		const int32 MaterialCount = GetNumMaterials();
+
 		TArray<TObjectPtr<UMaterialInterface>> MappedMaterials;
-		MappedMaterials.Init(nullptr, static_cast<TArray<TObjectPtr<UMaterialInterface>>::SizeType>(EText3DGroupType::TypeCount));
+		MappedMaterials.Init(nullptr, MaterialCount);
 
 		TBitArray<> SetFlags;
-        SetFlags.Init(false, MappedMaterials.Num());
+        SetFlags.Init(false, MaterialCount);
 
 		const FSoftComponentReference ComponentReference = UE::AvaMask::Internal::MakeComponentReference(TextComponent->GetOwner(), TextComponent);
-		
-		MappedMaterials[0] = InFunction(ComponentReference, static_cast<int32>(EText3DGroupType::Front), TextComponent->GetFrontMaterial());
-		SetFlags[0] = MappedMaterials[0] != nullptr;
 
-		MappedMaterials[1] = InFunction(ComponentReference, static_cast<int32>(EText3DGroupType::Bevel), TextComponent->GetBevelMaterial());
-		SetFlags[1] = MappedMaterials[1] != nullptr;
+		const UText3DMaterialExtensionBase* MaterialExtension = TextComponent->GetMaterialExtension();
 
-		MappedMaterials[2] = InFunction(ComponentReference, static_cast<int32>(EText3DGroupType::Extrude), TextComponent->GetExtrudeMaterial());
-		SetFlags[2] = MappedMaterials[2] != nullptr;
+		TArray<FName> OutNames;
+		MaterialExtension->GetMaterialNames(OutNames);
 
-		MappedMaterials[3] = InFunction(ComponentReference, static_cast<int32>(EText3DGroupType::Back), TextComponent->GetBackMaterial());
-		SetFlags[3] = MappedMaterials[3] != nullptr;
+		for (int32 Index = 0; Index < MaterialCount; ++Index)
+		{
+			const int32 MaterialIndex = Index / 4;
+			const EText3DGroupType GroupType = static_cast<EText3DGroupType>(Index % 4);
+
+			UE::Text3D::Material::FMaterialParameters Parameters;
+			Parameters.Group = GroupType;
+			Parameters.Tag = OutNames[MaterialIndex];
+			UMaterialInterface* Material = MaterialExtension->GetMaterial(Parameters);
+
+			MappedMaterials[Index] = InFunction(ComponentReference, Index, Material);
+			SetFlags[Index] = MappedMaterials[Index] != nullptr;
+		}
 
 		SetMaterials(MappedMaterials, SetFlags);
 	}
@@ -328,41 +314,39 @@ void FAvaMaskText3DActorMaterialCollectionHandle::MapEachMaterialHandle(TFunctio
 		// Effectively refreshes material handles if necessary
 		GetMaterialHandles();
 
+		const int32 MaterialCount = GetNumMaterials();
+
 		TArray<TObjectPtr<UMaterialInterface>> MappedMaterials;
-		MappedMaterials.SetNum(GetNumMaterials());
+		MappedMaterials.SetNum(MaterialCount);
 
 		TBitArray<> SetFlags;
-		SetFlags.Init(false, GetNumMaterials());
+		SetFlags.Init(false, MaterialCount);
 
 		const FSoftComponentReference ComponentReference = UE::AvaMask::Internal::MakeComponentReference(TextComponent->GetOwner(), TextComponent);
 
-		auto TryMap = [this, InFunction, &ComponentReference, &MappedMaterials, &SetFlags](const EText3DGroupType InMaterialGroup)
+		const UText3DMaterialExtensionBase* MaterialExtension = TextComponent->GetMaterialExtension();
+
+		TArray<FName> OutNames;
+		MaterialExtension->GetMaterialNames(OutNames);
+
+		for (int32 Index = 0; Index < MaterialCount; ++Index)
 		{
-			int32 InMaterialGroupIdx = static_cast<int32>(InMaterialGroup);			
-			
-			TSharedPtr<IAvaMaskMaterialHandle>& MaterialHandle = MaterialHandles[InMaterialGroupIdx];
+			TSharedPtr<IAvaMaskMaterialHandle>& MaterialHandle = MaterialHandles[Index];
+			const TSharedPtr<IAvaMaskMaterialHandle> MappedMaterialHandle = InFunction(ComponentReference, Index, true, MaterialHandle);
 
+			UMaterialInterface* RemappedMaterial = nullptr;
+			if (MappedMaterialHandle.IsValid())
 			{
-				const TSharedPtr<IAvaMaskMaterialHandle> MappedMaterialHandle = InFunction(ComponentReference, InMaterialGroupIdx, true, MaterialHandle);
-				UMaterialInterface* RemappedMaterial = nullptr;
-				if (MappedMaterialHandle.IsValid())
-				{
-					RemappedMaterial = MappedMaterialHandle->GetMaterial();	
-				}
-				else if (MaterialHandle.IsValid())
-				{
-					RemappedMaterial = MaterialHandle->GetMaterial();
-				}
-
-				MappedMaterials[InMaterialGroupIdx] = RemappedMaterial;
-				SetFlags[InMaterialGroupIdx] = RemappedMaterial != nullptr;
+				RemappedMaterial = MappedMaterialHandle->GetMaterial();	
+			}
+			else if (MaterialHandle.IsValid())
+			{
+				RemappedMaterial = MaterialHandle->GetMaterial();
 			}
-		};
 
-		TryMap(EText3DGroupType::Front);
-		TryMap(EText3DGroupType::Bevel);
-		TryMap(EText3DGroupType::Extrude);
-		TryMap(EText3DGroupType::Back);
+			MappedMaterials[Index] = RemappedMaterial;
+			SetFlags[Index] = RemappedMaterial != nullptr;
+		}
 
 		SetMaterials(MappedMaterials, SetFlags);
 	}
@@ -417,15 +401,24 @@ bool FAvaMaskText3DActorMaterialCollectionHandle::ApplyOriginalState(const FStru
 			// Only restore if the original was NOT translucent AND the current alpha is 1.0
 			if (UText3DDefaultMaterialExtension* MaterialExtension = TextComponent->GetCastedMaterialExtension<UText3DDefaultMaterialExtension>())
 			{
-				MaterialExtension->SetBlendMode(HandleData->BlendMode);
+				if (HandleData->BlendMode != MaterialExtension->GetBlendMode())
+				{
+					MaterialExtension->SetBlendMode(HandleData->BlendMode);
+
+					// Refresh materials immediately
+					constexpr bool bImmediateUpdate = true;
+					TextComponent->RequestUpdate(EText3DRendererFlags::Material, bImmediateUpdate);
+
+					// Need to refresh handles since blend mode didn't match
+					MaterialHandles.Reset();
+				}
 			}
-			
-			Super::ApplyOriginalState(InHandleData);
-			
-			return true;
-		}
 
-		return false;
+			if (Super::ApplyOriginalState(InHandleData))
+			{
+				return true;
+			}
+		}
 	}
 
 	return false;
@@ -441,11 +434,18 @@ bool FAvaMaskText3DActorMaterialCollectionHandle::ApplyModifiedState(const FAvaM
 			if (UText3DDefaultMaterialExtension* MaterialExtension = TextComponent->GetCastedMaterialExtension<UText3DDefaultMaterialExtension>())
 			{
 				const EText3DMaterialBlendMode TargetTranslucencyStyle = UE::AvaMask::Private::GetTargetTranslucencyType(MaterialExtension->GetBlendMode(), InModifiedParameters.MaterialParameters.BlendMode);
-				MaterialExtension->SetBlendMode(TargetTranslucencyStyle);
 
-				// Refresh materials if needed immediately
-				constexpr bool bImmediateUpdate = true;
-				TextComponent->RequestUpdate(EText3DRendererFlags::Material, bImmediateUpdate);
+				if (MaterialExtension->GetBlendMode() != TargetTranslucencyStyle)
+				{
+					MaterialExtension->SetBlendMode(TargetTranslucencyStyle);
+
+					// Refresh materials immediately
+					constexpr bool bImmediateUpdate = true;
+					TextComponent->RequestUpdate(EText3DRendererFlags::Material, bImmediateUpdate);
+
+					// Need to refresh handles since blend mode didn't match
+					MaterialHandles.Reset();
+				}
 			}
 		}
 
@@ -503,7 +503,7 @@ FStructView FAvaMaskText3DActorMaterialCollectionHandle::GetOrAddMaterialHandleD
 	return nullptr;
 }
 
-UText3DComponent* FAvaMaskText3DActorMaterialCollectionHandle::GetComponent()
+UText3DComponent* FAvaMaskText3DActorMaterialCollectionHandle::GetComponent() const
 {
 	if (UText3DComponent* Component = WeakComponent.Get())
 	{
@@ -512,11 +512,10 @@ UText3DComponent* FAvaMaskText3DActorMaterialCollectionHandle::GetComponent()
 
 	if (const AActor* Actor = WeakActor.Get())
 	{
-		WeakComponent = Actor->GetComponentByClass<UText3DComponent>();
-		if (UText3DComponent* Component = WeakComponent.Get())
-		{
-			return Component;
-		}
+		UText3DComponent* Component = Actor->GetComponentByClass<UText3DComponent>();
+		FAvaMaskText3DActorMaterialCollectionHandle* MutableThis = const_cast<FAvaMaskText3DActorMaterialCollectionHandle*>(this);
+		MutableThis->WeakComponent = Component;
+		return Component;
 	}
 
 	return nullptr;
@@ -526,11 +525,6 @@ void FAvaMaskText3DActorMaterialCollectionHandle::OnTextPostUpdate(UText3DCompon
 {
 	if (EnumHasAnyFlags(InFlags, EText3DRendererFlags::Material))
 	{
-		constexpr int32 FrontIndex = static_cast<int32>(EText3DGroupType::Front);
-		constexpr int32 BackIndex = static_cast<int32>(EText3DGroupType::Back);
-		constexpr int32 ExtrudeIndex = static_cast<int32>(EText3DGroupType::Extrude);
-		constexpr int32 BevelIndex = static_cast<int32>(EText3DGroupType::Bevel);
-
 		auto HasInvalidHandles = [this]()
 		{
 			return MaterialHandles.ContainsByPredicate([](const TSharedPtr<IAvaMaskMaterialHandle>& InHandle)
@@ -539,18 +533,42 @@ void FAvaMaskText3DActorMaterialCollectionHandle::OnTextPostUpdate(UText3DCompon
 			});
 		};
 
+		UText3DMaterialExtensionBase* MaterialExtension = InComponent->GetMaterialExtension();
+		auto CheckMaterialsMatch = [this, MaterialExtension]()
+		{
+			if (MaterialHandles.Num() != GetNumMaterials())
+			{
+				return false;
+			}
+
+			int32 Index = 0;
+			bool bMatch = true;
+			MaterialExtension->ForEachMaterial([this, &bMatch, &Index](const UE::Text3D::Material::FMaterialParameters& InParameters, UMaterialInterface* InMaterial)
+			{
+				UMaterialInterface* Material = MaterialHandles[Index]->GetMaterial();
+				if (Material != InMaterial)
+				{
+					if (!Material || Material->GetBaseMaterial() != InMaterial->GetBaseMaterial())
+					{
+						bMatch = false;
+					}
+				}
+
+				Index++;
+
+				return bMatch;
+			});
+	
+			return bMatch;
+		};
+
 		const bool bHasInvalidHandles = HasInvalidHandles();
+		const bool bAllMaterialsMatch = !bHasInvalidHandles && CheckMaterialsMatch();
 	
-		if (MaterialHandles.IsEmpty()
-			|| bHasInvalidHandles
-			|| MaterialHandles[FrontIndex]->GetMaterial() != InComponent->GetFrontMaterial()
-			|| MaterialHandles[BevelIndex]->GetMaterial() != InComponent->GetBevelMaterial()
-			|| MaterialHandles[ExtrudeIndex]->GetMaterial() != InComponent->GetExtrudeMaterial()
-			|| MaterialHandles[BackIndex]->GetMaterial() != InComponent->GetBackMaterial())
+		if (!bAllMaterialsMatch)
 		{
-			MaterialHandles.Empty();
-
 			// refresh handles
+			MaterialHandles.Empty();
 			GetMaterialHandles();
 
 			if (!HasInvalidHandles())
diff --git a/Engine/Plugins/VirtualProduction/Avalanche/Source/AvalancheMask/Private/Handling/AvaMaskText3DActorMaterialCollectionHandle.h b/Engine/Plugins/VirtualProduction/Avalanche/Source/AvalancheMask/Private/Handling/AvaMaskText3DActorMaterialCollectionHandle.h
--- a/Engine/Plugins/VirtualProduction/Avalanche/Source/AvalancheMask/Private/Handling/AvaMaskText3DActorMaterialCollectionHandle.h
+++ b/Engine/Plugins/VirtualProduction/Avalanche/Source/AvalancheMask/Private/Handling/AvaMaskText3DActorMaterialCollectionHandle.h
@@ -103,7 +103,7 @@ protected:
 		, const int32 InSlotIdx) override;
 	// ~End TAvaMaskMaterialCollectionHandle
 	
-	UText3DComponent* GetComponent();
+	UText3DComponent* GetComponent() const;
 
 	void OnTextPostUpdate(UText3DComponent* InComponent, EText3DRendererFlags InFlags);
 
diff --git a/Engine/Plugins/VirtualProduction/Avalanche/Source/AvalancheMask/Private/Handling/IAvaMaskMaterialHandle.h b/Engine/Plugins/VirtualProduction/Avalanche/Source/AvalancheMask/Private/Handling/IAvaMaskMaterialHandle.h
--- a/Engine/Plugins/VirtualProduction/Avalanche/Source/AvalancheMask/Private/Handling/IAvaMaskMaterialHandle.h
+++ b/Engine/Plugins/VirtualProduction/Avalanche/Source/AvalancheMask/Private/Handling/IAvaMaskMaterialHandle.h
@@ -123,7 +123,7 @@ bool TAvaMaskMaterialHandle<HandleDataType>::SetMaskParameters(UTexture* InTextu
 {
 	if (UMaterialInstanceDynamic* MaterialInstance = GetMaterialInstance())
 	{
-		UE_LOG(LogAvaMask, VeryVerbose, TEXT("SetParameters: Texture:%s, Channel:%s"), InTexture ? *InTexture->GetName() : TEXT("(None)"), *UE::AvaMask::Internal::MaskChannelEnumToVector[InChannel].ToString());
+		UE_LOG(LogAvaMask, Verbose, TEXT("SetParameters: Texture:%s, Channel:%s, Material:%s"), InTexture ? *InTexture->GetName() : TEXT("(None)"), *UE::AvaMask::Internal::MaskChannelEnumToVector[InChannel].ToString(), *MaterialInstance->GetName());
 
 		UTexture* DefaultTexture = nullptr;
 		MaterialInstance->GetTextureParameterDefaultValue(UE::AvaMask::Internal::TextureParameterInfo, DefaultTexture);
diff --git a/Engine/Plugins/VirtualProduction/Avalanche/Source/AvalancheMask/Private/Mask2D/AvaMask2DBaseModifier.cpp b/Engine/Plugins/VirtualProduction/Avalanche/Source/AvalancheMask/Private/Mask2D/AvaMask2DBaseModifier.cpp
--- a/Engine/Plugins/VirtualProduction/Avalanche/Source/AvalancheMask/Private/Mask2D/AvaMask2DBaseModifier.cpp
+++ b/Engine/Plugins/VirtualProduction/Avalanche/Source/AvalancheMask/Private/Mask2D/AvaMask2DBaseModifier.cpp
@@ -466,6 +466,8 @@ void UAvaMask2DBaseModifier::RestoreActorPreState(AActor* InActor, FAvaMask2DAct
     		{
     			TSharedPtr<IAvaMaskMaterialCollectionHandle> Handle = GetObjectHandleSubsystem()->MakeHandle<IAvaMaskMaterialCollectionHandle>(InActor, UE::AvaMask::Internal::HandleTag);
 
+    			Handle->OnSourceMaterialsChanged().Unbind();
+
     			for (const TPair<FAvaMask2DComponentMaterialPath, TObjectPtr<UMaterialInterface>>& OriginalMaterial : InActorData.OriginalMaterials)
     			{
     				Handle->SetMaterial(OriginalMaterial.Key.Component, OriginalMaterial.Key.SlotIdx, OriginalMaterial.Value);
