Motion Design: fixed typo

diff --git a/Engine/Plugins/VirtualProduction/Avalanche/Source/AvalancheEditorCore/Private/LevelEditor/AvaLevelEditorCommands.cpp b/Engine/Plugins/VirtualProduction/Avalanche/Source/AvalancheEditorCore/Private/LevelEditor/AvaLevelEditorCommands.cpp
--- a/Engine/Plugins/VirtualProduction/Avalanche/Source/AvalancheEditorCore/Private/LevelEditor/AvaLevelEditorCommands.cpp
+++ b/Engine/Plugins/VirtualProduction/Avalanche/Source/AvalancheEditorCore/Private/LevelEditor/AvaLevelEditorCommands.cpp
@@ -17,7 +17,7 @@ void FAvaLevelEditorCommands::RegisterCommands()
 {
 	UI_COMMAND(CreateScene
 		, "Create Scene"
-		, "Creates an Motion Design Scene for the Level and opens the Motion Design Editor Interface"
+		, "Creates a Motion Design Scene for the Level and opens the Motion Design Editor Interface"
 		, EUserInterfaceActionType::Button
 		, FInputChord());
 
