MotionDesign : Shapes
- Fixed line shape actor spawn issue using interactive tools

diff --git a/Engine/Plugins/VirtualProduction/Avalanche/Source/AvalancheInteractiveTools/Private/Tools/AvaInteractiveToolsToolBase.cpp b/Engine/Plugins/VirtualProduction/Avalanche/Source/AvalancheInteractiveTools/Private/Tools/AvaInteractiveToolsToolBase.cpp
--- a/Engine/Plugins/VirtualProduction/Avalanche/Source/AvalancheInteractiveTools/Private/Tools/AvaInteractiveToolsToolBase.cpp
+++ b/Engine/Plugins/VirtualProduction/Avalanche/Source/AvalancheInteractiveTools/Private/Tools/AvaInteractiveToolsToolBase.cpp
@@ -296,11 +296,17 @@ AActor* UAvaInteractiveToolsToolBase::SpawnActor(TSubclassOf<AActor> InActorClas
 		SpawnParams.NameMode = FActorSpawnParameters::ESpawnActorNameMode::Requested;
 		SpawnParams.ObjectFlags |= RF_Transient;
 	}
-	else if (bUseOverride)
+	else
 	{
-		SpawnParams.Name = FName(*InActorLabelOverride);
-		SpawnParams.NameMode = FActorSpawnParameters::ESpawnActorNameMode::Requested;
-		SpawnParams.ObjectFlags |= RF_Transactional;
+		// No longer needed if we spawn final actor
+		DestroyPreviewActor();
+
+		if (bUseOverride)
+		{
+			SpawnParams.Name = FName(*InActorLabelOverride);
+			SpawnParams.NameMode = FActorSpawnParameters::ESpawnActorNameMode::Requested;
+			SpawnParams.ObjectFlags |= RF_Transactional;
+		}
 	}
 
 	UActorFactory* ActorFactory = nullptr;
diff --git a/Engine/Plugins/VirtualProduction/Avalanche/Source/AvalancheShapesEditor/Private/Tools/AvaShapesEditorShapeToolIrregularPoly.cpp b/Engine/Plugins/VirtualProduction/Avalanche/Source/AvalancheShapesEditor/Private/Tools/AvaShapesEditorShapeToolIrregularPoly.cpp
--- a/Engine/Plugins/VirtualProduction/Avalanche/Source/AvalancheShapesEditor/Private/Tools/AvaShapesEditorShapeToolIrregularPoly.cpp
+++ b/Engine/Plugins/VirtualProduction/Avalanche/Source/AvalancheShapesEditor/Private/Tools/AvaShapesEditorShapeToolIrregularPoly.cpp
@@ -317,7 +317,6 @@ void UAvaShapesEditorShapeToolIrregularPoly::AddFinishedShape()
 							NewShapeActor->SetActorLocation(Vector);
 						}
 
-						PreviewActor->Destroy();
 						RequestShutdown(EToolShutdownType::Completed);
 						return;
 					}
diff --git a/Engine/Plugins/VirtualProduction/Avalanche/Source/AvalancheShapesEditor/Private/Tools/AvaShapesEditorShapeToolLine.cpp b/Engine/Plugins/VirtualProduction/Avalanche/Source/AvalancheShapesEditor/Private/Tools/AvaShapesEditorShapeToolLine.cpp
--- a/Engine/Plugins/VirtualProduction/Avalanche/Source/AvalancheShapesEditor/Private/Tools/AvaShapesEditorShapeToolLine.cpp
+++ b/Engine/Plugins/VirtualProduction/Avalanche/Source/AvalancheShapesEditor/Private/Tools/AvaShapesEditorShapeToolLine.cpp
@@ -117,11 +117,6 @@ void UAvaShapesEditorShapeToolLine::OnViewportPlannerUpdate()
 			
 			case 2:
 			{
-				if (PreviewActor)
-				{
-					PreviewActor->Destroy();
-				}
-
 				const FVector2f& Center = Positions[0] * 0.5f + Positions[1] * 0.5f;
 				SpawnedActor = SpawnActor(ActorClass, EAvaViewportStatus::Focused, Center, /* Preview */ false);
 				SetLineEnds(Cast<AAvaShapeActor>(SpawnedActor), Positions[0], Positions[1]);
