Motion Design: fix issue where having an operator stack, or material designer details opened would cause sequencer to not be closeable

diff --git a/Engine/Plugins/Experimental/AvalancheSequenceNavigator/Source/AvalancheSequenceNavigator/Private/AvaNavigationToolProvider.cpp b/Engine/Plugins/Experimental/AvalancheSequenceNavigator/Source/AvalancheSequenceNavigator/Private/AvaNavigationToolProvider.cpp
--- a/Engine/Plugins/Experimental/AvalancheSequenceNavigator/Source/AvalancheSequenceNavigator/Private/AvaNavigationToolProvider.cpp
+++ b/Engine/Plugins/Experimental/AvalancheSequenceNavigator/Source/AvalancheSequenceNavigator/Private/AvaNavigationToolProvider.cpp
@@ -286,7 +286,10 @@ TSharedPtr<SequenceNavigator::INavigationTool> FAvaNavigationToolProvider::GetNa
 {
 	if (const TSharedPtr<IAvaSequencer> AvaSequencer = WeakAvaSequencer.Pin())
 	{
-		return SequenceNavigator::FNavigationToolExtender::FindNavigationTool(AvaSequencer->GetSequencer());
+		if (TSharedPtr<ISequencer> Sequencer = AvaSequencer->GetSequencerPtr())
+		{
+			return SequenceNavigator::FNavigationToolExtender::FindNavigationTool(Sequencer.ToSharedRef());
+		}
 	}
 	return nullptr;
 }
@@ -452,7 +455,7 @@ void FAvaNavigationToolProvider::OnSequenceAdded(UAvaSequence* const InAvaSequen
 		return;
 	}
 
-	const TSharedPtr<ISequencer> Sequencer = AvaSequencer->GetSequencer();
+	const TSharedPtr<ISequencer> Sequencer = AvaSequencer->GetSequencerPtr();
 	if (!Sequencer.IsValid())
 	{
 		return;
@@ -602,7 +605,13 @@ void FAvaNavigationToolProvider::RelabelSelection()
 		return;
 	}
 
-	const TSharedPtr<INavigationTool> NavigationTool = FNavigationToolExtender::FindNavigationTool(AvaSequencer->GetSequencer());
+	const TSharedPtr<ISequencer> Sequencer = AvaSequencer->GetSequencerPtr();
+	if (!Sequencer.IsValid())
+	{
+		return;
+	}
+
+	const TSharedPtr<INavigationTool> NavigationTool = FNavigationToolExtender::FindNavigationTool(Sequencer.ToSharedRef());
 	if (!NavigationTool.IsValid())
 	{
 		return;
@@ -646,7 +655,13 @@ void FAvaNavigationToolProvider::AddSequenceToSelection()
 		return;
 	}
 
-	const TSharedPtr<INavigationTool> NavigationTool = FNavigationToolExtender::FindNavigationTool(AvaSequencer->GetSequencer());
+	const TSharedPtr<ISequencer> Sequencer = AvaSequencer->GetSequencerPtr();
+	if (!Sequencer.IsValid())
+	{
+		return;
+	}
+
+	const TSharedPtr<INavigationTool> NavigationTool = FNavigationToolExtender::FindNavigationTool(Sequencer.ToSharedRef());
 	if (!NavigationTool.IsValid())
 	{
 		return;
diff --git a/Engine/Plugins/Experimental/AvalancheSequenceNavigator/Source/AvalancheSequenceNavigator/Private/AvalancheSequenceNavigatorModule.cpp b/Engine/Plugins/Experimental/AvalancheSequenceNavigator/Source/AvalancheSequenceNavigator/Private/AvalancheSequenceNavigatorModule.cpp
--- a/Engine/Plugins/Experimental/AvalancheSequenceNavigator/Source/AvalancheSequenceNavigator/Private/AvalancheSequenceNavigatorModule.cpp
+++ b/Engine/Plugins/Experimental/AvalancheSequenceNavigator/Source/AvalancheSequenceNavigator/Private/AvalancheSequenceNavigatorModule.cpp
@@ -71,7 +71,7 @@ void FAvalancheSequenceNavigatorModule::OnSequencerClosed(const TSharedRef<ISequ
 
 void FAvalancheSequenceNavigatorModule::OnAvaSequencerCreated(const TSharedRef<IAvaSequencer> InAvaSequencer)
 {
-	const TSharedPtr<ISequencer> Sequencer = InAvaSequencer->GetSequencer();
+	const TSharedPtr<ISequencer> Sequencer = InAvaSequencer->GetSequencerPtr();
 	if (!Sequencer.IsValid())
 	{
 		return;
diff --git a/Engine/Plugins/VirtualProduction/Avalanche/Source/AvalancheEditor/Private/Sequencer/AvaSequencerExtension.cpp b/Engine/Plugins/VirtualProduction/Avalanche/Source/AvalancheEditor/Private/Sequencer/AvaSequencerExtension.cpp
--- a/Engine/Plugins/VirtualProduction/Avalanche/Source/AvalancheEditor/Private/Sequencer/AvaSequencerExtension.cpp
+++ b/Engine/Plugins/VirtualProduction/Avalanche/Source/AvalancheEditor/Private/Sequencer/AvaSequencerExtension.cpp
@@ -37,7 +37,7 @@ FAvaSequencerExtension::FAvaSequencerExtension(FName InSequencerTabId, bool bInS
 
 TSharedPtr<ISequencer> FAvaSequencerExtension::GetSequencer() const
 {
-	return AvaSequencer.IsValid() ? AvaSequencer->GetSequencer() : TSharedPtr<ISequencer>();
+	return AvaSequencer.IsValid() ? AvaSequencer->GetSequencerPtr() : TSharedPtr<ISequencer>();
 }
 
 void FAvaSequencerExtension::OnObjectRenamed(UObject* InRenamedObject, const FText& InDisplayNameText)
@@ -47,7 +47,7 @@ void FAvaSequencerExtension::OnObjectRenamed(UObject* InRenamedObject, const FTe
 		return;
 	}
 
-	if (const TSharedPtr<ISequencer> Sequencer = AvaSequencer->GetSequencer())
+	if (const TSharedPtr<ISequencer> Sequencer = AvaSequencer->GetSequencerPtr())
 	{
 		const FGuid ObjectGuid = Sequencer->FindObjectId(*InRenamedObject, Sequencer->GetFocusedTemplateID());
 		if (ObjectGuid.IsValid())
diff --git a/Engine/Plugins/VirtualProduction/Avalanche/Source/AvalancheSequencer/Private/AvaSequencer.cpp b/Engine/Plugins/VirtualProduction/Avalanche/Source/AvalancheSequencer/Private/AvaSequencer.cpp
--- a/Engine/Plugins/VirtualProduction/Avalanche/Source/AvalancheSequencer/Private/AvaSequencer.cpp
+++ b/Engine/Plugins/VirtualProduction/Avalanche/Source/AvalancheSequencer/Private/AvaSequencer.cpp
@@ -430,7 +430,11 @@ void FAvaSequencer::GetSelectedObjects(const TArray<FGuid>& InObjectGuids
 	, TArray<UObject*>& OutSelectedComponents
 	, TArray<UObject*>& OutSelectedObjects) const
 {
-	TSharedRef<ISequencer> Sequencer = GetSequencer();
+	TSharedPtr<ISequencer> Sequencer = GetSequencerPtr();
+	if (!Sequencer.IsValid())
+	{
+		return;
+	}
 
 	UMovieSceneSequence* const ActiveSequence = Sequencer->GetFocusedMovieSceneSequence();
 	UObject* const PlaybackContext = Provider.GetPlaybackContext();
@@ -878,8 +882,13 @@ void FAvaSequencer::AddTransformKey(EMovieSceneTransformChannel InTransformChann
 		return;
 	}
 
-	TSharedRef<ISequencer> Sequencer = GetSequencer();
-	const TArray<TSharedPtr<ISequencerTrackEditor>>& TrackEditors = StaticCastSharedRef<FSequencer>(Sequencer)->GetTrackEditors();
+	TSharedPtr<ISequencer> Sequencer = GetSequencerPtr();
+	if (!Sequencer.IsValid())
+	{
+		return;
+	}
+
+	const TArray<TSharedPtr<ISequencerTrackEditor>>& TrackEditors = StaticCastSharedPtr<FSequencer>(Sequencer)->GetTrackEditors();
 
 	TArray<TSharedPtr<ISequencerTrackEditor>> TransformTrackEditors;
 
@@ -1450,8 +1459,10 @@ void FAvaSequencer::ApplyCurrentState()
 
 	if (Response == EAppReturnType::Yes)
 	{
-		TSharedRef<ISequencer> Sequencer = GetSequencer();
-		Sequencer->PreAnimatedState.DiscardPreAnimatedState();
+		if (TSharedPtr<ISequencer> Sequencer = GetSequencerPtr())
+		{
+			Sequencer->PreAnimatedState.DiscardPreAnimatedState();
+		}
 	}
 }
 
@@ -1476,8 +1487,10 @@ void FAvaSequencer::FixBindingPaths()
 		return;
 	}
 
-	TSharedRef<ISequencer> Sequencer = GetSequencer();
-	Sequencer->NotifyMovieSceneDataChanged(EMovieSceneDataChangeType::MovieSceneStructureItemsChanged);
+	if (TSharedPtr<ISequencer> Sequencer = GetSequencerPtr())
+	{
+		Sequencer->NotifyMovieSceneDataChanged(EMovieSceneDataChangeType::MovieSceneStructureItemsChanged);
+	}
 }
 
 void FAvaSequencer::FixInvalidBindings()
@@ -1656,7 +1669,11 @@ void FAvaSequencer::FixBindingHierarchy()
 		return;
 	}
 
-	TSharedRef<ISequencer> Sequencer = GetSequencer();
+	TSharedPtr<ISequencer> Sequencer = GetSequencerPtr();
+	if (!Sequencer.IsValid())
+	{
+		return;
+	}
 
 	FScopedTransaction Transaction(LOCTEXT("FixBindingHierarchy", "Fix Binding Hierarchy"));
 	MovieScene->Modify();
@@ -1686,7 +1703,7 @@ void FAvaSequencer::FixBindingHierarchy()
 
 	// Pass #1: Ensure that all the Parent Objects have a valid possessable handle
 	ForEachPossessableWithFoundParent(
-		[Sequencer](UObject& InPossessableObject, FMovieScenePossessable& InPossessable, UObject& InParentObject)
+		[&Sequencer](UObject& InPossessableObject, FMovieScenePossessable& InPossessable, UObject& InParentObject)
 		{
 			constexpr bool bCreateHandleToObject = true;
 			Sequencer->GetHandleToObject(&InParentObject, bCreateHandleToObject);
@@ -1694,7 +1711,7 @@ void FAvaSequencer::FixBindingHierarchy()
 
 	// Pass #2: Fix the hierarchy now that all relevant objects have a valid handle
 	ForEachPossessableWithFoundParent(
-		[Sequencer, MovieScene, Sequence, PlaybackContext](UObject& InObject, FMovieScenePossessable& InPossessable, UObject& InParentObject)
+		[&Sequencer, MovieScene, Sequence, PlaybackContext](UObject& InObject, FMovieScenePossessable& InPossessable, UObject& InParentObject)
 		{
 			constexpr bool bCreateHandleToObject = false;
 			const FGuid ParentGuid = Sequencer->GetHandleToObject(&InParentObject, bCreateHandleToObject);
@@ -1729,7 +1746,12 @@ void FAvaSequencer::FixBindingHierarchy()
 
 TArrayView<TWeakObjectPtr<>> FAvaSequencer::ResolveBoundObjects(const FGuid& InBindingId, class UMovieSceneSequence* Sequence) const
 {
-	TSharedRef<ISequencer> Sequencer = GetSequencer();
+	TSharedPtr<ISequencer> Sequencer = GetSequencerPtr();
+	if (!Sequencer.IsValid())
+	{
+		return TArrayView<TWeakObjectPtr<>>();
+	}
+
 	TSharedRef<UE::MovieScene::FSharedPlaybackState> SharedPlaybackState = Sequencer->GetSharedPlaybackState();
 	// TODO: It would be better if FAvaSequencer saved the SequenceID. It's possible looking that this might always be FMovieSceneSequenceID::Root, but I'm unsure.
 	if (FMovieSceneEvaluationState* EvaluationState = SharedPlaybackState->FindCapability<FMovieSceneEvaluationState>())
@@ -1740,12 +1762,17 @@ TArrayView<TWeakObjectPtr<>> FAvaSequencer::ResolveBoundObjects(const FGuid& InB
 	return TArrayView<TWeakObjectPtr<>>();
 }
 
-TSharedRef<ISequencer> FAvaSequencer::GetSequencer() const
+TSharedRef<ISequencer> FAvaSequencer::GetOrCreateSequencer() const
 {
 	const_cast<FAvaSequencer*>(this)->EnsureSequencer();
 	return SequencerWeak.Pin().ToSharedRef();
 }
 
+TSharedPtr<ISequencer> FAvaSequencer::GetSequencerPtr() const
+{
+	return SequencerWeak.Pin();
+}
+
 USequencerSettings* FAvaSequencer::GetSequencerSettings() const
 {
 	if (TSharedPtr<ISequencer> Sequencer = SequencerWeak.Pin())
@@ -1821,7 +1848,7 @@ TSharedRef<SWidget> FAvaSequencer::CreateSequenceWidget()
 		SequencerWeak.Reset();
 	}
 
-	const TSharedRef<ISequencer> Sequencer = GetSequencer();
+	const TSharedRef<ISequencer> Sequencer = GetOrCreateSequencer();
 
 	USequencerSettings* const SequencerSettings = Sequencer->GetSequencerSettings();
 	check(SequencerSettings);
@@ -1886,9 +1913,13 @@ void FAvaSequencer::OnEditorSelectionChanged(const FAvaEditorSelection& InEditor
 		return;
 	}
 
-	TGuardValue<bool> Guard(bUpdatingSelection, true);
+	TSharedPtr<ISequencer> Sequencer = GetSequencerPtr();
+	if (!Sequencer.IsValid())
+	{
+		return;
+	}
 
-	TSharedRef<ISequencer> Sequencer = GetSequencer();
+	TGuardValue<bool> Guard(bUpdatingSelection, true);
 	Sequencer->EmptySelection();
 
 	using namespace UE::Sequencer;
@@ -2152,7 +2183,13 @@ void FAvaSequencer::ExtendSidebarMarkedFramesMenu(FMenuBuilder& OutMenuBuilder)
 		return;
 	}
 
-	const TSharedPtr<UE::Sequencer::FSequencerSelection> SequencerSelection = GetSequencer()->GetViewModel()->GetSelection();
+	TSharedPtr<ISequencer> Sequencer = GetSequencerPtr();
+	if (!Sequencer.IsValid())
+	{
+		return;
+	}
+
+	const TSharedPtr<UE::Sequencer::FSequencerSelection> SequencerSelection = Sequencer->GetViewModel()->GetSelection();
 	if (!SequencerSelection.IsValid())
 	{
 		return;
@@ -2172,7 +2209,7 @@ void FAvaSequencer::ExtendSidebarMarkedFramesMenu(FMenuBuilder& OutMenuBuilder)
 
 void FAvaSequencer::OnSequencerClosed(const TSharedRef<ISequencer> InSequencer)
 {
-	if (InSequencer != GetSequencer())
+	if (InSequencer != GetSequencerPtr())
 	{
 		return;
 	}
diff --git a/Engine/Plugins/VirtualProduction/Avalanche/Source/AvalancheSequencer/Private/AvaSequencer.h b/Engine/Plugins/VirtualProduction/Avalanche/Source/AvalancheSequencer/Private/AvaSequencer.h
--- a/Engine/Plugins/VirtualProduction/Avalanche/Source/AvalancheSequencer/Private/AvaSequencer.h
+++ b/Engine/Plugins/VirtualProduction/Avalanche/Source/AvalancheSequencer/Private/AvaSequencer.h
@@ -170,9 +170,11 @@ public:
 
 	TArrayView<TWeakObjectPtr<>> ResolveBoundObjects(const FGuid& InBindingId, UMovieSceneSequence* Sequence) const;
 
+	TSharedRef<ISequencer> GetOrCreateSequencer() const;
+
 	//~ Begin IAvaSequencer
 	virtual const IAvaSequencerProvider& GetProvider() const override { return Provider; }
-	virtual TSharedRef<ISequencer> GetSequencer() const override;
+	virtual TSharedPtr<ISequencer> GetSequencerPtr() const override;
 	virtual USequencerSettings* GetSequencerSettings() const override;
 	virtual void SetBaseCommandList(const TSharedPtr<FUICommandList>& InBaseCommandList) override;
 	virtual UAvaSequence* GetViewedSequence() const override;
diff --git a/Engine/Plugins/VirtualProduction/Avalanche/Source/AvalancheSequencer/Private/Clipboard/AvaSequenceExporter.cpp b/Engine/Plugins/VirtualProduction/Avalanche/Source/AvalancheSequencer/Private/Clipboard/AvaSequenceExporter.cpp
--- a/Engine/Plugins/VirtualProduction/Avalanche/Source/AvalancheSequencer/Private/Clipboard/AvaSequenceExporter.cpp
+++ b/Engine/Plugins/VirtualProduction/Avalanche/Source/AvalancheSequencer/Private/Clipboard/AvaSequenceExporter.cpp
@@ -31,11 +31,20 @@ namespace UE::AvaSequencer::Private
 			BoundActors = InBoundActors;
 		}
 
-		TConstArrayView<AActor*> GetBoundActors() const { return BoundActors; }
+		TConstArrayView<AActor*> GetBoundActors() const
+		{
+			return BoundActors;
+		}
 
-		UObject* GetPlaybackContext() const { return AvaSequencerWeak.IsValid() ? AvaSequencerWeak.Pin()->GetPlaybackContext() : nullptr; }
+		UObject* GetPlaybackContext() const
+		{
+			return AvaSequencerWeak.IsValid() ? AvaSequencerWeak.Pin()->GetPlaybackContext() : nullptr;
+		}
 
-		TSharedPtr<FAvaSequencer> GetSequencer() const { return AvaSequencerWeak.Pin(); }
+		TSharedPtr<ISequencer> GetSequencer() const
+		{
+			return AvaSequencerWeak.IsValid() ? AvaSequencerWeak.Pin()->GetSequencerPtr() : nullptr;
+		}
 
 	private:
 		TArray<TObjectPtr<AActor>> BoundActors;
@@ -121,14 +130,12 @@ bool UAvaSequenceExporter::ExportText(const FExportObjectInnerContext* InContext
 	const UE::AvaSequencer::Private::FAvaSequenceExportObjectInnerContext& Context
 		= static_cast<const UE::AvaSequencer::Private::FAvaSequenceExportObjectInnerContext&>(*InContext);
 
-	TSharedPtr<FAvaSequencer> Sequencer = Context.GetSequencer();
-	 
-	if (!Sequencer)
+	TSharedPtr<ISequencer> Sequencer = Context.GetSequencer();
+	if (!Sequencer.IsValid())
 	{
 		return false;
 	}
 
-
 	UObject* const PlaybackContext = Context.GetPlaybackContext();
 
 	// Gather guids for the object nodes and any child object nodes
@@ -174,7 +181,7 @@ bool UAvaSequenceExporter::ExportText(const FExportObjectInnerContext* InContext
 		, *Sequence->GetLabel().ToString());
 
 	TArray<UMovieSceneFolder*> Folders;
-	FSequencerUtilities::CopyBindings(Sequencer->GetSequencer(), Bindings, Folders, Ar);
+	FSequencerUtilities::CopyBindings(Sequencer.ToSharedRef(), Bindings, Folders, Ar);
 
 	Ar.Logf(TEXT("%sEnd Sequence\r\n"), FCString::Spc(TextIndent));
 
diff --git a/Engine/Plugins/VirtualProduction/Avalanche/Source/AvalancheSequencer/Private/Clipboard/AvaSequenceImporter.cpp b/Engine/Plugins/VirtualProduction/Avalanche/Source/AvalancheSequencer/Private/Clipboard/AvaSequenceImporter.cpp
--- a/Engine/Plugins/VirtualProduction/Avalanche/Source/AvalancheSequencer/Private/Clipboard/AvaSequenceImporter.cpp
+++ b/Engine/Plugins/VirtualProduction/Avalanche/Source/AvalancheSequencer/Private/Clipboard/AvaSequenceImporter.cpp
@@ -27,7 +27,12 @@ void FAvaSequenceImporter::ImportText(FStringView InPastedData, const TMap<FName
 	}
 
 	IAvaSequencerProvider& Provider  = AvaSequencer->GetProvider();
-	TSharedRef<ISequencer> Sequencer = AvaSequencer->GetSequencer();
+
+	TSharedPtr<ISequencer> Sequencer = AvaSequencer->GetSequencerPtr();
+	if (!Sequencer.IsValid())
+	{
+		return;
+	}
 
 	UAvaSequence* const OriginallyViewedSequence = AvaSequencer->GetViewedSequence();
 
@@ -72,7 +77,7 @@ void FAvaSequenceImporter::ImportText(FStringView InPastedData, const TMap<FName
 		TArray<FMovieSceneBindingProxy> OutBindings;
 		FMovieScenePasteBindingsParams PasteParams;
 		Algo::Transform(InPastedActors, PasteParams.PastedActors, [](TTuple<FName, AActor*> Pair) { return TTuple<FName, TObjectPtr<AActor>>(Pair.Key, Pair.Value);});
-		FSequencerUtilities::PasteBindings(BindingsString, Sequencer, PasteParams, OutBindings, PasteErrors);
+		FSequencerUtilities::PasteBindings(BindingsString, Sequencer.ToSharedRef(), PasteParams, OutBindings, PasteErrors);
 	}
 
 	AvaSequencer->SetViewedSequence(OriginallyViewedSequence);
diff --git a/Engine/Plugins/VirtualProduction/Avalanche/Source/AvalancheSequencer/Private/Commands/AvaSequencerAction.cpp b/Engine/Plugins/VirtualProduction/Avalanche/Source/AvalancheSequencer/Private/Commands/AvaSequencerAction.cpp
--- a/Engine/Plugins/VirtualProduction/Avalanche/Source/AvalancheSequencer/Private/Commands/AvaSequencerAction.cpp
+++ b/Engine/Plugins/VirtualProduction/Avalanche/Source/AvalancheSequencer/Private/Commands/AvaSequencerAction.cpp
@@ -8,7 +8,13 @@ TSharedPtr<UE::Sequencer::FSequencerSelection> FAvaSequencerAction::GetSequencer
 {
 	using namespace UE::Sequencer;
 
-	if (TSharedPtr<FSequencerEditorViewModel> EditorViewModel = Owner.GetSequencer()->GetViewModel())
+	TSharedPtr<ISequencer> Sequencer = Owner.GetSequencerPtr();
+	if (!Sequencer.IsValid())
+	{
+		return nullptr;
+	}
+
+	if (TSharedPtr<FSequencerEditorViewModel> EditorViewModel = Sequencer->GetViewModel())
 	{
 		return EditorViewModel->GetSelection();
 	}
diff --git a/Engine/Plugins/VirtualProduction/Avalanche/Source/AvalancheSequencer/Private/DetailsView/SAvaKeyFrameEdit.cpp b/Engine/Plugins/VirtualProduction/Avalanche/Source/AvalancheSequencer/Private/DetailsView/SAvaKeyFrameEdit.cpp
--- a/Engine/Plugins/VirtualProduction/Avalanche/Source/AvalancheSequencer/Private/DetailsView/SAvaKeyFrameEdit.cpp
+++ b/Engine/Plugins/VirtualProduction/Avalanche/Source/AvalancheSequencer/Private/DetailsView/SAvaKeyFrameEdit.cpp
@@ -11,22 +11,20 @@ using namespace UE::Sequencer;
 
 #define LOCTEXT_NAMESPACE "SAvaKeyFrameEdit"
 
-void SAvaKeyFrameEdit::Construct(const FArguments& InArgs, const TSharedRef<FAvaSequencer>& InSequencer)
+void SAvaKeyFrameEdit::Construct(const FArguments& InArgs, const TSharedRef<FAvaSequencer>& InAvaSequencer)
 {
-	AvaSequencerWeak = InSequencer;
-
+	AvaSequencerWeak = InAvaSequencer;
 	KeyEditData = InArgs._KeyEditData;
 
-	const TSharedRef<ISequencer> Sequencer = InSequencer->GetSequencer();
-
-	if (const TSharedPtr<FSequencerEditorViewModel> SequencerViewModel = Sequencer->GetViewModel())
+	const TSharedPtr<ISequencer> Sequencer = InAvaSequencer->GetSequencerPtr();
+	if (!Sequencer.IsValid())
 	{
-		SequencerSelectionWeak = SequencerViewModel->GetSelection();
+		return;
 	}
 
-	if (!AvaSequencerWeak.IsValid())
+	if (const TSharedPtr<FSequencerEditorViewModel> SequencerViewModel = Sequencer->GetViewModel())
 	{
-		return;
+		SequencerSelectionWeak = SequencerViewModel->GetSelection();
 	}
 
 	const TSharedPtr<FAvaSequencer> AvaSequencer = AvaSequencerWeak.Pin();
@@ -38,7 +36,7 @@ void SAvaKeyFrameEdit::Construct(const FArguments& InArgs, const TSharedRef<FAva
 			.AutoSize()
 			.HAlign(HAlign_Fill)
 			[
-				SNew(SKeyEditInterface, AvaSequencer->GetSequencer())
+				SNew(SKeyEditInterface, Sequencer.ToSharedRef())
 				.EditData(KeyEditData)
 			]
 		];
diff --git a/Engine/Plugins/VirtualProduction/Avalanche/Source/AvalancheSequencer/Private/Outliner/AvaOutlinerSequenceProxy.cpp b/Engine/Plugins/VirtualProduction/Avalanche/Source/AvalancheSequencer/Private/Outliner/AvaOutlinerSequenceProxy.cpp
--- a/Engine/Plugins/VirtualProduction/Avalanche/Source/AvalancheSequencer/Private/Outliner/AvaOutlinerSequenceProxy.cpp
+++ b/Engine/Plugins/VirtualProduction/Avalanche/Source/AvalancheSequencer/Private/Outliner/AvaOutlinerSequenceProxy.cpp
@@ -27,7 +27,6 @@ FAvaOutlinerSequenceProxy::FAvaOutlinerSequenceProxy(IAvaOutliner& InOutliner, c
 		if (TSharedPtr<IAvaSequencer> AvaSequencer = SequencerSubsystem->GetSequencer())
 		{
 			AvaSequencerWeak = AvaSequencer;
-			SequencerWeak    = AvaSequencer->GetSequencer();
 		}
 	}
 }
@@ -127,22 +126,32 @@ void FAvaOutlinerSequenceProxy::GetProxiedItems(const TSharedRef<IAvaOutlinerIte
 void FAvaOutlinerSequenceProxy::BindDelegates()
 {
 	UnbindDelegates();
-	if (TSharedPtr<ISequencer> Sequencer = SequencerWeak.Pin())
+
+	TSharedPtr<IAvaSequencer> AvaSequencer = AvaSequencerWeak.Pin();
+	if (!AvaSequencer.IsValid())
+	{
+		return;
+	}
+
+	if (TSharedPtr<ISequencer> Sequencer = AvaSequencer->GetSequencerPtr())
 	{
 		OnMovieSceneDataChangedHandle = Sequencer->OnMovieSceneDataChanged().AddSP(this, &FAvaOutlinerSequenceProxy::OnMovieSceneDataChanged);
 		OnActorAddedToSequencerHandle = Sequencer->OnActorAddedToSequencer().AddSP(this, &FAvaOutlinerSequenceProxy::OnActorAddedToSequencer);
+		LastBoundSequencerWeak = Sequencer;
 	}
 }
 
 void FAvaOutlinerSequenceProxy::UnbindDelegates()
 {
-	if (TSharedPtr<ISequencer> Sequencer = SequencerWeak.Pin())
+	if (TSharedPtr<ISequencer> Sequencer = LastBoundSequencerWeak.Pin())
 	{
 		Sequencer->OnMovieSceneDataChanged().Remove(OnMovieSceneDataChangedHandle);
 		Sequencer->OnActorAddedToSequencer().Remove(OnActorAddedToSequencerHandle);
 
 		OnMovieSceneDataChangedHandle.Reset();
 		OnActorAddedToSequencerHandle.Reset();
+
+		LastBoundSequencerWeak.Reset();
 	}
 }
 
diff --git a/Engine/Plugins/VirtualProduction/Avalanche/Source/AvalancheSequencer/Private/Outliner/AvaOutlinerSequenceProxy.h b/Engine/Plugins/VirtualProduction/Avalanche/Source/AvalancheSequencer/Private/Outliner/AvaOutlinerSequenceProxy.h
--- a/Engine/Plugins/VirtualProduction/Avalanche/Source/AvalancheSequencer/Private/Outliner/AvaOutlinerSequenceProxy.h
+++ b/Engine/Plugins/VirtualProduction/Avalanche/Source/AvalancheSequencer/Private/Outliner/AvaOutlinerSequenceProxy.h
@@ -47,7 +47,8 @@ private:
 
 	TWeakPtr<IAvaSequencer> AvaSequencerWeak;
 
-	TWeakPtr<ISequencer> SequencerWeak;
+	/** Used to keep track of the latest Sequencer instance that our delegates are bound to */
+	TWeakPtr<ISequencer> LastBoundSequencerWeak;
 
 	FDelegateHandle OnMovieSceneDataChangedHandle;
 
diff --git a/Engine/Plugins/VirtualProduction/Avalanche/Source/AvalancheSequencer/Private/StaggerTool/AvaStaggerTool.cpp b/Engine/Plugins/VirtualProduction/Avalanche/Source/AvalancheSequencer/Private/StaggerTool/AvaStaggerTool.cpp
--- a/Engine/Plugins/VirtualProduction/Avalanche/Source/AvalancheSequencer/Private/StaggerTool/AvaStaggerTool.cpp
+++ b/Engine/Plugins/VirtualProduction/Avalanche/Source/AvalancheSequencer/Private/StaggerTool/AvaStaggerTool.cpp
@@ -95,7 +95,9 @@ void FAvaStaggerTool::Execute()
 			ToolWindow = nullptr;
 		}));
 
-	const TSharedPtr<SWindow> ParentWindow = FSlateApplication::Get().FindBestParentWindowForDialogs(Owner.GetSequencer()->GetSequencerWidget());
+	const TSharedPtr<ISequencer> Sequencer = Owner.GetSequencerPtr();
+
+	const TSharedPtr<SWindow> ParentWindow = Sequencer.IsValid() ? FSlateApplication::Get().FindBestParentWindowForDialogs(Sequencer->GetSequencerWidget()) : nullptr;
 	if (ParentWindow.IsValid())
 	{
 		FSlateApplication::Get().AddWindowAsNativeChild(ToolWindow.ToSharedRef(), ParentWindow.ToSharedRef());
@@ -110,7 +112,12 @@ void FAvaStaggerTool::Execute()
 
 void FAvaStaggerTool::BindToSequencerSelectionChange()
 {
-	const TSharedRef<ISequencer> Sequencer = Owner.GetSequencer();
+	TSharedPtr<ISequencer> Sequencer = Owner.GetSequencerPtr();
+	if (!Sequencer.IsValid())
+	{
+		return;
+	}
+
 	if (const TSharedPtr<FSequencerEditorViewModel> ViewModel = Sequencer->GetViewModel())
 	{
 		if (const TSharedPtr<FSequencerSelection> Selection = ViewModel->GetSelection())
@@ -122,7 +129,12 @@ void FAvaStaggerTool::BindToSequencerSelectionChange()
 
 void FAvaStaggerTool::UnbindToSequencerSelectionChange()
 {
-	const TSharedRef<ISequencer> Sequencer = Owner.GetSequencer();
+	TSharedPtr<ISequencer> Sequencer = Owner.GetSequencerPtr();
+	if (!Sequencer.IsValid())
+	{
+		return;
+	}
+
 	if (const TSharedPtr<FSequencerEditorViewModel> ViewModel = Sequencer->GetViewModel())
 	{
 		if (const TSharedPtr<FSequencerSelection> Selection = ViewModel->GetSelection())
@@ -213,7 +225,7 @@ int32 FAvaStaggerTool::GetSelectionCount() const
 
 FFrameNumber FAvaStaggerTool::ConvertTime(FFrameTime InSourceTime) const
 {
-	const TSharedRef<ISequencer> Sequencer = Owner.GetSequencer();
+	TSharedRef<ISequencer> Sequencer = Owner.GetOrCreateSequencer();
 	const FFrameRate FocusedDisplayRate = Sequencer->GetFocusedDisplayRate();
 	const FFrameRate FocusedTickResolution = Sequencer->GetFocusedTickResolution();
 	return ConvertFrameTime(InSourceTime, FocusedDisplayRate, FocusedTickResolution).RoundToFrame();
@@ -221,7 +233,7 @@ FFrameNumber FAvaStaggerTool::ConvertTime(FFrameTime InSourceTime) const
 
 TRange<FFrameNumber> FAvaStaggerTool::GetOperationRange(const EAvaSequencerStaggerRange InRange) const
 {
-	const TSharedRef<ISequencer> Sequencer = Owner.GetSequencer();
+	TSharedRef<ISequencer> Sequencer = Owner.GetOrCreateSequencer();
 
 	switch (InRange)
 	{
@@ -320,12 +332,12 @@ FFrameNumber FAvaStaggerTool::FindFirstBarStaggerPoint() const
 		switch (Settings->ToolOptions.Range)
 		{
 		case EAvaSequencerStaggerRange::Playback:
-			OutFrameNumber = Owner.GetSequencer()->GetPlaybackRange().GetLowerBoundValue();
+			OutFrameNumber = Owner.GetOrCreateSequencer()->GetPlaybackRange().GetLowerBoundValue();
 			OutFrameNumber += CachedShiftFrames;
 			return MoveTemp(OutFrameNumber);
 
 		case EAvaSequencerStaggerRange::Selection:
-			OutFrameNumber = Owner.GetSequencer()->GetSelectionRange().GetLowerBoundValue();
+			OutFrameNumber = Owner.GetOrCreateSequencer()->GetSelectionRange().GetLowerBoundValue();
 			OutFrameNumber += CachedShiftFrames;
 			return MoveTemp(OutFrameNumber);
 
@@ -690,12 +702,16 @@ void FAvaStaggerTool::StaggerKeyElements()
 
 bool FAvaStaggerTool::CanAlignToPlayhead() const
 {
-	return FSequencerSelectionAlignmentUtils::CanAlignSelection(*Owner.GetSequencer());
+	TSharedPtr<ISequencer> Sequencer = Owner.GetSequencerPtr();
+	return Sequencer.IsValid() ? FSequencerSelectionAlignmentUtils::CanAlignSelection(*Sequencer) : false;
 }
 
 void FAvaStaggerTool::AlignToPlayhead()
 {
-	FSequencerSelectionAlignmentUtils::AlignSelectionToPlayhead(*Owner.GetSequencer());
+	if (TSharedPtr<ISequencer> Sequencer = Owner.GetSequencerPtr())
+	{
+		FSequencerSelectionAlignmentUtils::AlignSelectionToPlayhead(*Sequencer);	
+	}
 }
 
 void FAvaStaggerTool::SetKeyElementTime(const FAvaStaggerKeyElement& InElement, const FFrameNumber InKeyTime)
diff --git a/Engine/Plugins/VirtualProduction/Avalanche/Source/AvalancheSequencer/Private/StaggerTool/AvaStaggerTool.h b/Engine/Plugins/VirtualProduction/Avalanche/Source/AvalancheSequencer/Private/StaggerTool/AvaStaggerTool.h
--- a/Engine/Plugins/VirtualProduction/Avalanche/Source/AvalancheSequencer/Private/StaggerTool/AvaStaggerTool.h
+++ b/Engine/Plugins/VirtualProduction/Avalanche/Source/AvalancheSequencer/Private/StaggerTool/AvaStaggerTool.h
@@ -96,15 +96,15 @@ private:
 			break;
 
 		case EAvaSequencerStaggerStartPosition::Playhead:
-			OutFrameNumber = Owner.GetSequencer()->GetGlobalTime().Time.FrameNumber;
+			OutFrameNumber = Owner.GetOrCreateSequencer()->GetGlobalTime().Time.FrameNumber;
 			break;
 
 		case EAvaSequencerStaggerStartPosition::PlaybackRange:
-			OutFrameNumber = Owner.GetSequencer()->GetPlaybackRange().GetLowerBoundValue();
+			OutFrameNumber = Owner.GetOrCreateSequencer()->GetPlaybackRange().GetLowerBoundValue();
 			break;
 
 		case EAvaSequencerStaggerStartPosition::SelectionRange:
-			OutFrameNumber = Owner.GetSequencer()->GetSelectionRange().GetLowerBoundValue();
+			OutFrameNumber = Owner.GetOrCreateSequencer()->GetSelectionRange().GetLowerBoundValue();
 			break;
 		}
 
diff --git a/Engine/Plugins/VirtualProduction/Avalanche/Source/AvalancheSequencer/Public/IAvaSequencer.h b/Engine/Plugins/VirtualProduction/Avalanche/Source/AvalancheSequencer/Public/IAvaSequencer.h
--- a/Engine/Plugins/VirtualProduction/Avalanche/Source/AvalancheSequencer/Public/IAvaSequencer.h
+++ b/Engine/Plugins/VirtualProduction/Avalanche/Source/AvalancheSequencer/Public/IAvaSequencer.h
@@ -30,7 +30,14 @@ public:
 
 	virtual const IAvaSequencerProvider& GetProvider() const = 0;
 
-	virtual TSharedRef<ISequencer> GetSequencer() const = 0;
+	UE_DEPRECATED(5.7, "Get sequencer by shared reference is deprecated. Use GetSequencerPtr instead.")
+	virtual TSharedRef<ISequencer> GetSequencer() const
+	{
+		return GetSequencerPtr().ToSharedRef();
+	}
+
+	/** Gets the sequencer ptr if available. */
+	virtual TSharedPtr<ISequencer> GetSequencerPtr() const = 0;
 
 	virtual USequencerSettings* GetSequencerSettings() const = 0;
 
