Motion Design: Post process settings
- Post process settings are now applied correctly on editor startup and map change.

diff --git a/Engine/Plugins/VirtualProduction/Avalanche/Source/AvalancheLevelViewport/Private/ViewportClient/AvaLevelViewportClient.cpp b/Engine/Plugins/VirtualProduction/Avalanche/Source/AvalancheLevelViewport/Private/ViewportClient/AvaLevelViewportClient.cpp
--- a/Engine/Plugins/VirtualProduction/Avalanche/Source/AvalancheLevelViewport/Private/ViewportClient/AvaLevelViewportClient.cpp
+++ b/Engine/Plugins/VirtualProduction/Avalanche/Source/AvalancheLevelViewport/Private/ViewportClient/AvaLevelViewportClient.cpp
@@ -18,6 +18,8 @@
 #include "Interaction/AvaDragOperation.h"
 #include "Interaction/AvaIsolateActorsOperation.h"
 #include "Interaction/AvaSnapOperation.h"
+#include "LevelEditor.h"
+#include "Modules/ModuleManager.h"
 #include "MouseDeltaTracker.h"
 #include "SAvaLevelViewport.h"
 #include "SceneView.h"
@@ -64,6 +66,11 @@ FAvaLevelViewportClient::FAvaLevelViewportClient()
 FAvaLevelViewportClient::~FAvaLevelViewportClient()
 {
 	UE::AvaLevelViewport::Private::ViewportClientRegistry.Remove(this);
+
+	if (FLevelEditorModule* LevelEditorModule = FModuleManager::GetModulePtr<FLevelEditorModule>(TEXT("LevelEditor")))
+	{
+		LevelEditorModule->OnMapChanged().RemoveAll(this);
+	}
 }
 
 void FAvaLevelViewportClient::Init()
@@ -72,7 +79,14 @@ void FAvaLevelViewportClient::Init()
 	PostProcessManager = MakeShared<FAvaViewportPostProcessManager>(SharedThis(this));
 	IsolateActorsOperation = MakeShared<FAvaIsolateActorsOperation>(SharedThis(this));
 
+	PostProcessManager->LoadPostProcessInfo();
+
 	LoadActiveCameraComponent();
+
+	if (FLevelEditorModule* LevelEditorModule = FModuleManager::GetModulePtr<FLevelEditorModule>(TEXT("LevelEditor")))
+	{
+		LevelEditorModule->OnMapChanged().AddSP(this, &FAvaLevelViewportClient::OnMapChanged);
+	}
 }
 
 bool FAvaLevelViewportClient::HasCachedViewportSize() const
@@ -802,4 +816,24 @@ void FAvaLevelViewportClient::SetActiveCameraComponent(UCameraComponent* InCamer
 	}
 }
 
+void FAvaLevelViewportClient::OnMapChanged(UWorld* InWorld, EMapChangeType InChangeType)
+{
+	switch (InChangeType)
+	{
+		case EMapChangeType::NewMap:
+		case EMapChangeType::LoadMap:
+		{
+			if (PostProcessManager.IsValid())
+			{
+				PostProcessManager->LoadPostProcessInfo();
+			}
+			break;
+		}
+
+		default:
+			// Do nothing
+			break;
+	}
+}
+
 #undef LOCTEXT_NAMESPACE
diff --git a/Engine/Plugins/VirtualProduction/Avalanche/Source/AvalancheLevelViewport/Private/ViewportClient/AvaLevelViewportClient.h b/Engine/Plugins/VirtualProduction/Avalanche/Source/AvalancheLevelViewport/Private/ViewportClient/AvaLevelViewportClient.h
--- a/Engine/Plugins/VirtualProduction/Avalanche/Source/AvalancheLevelViewport/Private/ViewportClient/AvaLevelViewportClient.h
+++ b/Engine/Plugins/VirtualProduction/Avalanche/Source/AvalancheLevelViewport/Private/ViewportClient/AvaLevelViewportClient.h
@@ -73,7 +73,6 @@ public:
 	virtual void UpdateMouseDelta() override;
 	virtual EMouseCursor::Type GetCursor(FViewport* InViewport, int32 InX, int32 InY) override;
 	virtual FMatrix CalcViewRotationMatrixForControllingActorView(const FRotator& InViewRotation) const override;
-
 	//~ End FEditorViewportClient
 
 	//~ Begin IAvaViewportWorldCoordinateConverter
@@ -146,4 +145,6 @@ protected:
 	UCameraComponent* UpdateActiveCameraComponent(AActor* InViewTarget);
 
 	void SetActiveCameraComponent(UCameraComponent* InCameraComponent);
+
+	void OnMapChanged(UWorld* InWorld, EMapChangeType InChangeType);
 };
