MotionDesign : ActorModifier
- Fixed masking issue with nested masks when changing position of mask actor, Mask Write and Mask Read would fight for material control when nested
- Fixed adding new components to mask input actor in details panel fails to update

diff --git a/Engine/Plugins/VirtualProduction/Avalanche/Source/AvalancheMask/Private/Mask2D/AvaMask2DBaseModifier.cpp b/Engine/Plugins/VirtualProduction/Avalanche/Source/AvalancheMask/Private/Mask2D/AvaMask2DBaseModifier.cpp
--- a/Engine/Plugins/VirtualProduction/Avalanche/Source/AvalancheMask/Private/Mask2D/AvaMask2DBaseModifier.cpp
+++ b/Engine/Plugins/VirtualProduction/Avalanche/Source/AvalancheMask/Private/Mask2D/AvaMask2DBaseModifier.cpp
@@ -346,14 +346,16 @@ void UAvaMask2DBaseModifier::OnModifiedActorTransformed()
 void UAvaMask2DBaseModifier::SavePreState()
 {
 	Super::SavePreState();
-	
+
+	ScannedPrimitiveComponents.Empty();
+
 	TMap<TWeakObjectPtr<AActor>, FAvaMask2DActorData> NewActorDatas;
 	NewActorDatas.Reserve(ActorData.Num());
 
 	ForEachActor<AActor>(
 		[this, &NewActorDatas](AActor* InActor)
 		{
-			const UActorModifierCoreBase* MaskModifier = UE::ActorModifierCore::Utilities::FindFirstActorModifierByClass(InActor, GetClass());
+			const UActorModifierCoreBase* MaskModifier = UE::ActorModifierCore::Utilities::FindFirstActorModifierByClass(InActor, UAvaMask2DBaseModifier::StaticClass());
 
 			// If first modifier in the hierarchy is not this modifier, then skip actor
 			if (!MaskModifier || MaskModifier == this)
@@ -382,7 +384,7 @@ void UAvaMask2DBaseModifier::RestorePreState()
 	}
 	
 	TArray<AActor*> ActorDataActors;
-	ActorDataActors.Reserve(ActorData.Num());	
+	ActorDataActors.Reserve(ActorData.Num());
 	Algo::Transform(ActorData, ActorDataActors, [](const TPair<TWeakObjectPtr<AActor>, FAvaMask2DActorData>& InActorDataPair)
 	{
 		return InActorDataPair.Key.Get();
@@ -453,6 +455,13 @@ void UAvaMask2DBaseModifier::SaveActorPreState(AActor* InActor, FAvaMask2DActorD
 
 	FInstancedStruct& HandleData = MaterialCollectionHandleData.FindOrAdd(InActor, MaterialCollectionHandle->MakeDataStruct());
 	MaterialCollectionHandle->SaveOriginalState(HandleData);
+
+	TArray<UPrimitiveComponent*> PrimitiveComponents;
+	InActor->GetComponents(PrimitiveComponents);
+	Algo::Transform(PrimitiveComponents, ScannedPrimitiveComponents, [](UPrimitiveComponent* InComponent)
+	{
+		return InComponent;
+	});
 }
 
 void UAvaMask2DBaseModifier::RestoreActorPreState(AActor* InActor, FAvaMask2DActorData& InActorData)
@@ -533,26 +542,35 @@ void UAvaMask2DBaseModifier::OnRenderStateUpdated(AActor* InActor, UActorCompone
 		return;
 	}
 
-	const UPrimitiveComponent* PrimitiveComponent = Cast<UPrimitiveComponent>(InComponent);
-	if (!PrimitiveComponent)
+	const UActorModifierCoreBase* MaskModifier = UE::ActorModifierCore::Utilities::FindFirstActorModifierByClass(InActor, UAvaMask2DBaseModifier::StaticClass());
+
+	// Only trigger update if first modifier found in the hierarchy above us is this modifier
+	if (MaskModifier != this)
 	{
 		return;
 	}
-	
-	if (CanvasWeak.IsValid())
+
+	const UPrimitiveComponent* PrimitiveComponent = Cast<UPrimitiveComponent>(InComponent);
+	if (!PrimitiveComponent)
 	{
 		return;
 	}
 
-	if (const FAvaMask2DActorData* ActorMaskData = ActorData.Find(InActor))
+	const int32 MaterialCount = PrimitiveComponent->GetNumMaterials();
+	if (MaterialCount != 0)
 	{
-		if (ActorMaskData->CanvasTextureWeak.IsValid())
+		if (const FAvaMask2DActorData* MaskData = ActorData.Find(InActor))
 		{
-			return;
+			if (CanvasWeak.IsValid()
+				&& MaskData->CanvasTextureWeak.IsValid()
+				&& ScannedPrimitiveComponents.Contains(PrimitiveComponent))
+			{
+				return;
+			}
 		}
-	}
 
-	MarkModifierDirty();
+		MarkModifierDirty();
+	}
 }
 
 FName UAvaMask2DBaseModifier::GenerateUniqueMaskName() const
@@ -891,7 +909,7 @@ void UAvaMask2DBaseModifier::OnMaterialsChanged(UObject* InMaterialOwner, const
 	{
 		if (const AActor* OwningActor = InMaterialOwner->GetTypedOuter<AActor>())
 		{
-			const UActorModifierCoreBase* MaskModifier = UE::ActorModifierCore::Utilities::FindFirstActorModifierByClass(OwningActor, GetClass());
+			const UActorModifierCoreBase* MaskModifier = UE::ActorModifierCore::Utilities::FindFirstActorModifierByClass(OwningActor, UAvaMask2DBaseModifier::StaticClass());
 
 			// Only trigger update if first modifier found in the hierarchy above us is this modifier
 			if (MaskModifier == this)
diff --git a/Engine/Plugins/VirtualProduction/Avalanche/Source/AvalancheMask/Public/Mask2D/AvaMask2DBaseModifier.h b/Engine/Plugins/VirtualProduction/Avalanche/Source/AvalancheMask/Public/Mask2D/AvaMask2DBaseModifier.h
--- a/Engine/Plugins/VirtualProduction/Avalanche/Source/AvalancheMask/Public/Mask2D/AvaMask2DBaseModifier.h
+++ b/Engine/Plugins/VirtualProduction/Avalanche/Source/AvalancheMask/Public/Mask2D/AvaMask2DBaseModifier.h
@@ -237,6 +237,9 @@ protected:
 
 	TMap<TObjectKey<UObject>, TSharedPtr<IAvaMaskMaterialCollectionHandle>> MaterialCollectionHandles;
 
+	/** Last scanned primitive components */
+	TSet<TObjectKey<UPrimitiveComponent>> ScannedPrimitiveComponents;
+
 private:
 #if WITH_EDITOR
 	/** Used for PECP */
diff --git a/Engine/Plugins/VirtualProduction/Avalanche/Source/AvalancheMask/Public/Mask2D/AvaMask2DWriteModifier.h b/Engine/Plugins/VirtualProduction/Avalanche/Source/AvalancheMask/Public/Mask2D/AvaMask2DWriteModifier.h
--- a/Engine/Plugins/VirtualProduction/Avalanche/Source/AvalancheMask/Public/Mask2D/AvaMask2DWriteModifier.h
+++ b/Engine/Plugins/VirtualProduction/Avalanche/Source/AvalancheMask/Public/Mask2D/AvaMask2DWriteModifier.h
@@ -40,29 +40,32 @@ class AVALANCHEMASK_API UAvaMask2DWriteModifier
 	friend class UAvaMask2DModifierShared;
 	
 public:
+	EGeometryMaskCompositeOperation GetWriteOperation() const
+	{
+		return WriteOperation;
+	}
+	void SetWriteOperation(const EGeometryMaskCompositeOperation InWriteOperation);
+
+protected:
 	//~ Begin UObject
 #if WITH_EDITOR
 	virtual void PostEditChangeProperty(FPropertyChangedEvent& InPropertyChangedEvent) override;
 #endif
 	//~ End UObject
 
-	EGeometryMaskCompositeOperation GetWriteOperation() const { return WriteOperation; }
-	void SetWriteOperation(const EGeometryMaskCompositeOperation InWriteOperation);
-
-protected:
 	//~ Begin UActorModifierCoreBase
 	virtual void OnModifierCDOSetup(FActorModifierCoreMetadata& InMetadata) override;
 	virtual void Apply() override;
 	//~ End UActorModifierCoreBase
 
-	void OnWriteOperationChanged();
+	//~ Begin UAvaMask2DBaseModifier
+	virtual void SetupMaskComponent(UActorComponent* InComponent) override;
+	//~ End UAvaMask2DBaseModifier
 
-	void SetupMaskComponent(UActorComponent* InComponent);
+	void OnWriteOperationChanged();
 	void SetupMaskWriteComponent(IGeometryMaskWriteInterface* InMaskWriter);
-
 	bool ApplyWrite(AActor* InActor, FAvaMask2DActorData& InActorData);
 
-protected:
 	/** How to write to the chosen mask channel */
 	UPROPERTY(EditInstanceOnly, BlueprintReadWrite, Getter, Setter, Category = "Mask2D", meta = (ValidEnumValues = "Add, Subtract", AllowPrivateAccess = "true"))
 	EGeometryMaskCompositeOperation WriteOperation = EGeometryMaskCompositeOperation::Add;
