MotionDesign : ActorModifier
- Fix nested visibility modifier not respecting hierarchy or hidden status when property is changed

diff --git a/Engine/Plugins/VirtualProduction/Avalanche/Source/AvalancheModifiers/Private/Modifiers/AvaVisibilityModifier.cpp b/Engine/Plugins/VirtualProduction/Avalanche/Source/AvalancheModifiers/Private/Modifiers/AvaVisibilityModifier.cpp
--- a/Engine/Plugins/VirtualProduction/Avalanche/Source/AvalancheModifiers/Private/Modifiers/AvaVisibilityModifier.cpp
+++ b/Engine/Plugins/VirtualProduction/Avalanche/Source/AvalancheModifiers/Private/Modifiers/AvaVisibilityModifier.cpp
@@ -37,12 +37,17 @@ void UAvaVisibilityModifier::Apply()
 
 	// Top most modifier in tree has priority over this one if it is hiding current one
 	bool bIsNestedVisibilityModifier = false;
-	if (const UAvaVisibilityModifier* VisibilityModifier = GetLastModifierAbove(ModifyActor))
+	const TArray<UAvaVisibilityModifier*> VisibilityModifiers = GetModifiersAbove(ModifyActor);
+	if (!VisibilityModifiers.IsEmpty())
 	{
-		// This actor is not visible and hidden by root visibility modifier do not proceed
-		if (VisibilityModifier->IsChildActorHidden(ModifyActor))
+		// This actor is not visible and hidden by another visibility modifier do not proceed
+		for (const UAvaVisibilityModifier* VisibilityModifier : VisibilityModifiers)
 		{
-			bIsNestedVisibilityModifier = true;
+			if (VisibilityModifier && VisibilityModifier->IsChildActorHidden(ModifyActor))
+			{
+				bIsNestedVisibilityModifier = true;
+				break;
+			}
 		}
 	}
 	// We are the top root modifier, if this actor is hidden, do not handle children actors
@@ -208,6 +213,11 @@ void UAvaVisibilityModifier::OnActorVisibilityChanged(AActor* InActor)
 		return;
 	}
 
+	if (Modifier && Modifier->IsChildActorHidden(InActor))
+	{
+		return;
+	}
+
 	MarkModifierDirty();
 }
 
@@ -296,17 +306,15 @@ UAvaVisibilityModifier* UAvaVisibilityModifier::GetFirstModifierAbove(AActor* In
 	return FirstModifierAbove;
 }
 
-UAvaVisibilityModifier* UAvaVisibilityModifier::GetLastModifierAbove(AActor* InActor)
+TArray<UAvaVisibilityModifier*> UAvaVisibilityModifier::GetModifiersAbove(AActor* InActor)
 {
-	UAvaVisibilityModifier* LastModifierAbove = nullptr;
-
+	TArray<UAvaVisibilityModifier*> Modifiers;
 	while (UAvaVisibilityModifier* VisibilityModifier = GetFirstModifierAbove(InActor))
 	{
-		LastModifierAbove = VisibilityModifier;
-		InActor = LastModifierAbove->GetModifiedActor();
+		Modifiers.Add(VisibilityModifier);
+		InActor = VisibilityModifier->GetModifiedActor();
 	}
-
-	return LastModifierAbove;
+	return Modifiers;
 }
 
 AActor* UAvaVisibilityModifier::GetDirectChildren(AActor* InParentActor, AActor* InChildActor)
diff --git a/Engine/Plugins/VirtualProduction/Avalanche/Source/AvalancheModifiers/Public/Modifiers/AvaVisibilityModifier.h b/Engine/Plugins/VirtualProduction/Avalanche/Source/AvalancheModifiers/Public/Modifiers/AvaVisibilityModifier.h
--- a/Engine/Plugins/VirtualProduction/Avalanche/Source/AvalancheModifiers/Public/Modifiers/AvaVisibilityModifier.h
+++ b/Engine/Plugins/VirtualProduction/Avalanche/Source/AvalancheModifiers/Public/Modifiers/AvaVisibilityModifier.h
@@ -76,7 +76,7 @@ protected:
 
 	/** Get first visibility modifier found above this actor */
 	UAvaVisibilityModifier* GetFirstModifierAbove(AActor* InActor);
-	UAvaVisibilityModifier* GetLastModifierAbove(AActor* InActor);
+	TArray<UAvaVisibilityModifier*> GetModifiersAbove(AActor* InActor);
 
 	/** Will find the direct children of InParentActor that contains this child actor */
 	AActor* GetDirectChildren(AActor* InParentActor, AActor* InChildActor);
