Motion Design: added error message when renaming broadcast channel + added extra check to prevent program channels renaming to a reserved name for default preview

diff --git a/Engine/Plugins/VirtualProduction/Avalanche/Source/AvalancheMedia/Private/Broadcast/AvaBroadcast.cpp b/Engine/Plugins/VirtualProduction/Avalanche/Source/AvalancheMedia/Private/Broadcast/AvaBroadcast.cpp
--- a/Engine/Plugins/VirtualProduction/Avalanche/Source/AvalancheMedia/Private/Broadcast/AvaBroadcast.cpp
+++ b/Engine/Plugins/VirtualProduction/Avalanche/Source/AvalancheMedia/Private/Broadcast/AvaBroadcast.cpp
@@ -8,6 +8,7 @@
 #include "HAL/FileManager.h"
 #include "IAvaMediaModule.h"
 #include "Misc/Paths.h"
+#include "Rundown/AvaRundown.h"
 #include "UObject/NameTypes.h"
 #include "UObject/Package.h"
 
@@ -586,16 +587,59 @@ void UAvaBroadcast::UpdateChannelNames()
 
 bool UAvaBroadcast::CanRenameChannel(FName InChannelName, FName InNewChannelName) const
 {
-	//Make sure new Channel name is valid, and we don't have it already in the List
-	return InNewChannelName != NAME_None
-		&& InChannelName != InNewChannelName
-		&& true == ChannelNames.Contains(InChannelName)
-		&& false == ChannelNames.Contains(InNewChannelName);
+	FText ErrorMessage;
+	return CanRenameChannel(InChannelName, InNewChannelName, ErrorMessage);
+}
+
+bool UAvaBroadcast::CanRenameChannel(FName InChannelName, FName InNewChannelName, FText& OutErrorMessage) const
+{
+	if (InNewChannelName.IsNone())
+	{
+		OutErrorMessage = LOCTEXT("RenameChannelError_None", "New name is not valid");
+		return false;
+	}
+
+	if (InChannelName == InNewChannelName)
+	{
+		OutErrorMessage = LOCTEXT("RenameChannelError_Unchanged", "The channel name is unchanged");
+		return false;
+	}
+
+	if (!ChannelNames.Contains(InChannelName))
+	{
+		OutErrorMessage = FText::Format(LOCTEXT("RenameChannelError_NoChannel", "The channel to rename '{0}' was not found."), FText::FromName(InChannelName));
+		return false;
+	}
+
+	if (ChannelNames.Contains(InNewChannelName))
+	{
+		OutErrorMessage = FText::Format(LOCTEXT("RenameChannelError_ExistingChannel", "A channel named '{0}' already exists."), FText::FromName(InNewChannelName));
+		return false;
+	}
+
+	// If the new channel name is the default preview channel name,
+	// check that the channel is set to be preview type to avoid unexpectedly disabling preview playback for these.
+	if (InNewChannelName == UAvaRundown::GetDefaultPreviewChannelName())
+	{
+		if (GetChannelType(InChannelName) != EAvaBroadcastChannelType::Preview)
+		{
+			OutErrorMessage = FText::Format(LOCTEXT("RenameChannelError_DefaultPreview", "'{0}' is a reserved name for default preview channels, but the channel type is not preview"), FText::FromName(InNewChannelName));
+			return false;
+		}
+	}
+
+	return true;
 }
 
 bool UAvaBroadcast::RenameChannel(FName InChannelName, FName InNewChannelName)
 {
-	if (CanRenameChannel(InChannelName, InNewChannelName))
+	FText ErrorMessage;
+	return RenameChannel(InChannelName, InNewChannelName, ErrorMessage);
+}
+
+bool UAvaBroadcast::RenameChannel(FName InChannelName, FName InNewChannelName, FText& OutErrorMessage)
+{
+	if (CanRenameChannel(InChannelName, InNewChannelName, OutErrorMessage))
 	{
 		const int32 Index = GetChannelIndex(InChannelName);
 		check(ChannelNames.IsValidIndex(Index));
diff --git a/Engine/Plugins/VirtualProduction/Avalanche/Source/AvalancheMedia/Private/Rundown/AvaRundownServer.cpp b/Engine/Plugins/VirtualProduction/Avalanche/Source/AvalancheMedia/Private/Rundown/AvaRundownServer.cpp
--- a/Engine/Plugins/VirtualProduction/Avalanche/Source/AvalancheMedia/Private/Rundown/AvaRundownServer.cpp
+++ b/Engine/Plugins/VirtualProduction/Avalanche/Source/AvalancheMedia/Private/Rundown/AvaRundownServer.cpp
@@ -2128,10 +2128,11 @@ void FAvaRundownServer::HandleRenameChannel(const FAvaRundownRenameChannel& InMe
 		return;
 	}
 
-	if (!Broadcast.RenameChannel(OldChannelName, NewChannelName))
+	FText ErrorMessage;
+	if (!Broadcast.RenameChannel(OldChannelName, NewChannelName, ErrorMessage))
 	{
 		LogAndSendMessage(InContext->GetSender(), InMessage.RequestId, ELogVerbosity::Error,
-			TEXT("\"RenameChannel\" Failed to rename channel \"%s\" to \"%s\" (Unknown reason)."), *OldChannelName.ToString(), *NewChannelName.ToString());
+			TEXT("\"RenameChannel\" Failed to rename channel \"%s\" to \"%s\". Reason: %s"), *OldChannelName.ToString(), *NewChannelName.ToString(), *ErrorMessage.ToString());
 		return;
 	}
 
diff --git a/Engine/Plugins/VirtualProduction/Avalanche/Source/AvalancheMedia/Public/Broadcast/AvaBroadcast.h b/Engine/Plugins/VirtualProduction/Avalanche/Source/AvalancheMedia/Public/Broadcast/AvaBroadcast.h
--- a/Engine/Plugins/VirtualProduction/Avalanche/Source/AvalancheMedia/Public/Broadcast/AvaBroadcast.h
+++ b/Engine/Plugins/VirtualProduction/Avalanche/Source/AvalancheMedia/Public/Broadcast/AvaBroadcast.h
@@ -114,7 +114,10 @@ public:
 	void UpdateChannelNames();
 	
 	bool CanRenameChannel(FName InChannelName, FName InNewChannelName) const;
+	bool CanRenameChannel(FName InChannelName, FName InNewChannelName, FText& OutErrorMessage) const;
+
 	bool RenameChannel(FName InChannelName, FName InNewChannelName);
+	bool RenameChannel(FName InChannelName, FName InNewChannelName, FText& OutErrorMessage);
 
 	/**
 	 * Sets the type (ex: Program or Preview) of the given channel across all profiles.
diff --git a/Engine/Plugins/VirtualProduction/Avalanche/Source/AvalancheMediaEditor/Private/Broadcast/ChannelGrid/Slate/SAvaBroadcastChannel.cpp b/Engine/Plugins/VirtualProduction/Avalanche/Source/AvalancheMediaEditor/Private/Broadcast/ChannelGrid/Slate/SAvaBroadcastChannel.cpp
--- a/Engine/Plugins/VirtualProduction/Avalanche/Source/AvalancheMediaEditor/Private/Broadcast/ChannelGrid/Slate/SAvaBroadcastChannel.cpp
+++ b/Engine/Plugins/VirtualProduction/Avalanche/Source/AvalancheMediaEditor/Private/Broadcast/ChannelGrid/Slate/SAvaBroadcastChannel.cpp
@@ -602,7 +602,7 @@ FText SAvaBroadcastChannel::GetChannelNameText() const
 
 bool SAvaBroadcastChannel::OnVerifyChannelNameTextChanged(const FText& InText, FText& OutErrorMessage)
 {
-	return UAvaBroadcast::Get().CanRenameChannel(ChannelName, FName(InText.ToString()));
+	return UAvaBroadcast::Get().CanRenameChannel(ChannelName, FName(InText.ToString()), OutErrorMessage);
 }
 
 void SAvaBroadcastChannel::OnChannelNameTextCommitted(const FText& InText, ETextCommit::Type CommitInfo)
