MotionDesign : Modifiers
- Fixed modifier by default now use class name instead of none and not registering at all
- Fixed taper modifier removing/deprecating resolution property since it is not useful and not needed

diff --git a/Engine/Plugins/VirtualProduction/ActorModifierCore/Source/ActorModifierCore/Private/Modifiers/ActorModifierCoreDefs.cpp b/Engine/Plugins/VirtualProduction/ActorModifierCore/Source/ActorModifierCore/Private/Modifiers/ActorModifierCoreDefs.cpp
--- a/Engine/Plugins/VirtualProduction/ActorModifierCore/Source/ActorModifierCore/Private/Modifiers/ActorModifierCoreDefs.cpp
+++ b/Engine/Plugins/VirtualProduction/ActorModifierCore/Source/ActorModifierCore/Private/Modifiers/ActorModifierCoreDefs.cpp
@@ -28,8 +28,9 @@ FActorModifierCoreMetadata::FActorModifierCoreMetadata(const UActorModifierCoreB
 
 	Class = InModifier->GetClass();
 	bIsStack = InModifier->IsA<UActorModifierCoreStack>();
-	Category = DefaultCategory;
 
+	SetCategory(DefaultCategory);
+	SetName(Class->GetFName());
 	SetProfilerClass<FActorModifierCoreProfiler>();
 	SetCompatibilityRule([](const AActor* InActor)->bool
 	{
diff --git a/Engine/Plugins/VirtualProduction/ActorModifierCore/Source/ActorModifierCoreEditor/Private/ActorModifierCoreEditorModule.cpp b/Engine/Plugins/VirtualProduction/ActorModifierCore/Source/ActorModifierCoreEditor/Private/ActorModifierCoreEditorModule.cpp
--- a/Engine/Plugins/VirtualProduction/ActorModifierCore/Source/ActorModifierCoreEditor/Private/ActorModifierCoreEditorModule.cpp
+++ b/Engine/Plugins/VirtualProduction/ActorModifierCore/Source/ActorModifierCoreEditor/Private/ActorModifierCoreEditorModule.cpp
@@ -21,24 +21,26 @@ void FActorModifierCoreEditorModule::ShutdownModule()
 	UnregisterDetailCustomizations();
 }
 
-void FActorModifierCoreEditorModule::RegisterDetailCustomizations() const
+void FActorModifierCoreEditorModule::RegisterDetailCustomizations()
 {
 	// Register custom layouts
-	static FName PropertyEditor("PropertyEditor");
-	FPropertyEditorModule& PropertyModule = FModuleManager::GetModuleChecked<FPropertyEditorModule>(PropertyEditor);
+	FPropertyEditorModule& PropertyModule = FModuleManager::GetModuleChecked<FPropertyEditorModule>("PropertyEditor");
 
-	PropertyModule.RegisterCustomClassLayout(UActorModifierCoreStack::StaticClass()->GetFName()
+	PropertyModule.RegisterCustomClassLayout(CustomizationNames.Add_GetRef(UActorModifierCoreStack::StaticClass()->GetFName())
 		, FOnGetDetailCustomizationInstance::CreateStatic(&FActorModifierCoreEditorDetailCustomization::MakeInstance));
 }
 
-void FActorModifierCoreEditorModule::UnregisterDetailCustomizations() const
+void FActorModifierCoreEditorModule::UnregisterDetailCustomizations()
 {
 	// Unregister custom layouts
-	static FName PropertyEditor("PropertyEditor");
-	if (FModuleManager::Get().IsModuleLoaded(PropertyEditor))
+	if (FPropertyEditorModule* PropertyModule = FModuleManager::GetModulePtr<FPropertyEditorModule>("PropertyEditor"))
 	{
-		FPropertyEditorModule& PropertyModule = FModuleManager::GetModuleChecked<FPropertyEditorModule>(PropertyEditor);
-		PropertyModule.UnregisterCustomClassLayout(UActorModifierCoreStack::StaticClass()->GetFName());
+		for (const FName CustomizationName : CustomizationNames)
+		{
+			PropertyModule->UnregisterCustomClassLayout(CustomizationName);
+		}
+
+		CustomizationNames.Empty();
 	}
 }
 
diff --git a/Engine/Plugins/VirtualProduction/ActorModifierCore/Source/ActorModifierCoreEditor/Private/ActorModifierCoreEditorModule.h b/Engine/Plugins/VirtualProduction/ActorModifierCore/Source/ActorModifierCoreEditor/Private/ActorModifierCoreEditorModule.h
--- a/Engine/Plugins/VirtualProduction/ActorModifierCore/Source/ActorModifierCoreEditor/Private/ActorModifierCoreEditorModule.h
+++ b/Engine/Plugins/VirtualProduction/ActorModifierCore/Source/ActorModifierCoreEditor/Private/ActorModifierCoreEditorModule.h
@@ -3,6 +3,8 @@
 #pragma once
 
 #include "Modules/ModuleInterface.h"
+#include "Containers/Array.h"
+#include "UObject/NameTypes.h"
 
 class FActorModifierCoreEditorModule : public IModuleInterface
 {
@@ -13,8 +15,10 @@ public:
 	//~ End IModuleInterface
 
 protected:
-	void RegisterDetailCustomizations() const;
-	void UnregisterDetailCustomizations() const;
+	void RegisterDetailCustomizations();
+	void UnregisterDetailCustomizations();
 
 	void RegisterBlueprintCustomizations();
+
+	TArray<FName> CustomizationNames;
 };
\ No newline at end of file
diff --git a/Engine/Plugins/VirtualProduction/Avalanche/Source/AvalancheModifiers/Private/Modifiers/AvaTaperModifier.cpp b/Engine/Plugins/VirtualProduction/Avalanche/Source/AvalancheModifiers/Private/Modifiers/AvaTaperModifier.cpp
--- a/Engine/Plugins/VirtualProduction/Avalanche/Source/AvalancheModifiers/Private/Modifiers/AvaTaperModifier.cpp
+++ b/Engine/Plugins/VirtualProduction/Avalanche/Source/AvalancheModifiers/Private/Modifiers/AvaTaperModifier.cpp
@@ -78,14 +78,12 @@ void UAvaTaperModifier::Apply()
 	// retrieve the dynamic mesh to be modified
 	if (FDynamicMesh3* const SharedMesh = DynMeshComp->GetMesh())
 	{
-		const int32 TaperControlResolution = FMath::Max(GetResolution(), GetSubdividersCuts());
-		
 		FAvaTaperSettings TaperSettings;
 		TaperSettings.Amount = Amount;
 		TaperSettings.Extent = GetRequiredExtent();
 		TaperSettings.Offset = GetRequiredOffset();
 		TaperSettings.InterpolationType = InterpolationType;
-		TaperSettings.ZAxisResolution = TaperControlResolution;
+		TaperSettings.ZAxisResolution = GetSubdividersCuts();
 
 		// setup the taper with by passing the mesh and the current settings
 		if (TaperTool->Setup(SharedMesh, TaperSettings))
@@ -124,7 +122,6 @@ void UAvaTaperModifier::PostEditChangeProperty(FPropertyChangedEvent& PropertyCh
 	const FName MemberName = PropertyChangedEvent.GetMemberPropertyName();
 	
 	static const FName AmountName = GET_MEMBER_NAME_CHECKED(UAvaTaperModifier, Amount);
-	static const FName TessellationName = GET_MEMBER_NAME_CHECKED(UAvaTaperModifier, Resolution);
 	static const FName InterpolationTypeName = GET_MEMBER_NAME_CHECKED(UAvaTaperModifier, InterpolationType);
 	static const FName OffsetName = GET_MEMBER_NAME_CHECKED(UAvaTaperModifier, Offset);
 	static const FName ReferenceFrameName = GET_MEMBER_NAME_CHECKED(UAvaTaperModifier, ReferenceFrame);
@@ -133,7 +130,6 @@ void UAvaTaperModifier::PostEditChangeProperty(FPropertyChangedEvent& PropertyCh
 	static const FName ExtentName = GET_MEMBER_NAME_CHECKED(UAvaTaperModifier, Extent);
 
 	if (MemberName == AmountName ||
-		MemberName == TessellationName ||
 		MemberName == InterpolationTypeName ||
 		MemberName == OffsetName ||
 		MemberName == ReferenceFrameName ||
@@ -214,6 +210,8 @@ void UAvaTaperModifier::SetReferenceFrame(EAvaTaperReferenceFrame InReferenceFra
 
 void UAvaTaperModifier::SetResolution(int32 InResolution)
 {
+	PRAGMA_DISABLE_DEPRECATION_WARNINGS
+
 	if (Resolution == InResolution)
 	{
 		return;
@@ -221,6 +219,8 @@ void UAvaTaperModifier::SetResolution(int32 InResolution)
 
 	Resolution = FMath::Clamp(InResolution, UAvaTaperModifier::MinTaperLatticeResolution, UAvaTaperModifier::MaxTaperLatticeResolution);
 	OnParameterChanged();
+
+	PRAGMA_ENABLE_DEPRECATION_WARNINGS
 }
 
 void UAvaTaperModifier::SetOffset(FVector2D InOffset)
@@ -271,14 +271,20 @@ int32 UAvaTaperModifier::GetSubdividersCuts() const
 {
 	int32 TotalSubdivisions = 0;
 	
-	if (GetModifierStack() && GetModifierStack()->ContainsModifier(UAvaSubdivideModifier::StaticClass()))
+	if (const UActorModifierCoreStack* Stack = GetModifierStack())
 	{
-		TArray<UAvaSubdivideModifier*> SubdivideModifiers;
-		GetModifierStack()->GetClassModifiers<UAvaSubdivideModifier>(SubdivideModifiers);
+		FActorModifierCoreStackSearchOp SearchOp;
+		SearchOp.bSkipStack = true;
+		SearchOp.Position = EActorModifierCoreStackPosition::Before;
+		SearchOp.PositionContext = const_cast<UAvaTaperModifier*>(this);
+		const TArray<UActorModifierCoreBase*> Modifiers = Stack->FindModifiers(UAvaSubdivideModifier::StaticClass(), SearchOp);
 
-		for (const UAvaSubdivideModifier* const SubdivideModifier : SubdivideModifiers)
+		for (const UActorModifierCoreBase* const Modifier : Modifiers)
 		{
-			TotalSubdivisions += SubdivideModifier->GetCuts();
+			if (const UAvaSubdivideModifier* SubdivideModifier = Cast<UAvaSubdivideModifier>(Modifier))
+			{
+				TotalSubdivisions += SubdivideModifier->GetCuts();
+			}
 		}
 	}
 
diff --git a/Engine/Plugins/VirtualProduction/Avalanche/Source/AvalancheModifiers/Public/Modifiers/AvaExtrudeModifier.h b/Engine/Plugins/VirtualProduction/Avalanche/Source/AvalancheModifiers/Public/Modifiers/AvaExtrudeModifier.h
--- a/Engine/Plugins/VirtualProduction/Avalanche/Source/AvalancheModifiers/Public/Modifiers/AvaExtrudeModifier.h
+++ b/Engine/Plugins/VirtualProduction/Avalanche/Source/AvalancheModifiers/Public/Modifiers/AvaExtrudeModifier.h
@@ -69,11 +69,11 @@ protected:
 	FVector GetExtrudeDirection() const;
 
 	/** Handles mesh depth to extrude primary section */
-	UPROPERTY(EditInstanceOnly, Setter="SetDepth", Getter="GetDepth", Category="Extrude", meta=(ClampMin="0", DisplayName="Extrude depth", Units="Centimeters", AllowPrivateAccess="true"))
+	UPROPERTY(EditInstanceOnly, Setter="SetDepth", Getter="GetDepth", Category="Extrude", meta=(ClampMin="0", DisplayName="Extrude Depth", Units="Centimeters", AllowPrivateAccess="true"))
 	float Depth = 30.f;
 
 	/** Closes the back of the extrude for a 2D shape for example */
-	UPROPERTY(EditInstanceOnly, Setter="SetCloseBack", Getter="GetCloseBack", Category="Extrude", meta=(DisplayName="Close back", AllowPrivateAccess="true"))
+	UPROPERTY(EditInstanceOnly, Setter="SetCloseBack", Getter="GetCloseBack", Category="Extrude", meta=(AllowPrivateAccess="true"))
 	bool bCloseBack = true;
 
 	/** Moves the mesh in the opposite extrude direction by the depth distance */
@@ -81,6 +81,6 @@ protected:
 	bool bMoveMeshOppositeDirection_DEPRECATED = true;
 
 	/** Specifies the Extrude direction */
-	UPROPERTY(EditInstanceOnly, Setter="SetExtrudeMode", Getter="GetExtrudeMode", Category="Extrude", meta=(DisplayName = "Extrude Mode", AllowPrivateAccess="true"))
+	UPROPERTY(EditInstanceOnly, Setter="SetExtrudeMode", Getter="GetExtrudeMode", Category="Extrude", meta=(AllowPrivateAccess="true"))
 	EAvaExtrudeMode ExtrudeMode = EAvaExtrudeMode::Opposite;
 };
diff --git a/Engine/Plugins/VirtualProduction/Avalanche/Source/AvalancheModifiers/Public/Modifiers/AvaTaperModifier.h b/Engine/Plugins/VirtualProduction/Avalanche/Source/AvalancheModifiers/Public/Modifiers/AvaTaperModifier.h
--- a/Engine/Plugins/VirtualProduction/Avalanche/Source/AvalancheModifiers/Public/Modifiers/AvaTaperModifier.h
+++ b/Engine/Plugins/VirtualProduction/Avalanche/Source/AvalancheModifiers/Public/Modifiers/AvaTaperModifier.h
@@ -83,13 +83,17 @@ public:
 		return ReferenceFrame;
 	}
 
-	UFUNCTION(BlueprintCallable, Category="Motion Design|Modifiers|Taper")
+	UE_DEPRECATED(5.7, "Use subdivide modifier before this modifier instead")
+	UFUNCTION(BlueprintCallable, Category="Motion Design|Modifiers|Taper", meta=(DeprecatedFunction, DeprecatedMessage="Use subdivide modifier before this modifier instead"))
 	AVALANCHEMODIFIERS_API void SetResolution(int32 InResolution);
 
-	UFUNCTION(BlueprintPure, Category="Motion Design|Modifiers|Taper")
+	UE_DEPRECATED(5.7, "Use subdivide modifier before this modifier instead")
+	UFUNCTION(BlueprintPure, Category="Motion Design|Modifiers|Taper", meta=(DeprecatedFunction, DeprecatedMessage="Use subdivide modifier before this modifier instead"))
 	int32 GetResolution() const
 	{
+		PRAGMA_DISABLE_DEPRECATION_WARNINGS
 		return Resolution;
+		PRAGMA_ENABLE_DEPRECATION_WARNINGS
 	}
 
 	UFUNCTION(BlueprintCallable, Category="Motion Design|Modifiers|Taper")
@@ -137,7 +141,8 @@ protected:
 	UPROPERTY(EditInstanceOnly, Setter="SetInterpolationType", Getter="GetInterpolationType", Category="Taper", meta=(AllowPrivateAccess="true"))
 	EAvaTaperInterpolationType InterpolationType = EAvaTaperInterpolationType::Linear;
 
-	UPROPERTY(EditInstanceOnly, Setter="SetResolution", Getter="GetResolution", Category="Taper", meta=(ClampMin="1", ClampMax="20", ToolTip="The number of vertical control points used to apply the taper. If the modifier is in a stack with Subdivide modifiers, taper will use the max value between Resolution and the total subdivision cuts.", AllowPrivateAccess="true"))
+	UE_DEPRECATED(5.7, "Use subdivide modifier before this modifier instead")
+	UPROPERTY(Transient, meta=(DeprecatedProperty, DeprecatedMessage="Use subdivide modifier before this modifier instead"))
 	int32 Resolution = 5;
 
 	UPROPERTY(EditInstanceOnly, Setter="SetReferenceFrame", Getter="GetReferenceFrame", Category="Taper", meta=(AllowPrivateAccess="true"))
