MotionDesign : Shapes
- Fixed InnerSize range adjusted to allow 0 and 1 values

diff --git a/Engine/Plugins/VirtualProduction/Avalanche/Source/AvalancheShapes/Private/DynamicMeshes/AvaShapeRingDynMesh.cpp b/Engine/Plugins/VirtualProduction/Avalanche/Source/AvalancheShapes/Private/DynamicMeshes/AvaShapeRingDynMesh.cpp
--- a/Engine/Plugins/VirtualProduction/Avalanche/Source/AvalancheShapes/Private/DynamicMeshes/AvaShapeRingDynMesh.cpp
+++ b/Engine/Plugins/VirtualProduction/Avalanche/Source/AvalancheShapes/Private/DynamicMeshes/AvaShapeRingDynMesh.cpp
@@ -6,33 +6,20 @@ const FString UAvaShapeRingDynamicMesh::MeshName = TEXT("Ring");
 
 void UAvaShapeRingDynamicMesh::SetNumSides(uint8 InNumSides)
 {
+	InNumSides = FMath::Clamp(InNumSides, UAvaShapeRingDynamicMesh::MinNumSides, UAvaShapeRingDynamicMesh::MaxNumSides);
 	if (NumSides == InNumSides)
 	{
 		return;
 	}
 
-	if (InNumSides < UAvaShapeRingDynamicMesh::MinNumSides || InNumSides > UAvaShapeRingDynamicMesh::MaxNumSides)
-	{
-		return;
-	}
-
 	NumSides = InNumSides;
 	OnNumSidesChanged();
 }
 
 void UAvaShapeRingDynamicMesh::SetInnerSize(float InInnerSize)
 {
-	if (InnerSize == InInnerSize)
-	{
-		return;
-	}
-
-	if (InInnerSize < 0.01f)
-	{
-		return;
-	}
-
-	if (InInnerSize > 0.99f)
+	InInnerSize = FMath::Clamp(InInnerSize, 0.f, 1.f);
+	if (FMath::IsNearlyEqual(InnerSize, InInnerSize))
 	{
 		return;
 	}
@@ -43,12 +30,8 @@ void UAvaShapeRingDynamicMesh::SetInnerSize(float InInnerSize)
 
 void UAvaShapeRingDynamicMesh::SetAngleDegree(float InDegree)
 {
-	if (AngleDegree == InDegree)
-	{
-		return;
-	}
-
-	if (InDegree < 0.f || InDegree > 360.f)
+	InDegree = FMath::Clamp(InDegree, 0.f, 360.f);
+	if (FMath::IsNearlyEqual(AngleDegree, InDegree))
 	{
 		return;
 	}
@@ -59,12 +42,8 @@ void UAvaShapeRingDynamicMesh::SetAngleDegree(float InDegree)
 
 void UAvaShapeRingDynamicMesh::SetStartDegree(float InDegree)
 {
-	if (StartDegree == InDegree)
-	{
-		return;
-	}
-
-	if (InDegree < 0.f || InDegree > 360.f)
+	InDegree = FMath::Clamp(InDegree, 0.f, 360.f);
+	if (FMath::IsNearlyEqual(StartDegree, InDegree))
 	{
 		return;
 	}
@@ -128,7 +107,7 @@ bool UAvaShapeRingDynamicMesh::IsMeshVisible(int32 MeshIndex)
 	switch (MeshIndex)
 	{
 		case MESH_INDEX_PRIMARY:
-			return AngleDegree > UE_KINDA_SMALL_NUMBER;
+			return AngleDegree > UE_KINDA_SMALL_NUMBER && InnerSize < (1.f - UE_KINDA_SMALL_NUMBER);
 		default:
 		return true;
 	}
diff --git a/Engine/Plugins/VirtualProduction/Avalanche/Source/AvalancheShapes/Public/DynamicMeshes/AvaShapeRingDynMesh.h b/Engine/Plugins/VirtualProduction/Avalanche/Source/AvalancheShapes/Public/DynamicMeshes/AvaShapeRingDynMesh.h
--- a/Engine/Plugins/VirtualProduction/Avalanche/Source/AvalancheShapes/Public/DynamicMeshes/AvaShapeRingDynMesh.h
+++ b/Engine/Plugins/VirtualProduction/Avalanche/Source/AvalancheShapes/Public/DynamicMeshes/AvaShapeRingDynMesh.h
@@ -79,7 +79,7 @@ protected:
 	UPROPERTY(EditInstanceOnly, BlueprintReadWrite, Setter, Getter, Category="Shape", meta=(ClampMin="3.0", ClampMax="128.0", DisplayName="Sides", AllowPrivateAccess="true"))
 	uint8 NumSides;
 
-	UPROPERTY(EditInstanceOnly, BlueprintReadWrite, Setter, Getter, Category="Shape", meta=(ClampMin="0.01", ClampMax="0.99",DisplayName="Inner Size", AllowPrivateAccess="true"))
+	UPROPERTY(EditInstanceOnly, BlueprintReadWrite, Setter, Getter, Category="Shape", meta=(ClampMin="0.0", ClampMax="1.0",DisplayName="Inner Size", AllowPrivateAccess="true"))
 	float InnerSize;
 
 	// represents the angle in degree for the ring
