MotionDesign : Shapes
- Fixed visualizer handles not showing when hovering actor

diff --git a/Engine/Plugins/VirtualProduction/Avalanche/Source/AvalancheComponentVisualizers/Private/AvaVisBase.cpp b/Engine/Plugins/VirtualProduction/Avalanche/Source/AvalancheComponentVisualizers/Private/AvaVisBase.cpp
--- a/Engine/Plugins/VirtualProduction/Avalanche/Source/AvalancheComponentVisualizers/Private/AvaVisBase.cpp
+++ b/Engine/Plugins/VirtualProduction/Avalanche/Source/AvalancheComponentVisualizers/Private/AvaVisBase.cpp
@@ -470,7 +470,10 @@ bool FAvaVisualizerBase::IsMouseOverComponent(const UActorComponent* Component,
 		return false;
 	}
 
-	const FOrientedBox ActorOrientedBox = FAvaActorUtils::MakeOrientedBox(ActorSizeBox, GetComponentTransform(Component));
+	// Cancel scale out since it's not taken into account in ActorSizeBox
+	FTransform ComponentTransform = GetComponentTransform(Component);
+	ComponentTransform.SetScale3D(FVector::OneVector);
+	const FOrientedBox ActorOrientedBox = FAvaActorUtils::MakeOrientedBox(ActorSizeBox, ComponentTransform);
 
 	// Calculate the rotation of the oriented box and revert it to axis aligned, moving the camera position and look direction along with it
 	// so that we can use the line/box intersection method.
