MotionDesign : Shapes
- Fixed switching material type property when primary material everywhere is toggled will apply to all mesh sections
- Added custom trigger function to update parametric material parameters on the shape actor to be called from RC or blueprint

diff --git a/Engine/Plugins/VirtualProduction/Avalanche/Source/AvalancheShapes/Private/AvaShapeActor.cpp b/Engine/Plugins/VirtualProduction/Avalanche/Source/AvalancheShapes/Private/AvaShapeActor.cpp
--- a/Engine/Plugins/VirtualProduction/Avalanche/Source/AvalancheShapes/Private/AvaShapeActor.cpp
+++ b/Engine/Plugins/VirtualProduction/Avalanche/Source/AvalancheShapes/Private/AvaShapeActor.cpp
@@ -98,3 +98,11 @@ UObject* AAvaShapeActor::GetModeDetailsObject_Implementation() const
 {
 	return DynamicMesh.Get();
 }
+
+void AAvaShapeActor::RefreshParametricMaterial()
+{
+	if (DynamicMesh)
+	{
+		DynamicMesh->RefreshParametricMaterial();
+	}
+}
diff --git a/Engine/Plugins/VirtualProduction/Avalanche/Source/AvalancheShapes/Private/DynamicMeshes/AvaShapeDynMeshBase.cpp b/Engine/Plugins/VirtualProduction/Avalanche/Source/AvalancheShapes/Private/DynamicMeshes/AvaShapeDynMeshBase.cpp
--- a/Engine/Plugins/VirtualProduction/Avalanche/Source/AvalancheShapes/Private/DynamicMeshes/AvaShapeDynMeshBase.cpp
+++ b/Engine/Plugins/VirtualProduction/Avalanche/Source/AvalancheShapes/Private/DynamicMeshes/AvaShapeDynMeshBase.cpp
@@ -150,6 +150,24 @@ void UAvaShapeDynamicMeshBase::SetupMeshes()
 	}
 }
 
+void UAvaShapeDynamicMeshBase::RefreshParametricMaterial()
+{
+	if (bUsePrimaryMaterialEverywhere)
+	{
+		OnParametricMaterialChanged(MESH_INDEX_PRIMARY);
+	}
+	else
+	{
+		for (const int32 Index : GetMeshesIndexes())
+		{
+			if (IsMaterialType(Index, EMaterialType::Parametric))
+			{
+				OnParametricMaterialChanged(Index);
+			}
+		}
+	}
+}
+
 bool UAvaShapeDynamicMeshBase::SetAlignmentSize(AActor* InActor, const FVector& InSizeMultiplier)
 {
 	if (const AAvaShapeActor* ShapeActor = Cast<AAvaShapeActor>(InActor))
@@ -1213,7 +1231,18 @@ void UAvaShapeDynamicMeshBase::OnMaterialTypeChanged(int32 MaterialIndex)
 					CurrentMaterial = nullptr;
 				}
 
-				SetMaterialDirect(MaterialIndex, CurrentMaterial);
+				if (!bUsePrimaryMaterialEverywhere)
+				{
+					SetMaterialDirect(MaterialIndex, CurrentMaterial);
+				}
+				else
+				{
+					for (const int32 SectionIndex : GetMeshesIndexes())
+					{
+						SetMaterialDirect(SectionIndex, CurrentMaterial);
+					}
+				}
+
 				break;
 			}
 
diff --git a/Engine/Plugins/VirtualProduction/Avalanche/Source/AvalancheShapes/Public/AvaShapeActor.h b/Engine/Plugins/VirtualProduction/Avalanche/Source/AvalancheShapes/Public/AvaShapeActor.h
--- a/Engine/Plugins/VirtualProduction/Avalanche/Source/AvalancheShapes/Public/AvaShapeActor.h
+++ b/Engine/Plugins/VirtualProduction/Avalanche/Source/AvalancheShapes/Public/AvaShapeActor.h
@@ -45,6 +45,10 @@ public:
 	virtual UObject* GetModeDetailsObject_Implementation() const override;
 	//~ End IAvaInteractiveToolsModeDetailsObjectProvider
 
+	/** Used to trigger a refresh of the parametric material if needed through RC or blueprint */
+	UFUNCTION(BlueprintCallable, CallInEditor, Category="Material")
+	void RefreshParametricMaterial();
+
 protected:
 	static const FName ShapeComponentName;
 
diff --git a/Engine/Plugins/VirtualProduction/Avalanche/Source/AvalancheShapes/Public/DynamicMeshes/AvaShapeDynMeshBase.h b/Engine/Plugins/VirtualProduction/Avalanche/Source/AvalancheShapes/Public/DynamicMeshes/AvaShapeDynMeshBase.h
--- a/Engine/Plugins/VirtualProduction/Avalanche/Source/AvalancheShapes/Public/DynamicMeshes/AvaShapeDynMeshBase.h
+++ b/Engine/Plugins/VirtualProduction/Avalanche/Source/AvalancheShapes/Public/DynamicMeshes/AvaShapeDynMeshBase.h
@@ -233,6 +233,9 @@ public:
 		return bRunAsync;
 	}
 
+	/** Updates parametric material by reapplying parameters */
+	AVALANCHESHAPES_API void RefreshParametricMaterial();
+
 protected:
 	static bool SetAlignmentSize(AActor* InActor, const FVector& InSizeMultiplier);
 	static EAvaAnchors GetAnchorFromNumerics(const FVector2D& AnchorNumeric);
