Motion Design: You can now remove colours from the viewport color picker themes.

diff --git a/Engine/Plugins/VirtualProduction/Avalanche/Source/AvalancheLevelViewport/Private/Widgets/SAvaLevelViewportActorColorMenu.cpp b/Engine/Plugins/VirtualProduction/Avalanche/Source/AvalancheLevelViewport/Private/Widgets/SAvaLevelViewportActorColorMenu.cpp
--- a/Engine/Plugins/VirtualProduction/Avalanche/Source/AvalancheLevelViewport/Private/Widgets/SAvaLevelViewportActorColorMenu.cpp
+++ b/Engine/Plugins/VirtualProduction/Avalanche/Source/AvalancheLevelViewport/Private/Widgets/SAvaLevelViewportActorColorMenu.cpp
@@ -1141,7 +1141,7 @@ TSharedRef<SWidget> SAvaLevelViewportActorColorMenu::BuildRenameColorMenu(int32
 
 	FLinearColor CurrentColor = ThemeColors[InColorIdx].Color;
 
-	constexpr bool bInShouldCloseWindowAfterMenuSelection = true;
+	constexpr bool bInShouldCloseWindowAfterMenuSelection = false;
 
 	FMenuBuilder ColorMenuBuilder(bInShouldCloseWindowAfterMenuSelection, nullptr);
 	{
