Motion Design: Fixed issue with saved active camera not being set as view target in map change

diff --git a/Engine/Plugins/VirtualProduction/Avalanche/Source/AvalancheLevelViewport/Private/ViewportClient/AvaLevelViewportClient.cpp b/Engine/Plugins/VirtualProduction/Avalanche/Source/AvalancheLevelViewport/Private/ViewportClient/AvaLevelViewportClient.cpp
--- a/Engine/Plugins/VirtualProduction/Avalanche/Source/AvalancheLevelViewport/Private/ViewportClient/AvaLevelViewportClient.cpp
+++ b/Engine/Plugins/VirtualProduction/Avalanche/Source/AvalancheLevelViewport/Private/ViewportClient/AvaLevelViewportClient.cpp
@@ -827,9 +827,9 @@ void FAvaLevelViewportClient::OnMapChanged(UWorld* InWorld, EMapChangeType InCha
 			{
 				PostProcessManager->LoadPostProcessInfo();
 			}
+			LoadActiveCameraComponent();
 			break;
 		}
-
 		default:
 			// Do nothing
 			break;
