[Motion Design] Fix Playback Client/Server issues with transitions in forked channels (includes display cluster).

The main issue is the transmission of the transition events can arrive and be processed out of order. This is critical in the case of the playable instance events (stop or discard) arriving after the "finish" transition event.
The fix is targeted to that particular case (i.e. I am not trying to fix all message ordering issues, we can revisit later but time is limited to fix this).

The server playback transition will relay the instance events as they occur, but will also keep them in memory and send them again in the final "finish" event message to ensure that which ever message arrives and is processed first on the client, the instance events will be processed before the finish transition event.

Other opportunistic changes included in this CL:
- Reset server will also reset transitions. This is important in case transitions are stuck on the server side.
- Playback Client's caching of the remote playback status will clean itself up to avoid keeping a ever growing set of stale playback statuses.
- Additional useful verbose logs have been added to help trouble shoot issues, especially around the remote transitions.

diff --git a/Engine/Plugins/VirtualProduction/Avalanche/Source/AvalancheMedia/Private/Playable/Transition/AvaPlayableRemoteTransition.cpp b/Engine/Plugins/VirtualProduction/Avalanche/Source/AvalancheMedia/Private/Playable/Transition/AvaPlayableRemoteTransition.cpp
--- a/Engine/Plugins/VirtualProduction/Avalanche/Source/AvalancheMedia/Private/Playable/Transition/AvaPlayableRemoteTransition.cpp
+++ b/Engine/Plugins/VirtualProduction/Avalanche/Source/AvalancheMedia/Private/Playable/Transition/AvaPlayableRemoteTransition.cpp
@@ -181,13 +181,26 @@ void UAvaPlayableRemoteTransition::HandlePlaybackTransitionEvent(IAvaPlaybackCli
 			// We only propagate the stop event if all forked channels have stopped.
 			if (bIsStopPlayableEvent)
 			{
+				UE_LOG(LogAvaPlayable, Verbose,
+					TEXT("%s Remote Playable Transition \"%s\": Marking playable instance Id \"%s\" for stop on server \"%s\"."),
+					*GetBriefFrameInfo(), *GetInstanceName(), *InArgs.InstanceId.ToString(), *InArgs.ServerName);
+
 				MarkPlayableForStop(InArgs.ServerName, InArgs.InstanceId);
 			}
 
 			// Relay locally through playable event.
 			if (!bIsStopPlayableEvent || IsPlayableMarkedForStopOnAllServers(InArgs.InstanceId))
 			{
-				UAvaPlayable::OnTransitionEvent().Broadcast(Playable, this, InArgs.EventFlags);
+				// Prevent locally broadcasting the instance event more than once.
+				if (!HasEventBeenBroadcast(InArgs.InstanceId, InArgs.EventFlags))
+				{
+					UE_LOG(LogAvaPlayable, Verbose,
+						TEXT("%s Remote Playable Transition \"%s\": Playable instance Id \"%s\" was marked for stop on all servers. Broadcasting locally."),
+						*GetBriefFrameInfo(), *GetInstanceName(), *InArgs.InstanceId.ToString());
+
+					UAvaPlayable::OnTransitionEvent().Broadcast(Playable, this, InArgs.EventFlags);
+					MarkEventsBroadcast(InArgs.InstanceId, InArgs.EventFlags);
+				}
 			}
 		}
 		else
@@ -203,6 +216,10 @@ void UAvaPlayableRemoteTransition::HandlePlaybackTransitionEvent(IAvaPlaybackCli
 		// Only propagate "finish" event when all remote transitions are finished.
 		if ((bIsTransitionFinishedEvent && IsTransitionFinishedOnAllServers()) || bIsTransitionStartingEvent)
 		{
+			UE_LOG(LogAvaPlayable, Verbose,
+				TEXT("%s Remote Playable Transition \"%s\": Finish event received from on all servers. Broadcasting locally."),
+				*GetBriefFrameInfo(), *GetInstanceName());
+
 			UAvaPlayable::OnTransitionEvent().Broadcast(nullptr, this, InArgs.EventFlags);
 		}
 	}
diff --git a/Engine/Plugins/VirtualProduction/Avalanche/Source/AvalancheMedia/Private/Playable/Transition/AvaPlayableRemoteTransition.h b/Engine/Plugins/VirtualProduction/Avalanche/Source/AvalancheMedia/Private/Playable/Transition/AvaPlayableRemoteTransition.h
--- a/Engine/Plugins/VirtualProduction/Avalanche/Source/AvalancheMedia/Private/Playable/Transition/AvaPlayableRemoteTransition.h
+++ b/Engine/Plugins/VirtualProduction/Avalanche/Source/AvalancheMedia/Private/Playable/Transition/AvaPlayableRemoteTransition.h
@@ -58,6 +58,29 @@ protected:
 	bool IsPlayableMarkedForStopOnAllServers(const FGuid& InInstanceId) const;
 	
 	FString GetInstanceName() const;
+
+	/** Mark the given events as broadcast for the given instance. */
+	void MarkEventsBroadcast(const FGuid& InInstanceId, EAvaPlayableTransitionEventFlags InFlags)
+	{
+		if (EAvaPlayableTransitionEventFlags* Flags = InstanceBroadcastEvent.Find(InInstanceId))
+		{
+			*Flags |= InFlags;
+		}
+		else
+		{
+			InstanceBroadcastEvent.Add(InInstanceId, InFlags);
+		}
+	}
+
+	/** Returns true if the given events where marked as broadcast for the given instance. */
+	bool HasEventBeenBroadcast(const FGuid& InInstanceId, EAvaPlayableTransitionEventFlags InFlags) const
+	{
+		if (const EAvaPlayableTransitionEventFlags* Flags = InstanceBroadcastEvent.Find(InInstanceId))
+		{
+			return EnumHasAllFlags(*Flags, InFlags);
+		}
+		return false;
+	}
 	
 protected:
 	FName ChannelName;
@@ -75,4 +98,7 @@ protected:
 	};
 	
 	TMap<FString, FRemoteStatusInfo> RemoteStatusPerServer;
+
+	// Keep track of the events that where locally broadcast to avoid sending them more than once.
+	TMap<FGuid, EAvaPlayableTransitionEventFlags> InstanceBroadcastEvent;
 };
\ No newline at end of file
diff --git a/Engine/Plugins/VirtualProduction/Avalanche/Source/AvalancheMedia/Private/Playback/AvaPlaybackClient.cpp b/Engine/Plugins/VirtualProduction/Avalanche/Source/AvalancheMedia/Private/Playback/AvaPlaybackClient.cpp
--- a/Engine/Plugins/VirtualProduction/Avalanche/Source/AvalancheMedia/Private/Playback/AvaPlaybackClient.cpp
+++ b/Engine/Plugins/VirtualProduction/Avalanche/Source/AvalancheMedia/Private/Playback/AvaPlaybackClient.cpp
@@ -134,7 +134,8 @@ void FAvaPlaybackClient::Init()
 		.Handling<FAvaPlaybackStatus>(this, &FAvaPlaybackClient::HandlePlaybackStatusMessage)
 		.Handling<FAvaPlaybackStatuses>(this, &FAvaPlaybackClient::HandlePlaybackStatusesMessage)
 		.Handling<FAvaPlaybackSequenceEvent>(this, &FAvaPlaybackClient::HandlePlaybackSequenceEventMessage)
-		.Handling<FAvaPlaybackTransitionEvent>(this, &FAvaPlaybackClient::HandlePlaybackTransitionEventMessage);
+		.Handling<FAvaPlaybackTransitionEvent>(this, &FAvaPlaybackClient::HandlePlaybackTransitionEventMessage)
+		.Handling<FAvaPlaybackTransitionEventMulti>(this, &FAvaPlaybackClient::HandlePlaybackTransitionEventMultiMessage);
 	
 	if (MessageEndpoint.IsValid())
 	{
@@ -1167,16 +1168,25 @@ void FAvaPlaybackClient::HandlePlaybackStatus(FServerInfo& InServerInfo,
 	{
 		PrevPlaybackStatus = EAvaPlaybackStatus::Unknown;
 	}
-	
-	InServerInfo.SetInstanceStatus(InChannelName, InInstanceId, InAssetPath, InStatus);
-	
-	if (bInValidUserData)
+
+	if (InStatus == EAvaPlaybackStatus::Available && InInstanceId.IsValid())
 	{
-		InServerInfo.SetInstanceUserData(InChannelName, InInstanceId, InAssetPath, InUserData);
+		// We need to clear the playback status cache when an actual instance returns to "available".
+		// This prevents the cache from accumulating an ever-growing list of stale instance statuses.
+		InServerInfo.RemoveInstanceInfo(InChannelName, InInstanceId, InAssetPath);
+	}
+	else
+	{
+		InServerInfo.SetInstanceStatus(InChannelName, InInstanceId, InAssetPath, InStatus);
+	
+		if (bInValidUserData)
+		{
+			InServerInfo.SetInstanceUserData(InChannelName, InInstanceId, InAssetPath, InUserData);
 
-		UE_LOG(LogAvaPlaybackClient, Verbose,
-		   TEXT("%s Received Remote Playback \"%s\" (id:%s) from \"%s\" (%s) User Data: %s."),
-		   *GetBriefFrameInfo(), *InAssetPath.GetAssetName(), *InInstanceId.ToString(), *InServerInfo.ServerName, *InChannelName, *InUserData);
+			UE_LOG(LogAvaPlaybackClient, Verbose,
+			   TEXT("%s Received Remote Playback \"%s\" (id:%s) from \"%s\" (%s) User Data: %s."),
+			   *GetBriefFrameInfo(), *InAssetPath.GetAssetName(), *InInstanceId.ToString(), *InServerInfo.ServerName, *InChannelName, *InUserData);
+		}
 	}
 
 	// Clear the pending request (if any).
@@ -1248,15 +1258,33 @@ void FAvaPlaybackClient::HandlePlaybackTransitionEventMessage(const FAvaPlayback
 	const FPlaybackTransitionEventArgs Args =
 	{
 		InMessage.TransitionId,
-		InMessage.InstanceId,
+		InMessage.Entry.InstanceId,
 		InMessage.ChannelName,
 		InMessage.ServerName,
-		InMessage.GetEventFlags(),
+		InMessage.Entry.GetEventFlags(),
 		InMessage.FrameNumber
 	};
 	GetOnPlaybackTransitionEvent().Broadcast(*this, Args);
 }
 
+void FAvaPlaybackClient::HandlePlaybackTransitionEventMultiMessage(const FAvaPlaybackTransitionEventMulti& InMessage, const TSharedRef<IMessageContext, ESPMode::ThreadSafe>& InContext)
+{
+	using namespace UE::AvaPlaybackClient::Delegates;
+	for (const FAvaPlaybackTransitionEventEntry& Entry : InMessage.Entries)
+	{
+		const FPlaybackTransitionEventArgs Args =
+		{
+			InMessage.TransitionId,
+			Entry.InstanceId,
+			InMessage.ChannelName,
+			InMessage.ServerName,
+			Entry.GetEventFlags(),
+			InMessage.FrameNumber
+		};
+		GetOnPlaybackTransitionEvent().Broadcast(*this, Args);
+	}
+}
+
 void FAvaPlaybackClient::RegisterCommands()
 {
 	if (ConsoleCommands.Num() != 0)
diff --git a/Engine/Plugins/VirtualProduction/Avalanche/Source/AvalancheMedia/Private/Playback/AvaPlaybackClient.h b/Engine/Plugins/VirtualProduction/Avalanche/Source/AvalancheMedia/Private/Playback/AvaPlaybackClient.h
--- a/Engine/Plugins/VirtualProduction/Avalanche/Source/AvalancheMedia/Private/Playback/AvaPlaybackClient.h
+++ b/Engine/Plugins/VirtualProduction/Avalanche/Source/AvalancheMedia/Private/Playback/AvaPlaybackClient.h
@@ -77,6 +77,7 @@ protected:
 	void HandlePlaybackStatusesMessage(const FAvaPlaybackStatuses& InMessage, const TSharedRef<IMessageContext, ESPMode::ThreadSafe>& InContext);
 	void HandlePlaybackSequenceEventMessage(const FAvaPlaybackSequenceEvent& InMessage, const TSharedRef<IMessageContext, ESPMode::ThreadSafe>& InContext);
 	void HandlePlaybackTransitionEventMessage(const FAvaPlaybackTransitionEvent& InMessage, const TSharedRef<IMessageContext, ESPMode::ThreadSafe>& InContext);
+	void HandlePlaybackTransitionEventMultiMessage(const FAvaPlaybackTransitionEventMulti& InMessage, const TSharedRef<IMessageContext, ESPMode::ThreadSafe>& InContext);
 
 private:
 	class FServerInfo;
@@ -228,6 +229,11 @@ private:
 			return InInstanceId.IsValid() ? InstanceByIds.Find(InInstanceId) : nullptr;
 		}
 
+		void RemoveInstanceInfo(const FGuid& InInstanceId)
+		{
+			InstanceByIds.Remove(InInstanceId);
+		}
+
 		FPlaybackInstanceInfo* GetOrAddInstanceInfo(const FGuid& InInstanceId)
 		{
 			FPlaybackInstanceInfo* Instance = InstanceByIds.Find(InInstanceId);
@@ -280,6 +286,14 @@ private:
 			}
 			return AssetInfoByPaths.Add(InAssetPath, MakeUnique<FPlaybackAssetInfo>()).Get();
 		}
+
+		void RemoveInstanceInfo(const FGuid& InInstanceId, const FSoftObjectPath& InAssetPath)
+		{
+			if (FPlaybackAssetInfo* AssetInfo = GetAssetInfo(InAssetPath))
+			{
+				AssetInfo->RemoveInstanceInfo(InInstanceId);
+			}
+		}
 		
 		TOptional<EAvaPlaybackStatus> GetInstanceStatus(const FGuid& InInstanceId, const FSoftObjectPath& InAssetPath) const
 		{
@@ -359,12 +373,27 @@ private:
 			return ChannelInfo ? (*ChannelInfo).Get() : nullptr;
 		}
 		
+		FPlaybackChannelInfo* GetPlaybackChannelInfo(const FString& InChannelName)
+		{
+			const TUniquePtr<FPlaybackChannelInfo>* ChannelInfo = PlaybackChannelInfosByName.Find(InChannelName);
+			return ChannelInfo ? (*ChannelInfo).Get() : nullptr;
+		}
+
+		
 		TOptional<EAvaPlaybackStatus> GetInstanceStatus(
 			const FString& InChannelName, const FGuid& InInstanceId, const FSoftObjectPath& InAssetPath) const
 		{
 			const FPlaybackChannelInfo* ChannelInfo = GetPlaybackChannelInfo(InChannelName);
 			return ChannelInfo ? ChannelInfo->GetInstanceStatus(InInstanceId, InAssetPath) : TOptional<EAvaPlaybackStatus>();
 		}
+
+		void RemoveInstanceInfo(const FString& InChannelName, const FGuid& InInstanceId, const FSoftObjectPath& InAssetPath)
+		{
+			if (FPlaybackChannelInfo* ChannelInfo = GetPlaybackChannelInfo(InChannelName))
+			{
+				ChannelInfo->RemoveInstanceInfo(InInstanceId, InAssetPath);
+			}
+		}
 		
 		void SetInstanceStatus(
 			const FString& InChannelName, const FGuid& InInstanceId, const FSoftObjectPath& InAssetPath, EAvaPlaybackStatus InStatus)
diff --git a/Engine/Plugins/VirtualProduction/Avalanche/Source/AvalancheMedia/Private/Playback/AvaPlaybackMessages.h b/Engine/Plugins/VirtualProduction/Avalanche/Source/AvalancheMedia/Private/Playback/AvaPlaybackMessages.h
--- a/Engine/Plugins/VirtualProduction/Avalanche/Source/AvalancheMedia/Private/Playback/AvaPlaybackMessages.h
+++ b/Engine/Plugins/VirtualProduction/Avalanche/Source/AvalancheMedia/Private/Playback/AvaPlaybackMessages.h
@@ -512,6 +512,37 @@ struct FAvaPlaybackTransitionStopRequest : public FAvaPlaybackClientMessageBase
 	bool bUnloadDiscardedInstances = false;
 };
 
+/**
+ * Defines a single transition event entry.
+ * This is the payload for the transition events messages.
+ */
+USTRUCT()
+struct FAvaPlaybackTransitionEventEntry
+{
+	GENERATED_BODY()
+	
+	UPROPERTY()
+	FGuid InstanceId;
+
+	UPROPERTY()
+	uint8 EventFlags = static_cast<uint8>(EAvaPlayableTransitionEventFlags::None);
+
+	FAvaPlaybackTransitionEventEntry() = default;
+	
+	FAvaPlaybackTransitionEventEntry(const FGuid& InInstanceId, EAvaPlayableTransitionEventFlags InFlags)
+		: InstanceId(InInstanceId)
+		, EventFlags(static_cast<uint8>(InFlags))
+	{
+	}
+
+	// Because the enum is used as flags, we need to convert to uint8 manually. (Can't use TEnumAsByte for this apparently.)
+	EAvaPlayableTransitionEventFlags GetEventFlags() const { return static_cast<EAvaPlayableTransitionEventFlags>(EventFlags);}
+	void SetEventFlags(EAvaPlayableTransitionEventFlags InFlags) { EventFlags = static_cast<uint8>(InFlags);}
+};
+
+/**
+ * Single entry transition event message.
+ */
 USTRUCT()
 struct FAvaPlaybackTransitionEvent : public FAvaPlaybackServerMessageBase
 {
@@ -522,19 +553,34 @@ struct FAvaPlaybackTransitionEvent : public FAvaPlaybackServerMessageBase
 
 	UPROPERTY()
 	FGuid TransitionId;
-	
-	UPROPERTY()
-	FGuid InstanceId;
 
 	UPROPERTY()
-	uint8 EventFlags = static_cast<uint8>(EAvaPlayableTransitionEventFlags::None);
+	FAvaPlaybackTransitionEventEntry Entry;
 
 	UPROPERTY()
 	int32 FrameNumber = 0;
+};
 
-	// Because the enum is used as flags, we need to convert to uint8 manually. (Can't use TEnumAsByte for this apparently.)
-	EAvaPlayableTransitionEventFlags GetEventFlags() const { return static_cast<EAvaPlayableTransitionEventFlags>(EventFlags);}
-	void SetEventFlags(EAvaPlayableTransitionEventFlags InFlags) { EventFlags = static_cast<uint8>(InFlags);}
+/**
+ * Multi-entries transition event message.
+ * We use this to group the events in a one message to ensure correct order on reception.
+ */
+USTRUCT()
+struct FAvaPlaybackTransitionEventMulti : public FAvaPlaybackServerMessageBase
+{
+	GENERATED_BODY()
+
+	UPROPERTY()
+	FString ChannelName;
+
+	UPROPERTY()
+	FGuid TransitionId;
+	
+	UPROPERTY()
+	TArray<FAvaPlaybackTransitionEventEntry> Entries;
+
+	UPROPERTY()
+	int32 FrameNumber = 0;
 };
 
 class UMediaOutput;
diff --git a/Engine/Plugins/VirtualProduction/Avalanche/Source/AvalancheMedia/Private/Playback/AvaPlaybackServer.cpp b/Engine/Plugins/VirtualProduction/Avalanche/Source/AvalancheMedia/Private/Playback/AvaPlaybackServer.cpp
--- a/Engine/Plugins/VirtualProduction/Avalanche/Source/AvalancheMedia/Private/Playback/AvaPlaybackServer.cpp
+++ b/Engine/Plugins/VirtualProduction/Avalanche/Source/AvalancheMedia/Private/Playback/AvaPlaybackServer.cpp
@@ -448,16 +448,31 @@ bool FAvaPlaybackServer::RemovePlaybackInstanceTransition(const FGuid& InTransit
 }
 
 void FAvaPlaybackServer::SendPlayableTransitionEvent(
-	const FGuid& InTransitionId, const FGuid& InInstanceId, EAvaPlayableTransitionEventFlags InFlags,
+	const FGuid& InTransitionId, const FAvaPlaybackTransitionEventEntry& InEntry,
 	const FName& InChannelName, const FString& InClientName)
 {
 	FAvaPlaybackTransitionEvent* Message = FMessageEndpoint::MakeMessage<FAvaPlaybackTransitionEvent>();
 	Message->ChannelName = InChannelName.ToString();
 	Message->TransitionId = InTransitionId;
-	Message->InstanceId = InInstanceId;
+	Message->Entry = InEntry;
 	Message->FrameNumber = GFrameNumber;
-	Message->SetEventFlags(InFlags);
-	SendResponse(Message, GetClientAddressSafe(InClientName));
+
+	// Transition events are sent reliably to avoid transitions getting stuck.
+	SendResponse(Message, GetClientAddressSafe(InClientName), EMessageFlags::Reliable);
+}
+
+void FAvaPlaybackServer::SendPlayableTransitionEventMulti(
+	const FGuid& InTransitionId, TArray<FAvaPlaybackTransitionEventEntry>&& InEntries,
+	const FName& InChannelName, const FString& InClientName)
+{
+	FAvaPlaybackTransitionEventMulti* Message = FMessageEndpoint::MakeMessage<FAvaPlaybackTransitionEventMulti>();
+	Message->ChannelName = InChannelName.ToString();
+	Message->TransitionId = InTransitionId;
+	Message->Entries = MoveTemp(InEntries);
+	Message->FrameNumber = GFrameNumber;
+
+	// Transition events are sent reliably to avoid transitions getting stuck.
+	SendResponse(Message, GetClientAddressSafe(InClientName), EMessageFlags::Reliable);
 }
 
 void FAvaPlaybackServer::HandlePlaybackPing(const FAvaPlaybackPing& InMessage, const TSharedRef<IMessageContext, ESPMode::ThreadSafe>& InContext)
@@ -700,6 +715,8 @@ void FAvaPlaybackServer::HandlePlayableTransitionStartRequest(const FAvaPlayback
 		}
 		else
 		{
+			// TODO: have a time out on this.
+			// Avoid having a transition be stuck because it is waiting on an instance that will never come.
 			Transition->AddPendingPlayingInstanceId(PlayingInstanceId);	// Will be resolved later.
 
 			UE_LOG(LogAvaPlaybackServer, Warning,
@@ -738,6 +755,15 @@ void FAvaPlaybackServer::HandlePlayableTransitionStopRequest(const FAvaPlaybackT
 		// todo: should reply with an error message to client. (need a request id)
 		return;
 	}
+
+	static FString None = TEXT("None");
+
+	// Stop all transitions.
+	if (!InMessage.TransitionId.IsValid() && InMessage.ChannelName.Equals(None))
+	{
+		PlaybackInstanceTransitions->Transitions.GetKeys(PlaybackInstanceTransitions->MarkedForStopAndDiscard);
+		return;
+	}
 	
 	if (PlaybackInstanceTransitions->Transitions.Contains(InMessage.TransitionId))
 	{
diff --git a/Engine/Plugins/VirtualProduction/Avalanche/Source/AvalancheMedia/Private/Playback/AvaPlaybackServer.h b/Engine/Plugins/VirtualProduction/Avalanche/Source/AvalancheMedia/Private/Playback/AvaPlaybackServer.h
--- a/Engine/Plugins/VirtualProduction/Avalanche/Source/AvalancheMedia/Private/Playback/AvaPlaybackServer.h
+++ b/Engine/Plugins/VirtualProduction/Avalanche/Source/AvalancheMedia/Private/Playback/AvaPlaybackServer.h
@@ -102,7 +102,11 @@ public:
 	bool RemovePlaybackInstanceTransition(const FGuid& InTransitionId);
 	
 	void SendPlayableTransitionEvent(
-		const FGuid& InTransitionId, const FGuid& InInstanceId, EAvaPlayableTransitionEventFlags InFlags,
+		const FGuid& InTransitionId, const FAvaPlaybackTransitionEventEntry& InEntry,
+		const FName& InChannelName, const FString& InClientName);
+	
+	void SendPlayableTransitionEventMulti(
+		const FGuid& InTransitionId, TArray<FAvaPlaybackTransitionEventEntry>&& InEntries,
 		const FName& InChannelName, const FString& InClientName);
 	
 public:
diff --git a/Engine/Plugins/VirtualProduction/Avalanche/Source/AvalancheMedia/Private/Playback/Transition/AvaPlaybackServerTransition.cpp b/Engine/Plugins/VirtualProduction/Avalanche/Source/AvalancheMedia/Private/Playback/Transition/AvaPlaybackServerTransition.cpp
--- a/Engine/Plugins/VirtualProduction/Avalanche/Source/AvalancheMedia/Private/Playback/Transition/AvaPlaybackServerTransition.cpp
+++ b/Engine/Plugins/VirtualProduction/Avalanche/Source/AvalancheMedia/Private/Playback/Transition/AvaPlaybackServerTransition.cpp
@@ -417,9 +417,7 @@ void UAvaPlaybackServerTransition::Stop()
 		// playable transition event path already.
 		if (!bIsFinishedSent)
 		{
-			constexpr EAvaPlayableTransitionEventFlags FinishedFlag = EAvaPlayableTransitionEventFlags::Finished;
-			PlaybackServer->SendPlayableTransitionEvent(TransitionId, FGuid(), FinishedFlag, ChannelName, ClientName);
-			bIsFinishedSent = true;
+			SendFinishTransitionEvent(*PlaybackServer, EAvaPlayableTransitionEventFlags::Finished);
 		}
 	}
 }
@@ -504,7 +502,12 @@ void UAvaPlaybackServerTransition::OnTransitionEvent(UAvaPlayable* InPlayable, U
 		// Relay the transition event back to the client
 		if (PlaybackServer)
 		{
-			PlaybackServer->SendPlayableTransitionEvent(TransitionId, InPlayable->GetInstanceId(), InTransitionFlags, ChannelName, ClientName);
+			// Keeping a record of the instance event. It will be packed in the final "finish" message to guarantee the order of events.			
+			FAvaPlaybackTransitionEventEntry& EventEntry = InstanceEventEntries.Emplace_GetRef(InPlayable->GetInstanceId(), InTransitionFlags);
+
+			// We also send it immediately in its own message in case the transition is long.
+			// The client will process only once, which ever comes first between those to transmission paths.
+			PlaybackServer->SendPlayableTransitionEvent(TransitionId, EventEntry, ChannelName, ClientName);
 		}
 		
 		if (EnumHasAnyFlags(InTransitionFlags, EAvaPlayableTransitionEventFlags::StopPlayable))
@@ -547,14 +550,22 @@ void UAvaPlaybackServerTransition::OnTransitionEvent(UAvaPlayable* InPlayable, U
 	{
 		if (PlaybackServer)
 		{
-			PlaybackServer->SendPlayableTransitionEvent(TransitionId, FGuid(), InTransitionFlags, ChannelName, ClientName);
-			bIsFinishedSent = true;
+			SendFinishTransitionEvent(*PlaybackServer, InTransitionFlags);
 		}
-
 		Stop();
 	}
 }
 
+void UAvaPlaybackServerTransition::SendFinishTransitionEvent(FAvaPlaybackServer& InPlaybackServer, EAvaPlayableTransitionEventFlags InTransitionFlags)
+{
+	// Transition Finish event must be last, i.e. after all playable instance events.
+	InstanceEventEntries.Add({FGuid(), InTransitionFlags});
+	
+	// This will resend all the playable instance events to ensure they are processed prior to the "finish" event.
+	InPlaybackServer.SendPlayableTransitionEventMulti(TransitionId, MoveTemp(InstanceEventEntries), ChannelName, ClientName);
+	bIsFinishedSent = true;
+}
+
 void UAvaPlaybackServerTransition::OnPlayableCreated(UAvaPlaybackGraph* InPlayback, UAvaPlayable* InPlayable)
 {
 	if (UAvaPlayableGroup* PlayableGroup = InPlayable->GetPlayableGroup())
diff --git a/Engine/Plugins/VirtualProduction/Avalanche/Source/AvalancheMedia/Private/Playback/Transition/AvaPlaybackServerTransition.h b/Engine/Plugins/VirtualProduction/Avalanche/Source/AvalancheMedia/Private/Playback/Transition/AvaPlaybackServerTransition.h
--- a/Engine/Plugins/VirtualProduction/Avalanche/Source/AvalancheMedia/Private/Playback/Transition/AvaPlaybackServerTransition.h
+++ b/Engine/Plugins/VirtualProduction/Avalanche/Source/AvalancheMedia/Private/Playback/Transition/AvaPlaybackServerTransition.h
@@ -3,8 +3,8 @@
 #pragma once
 
 #include "Playable/AvaPlayable.h"
+#include "Playback/AvaPlaybackMessages.h"
 #include "Playback/Transition/AvaPlaybackTransition.h"
-#include "UObject/Object.h"
 #include "AvaPlaybackServerTransition.generated.h"
 
 class FAvaPlaybackInstance;
@@ -86,6 +86,8 @@ protected:
 	bool AddPlaybackInstance(const TSharedPtr<FAvaPlaybackInstance>& InPlaybackInstance, TArray<TWeakPtr<FAvaPlaybackInstance>>& OutPlaybackInstancesWeak);	
 	void UpdateChannelName(const FAvaPlaybackInstance* InPlaybackInstance);
 
+	void SendFinishTransitionEvent(FAvaPlaybackServer& InPlaybackServer, EAvaPlayableTransitionEventFlags InTransitionFlags);
+
 protected:
 	TWeakPtr<FAvaPlaybackServer> PlaybackServerWeak;
 	
@@ -109,9 +111,12 @@ protected:
 	UPROPERTY(Transient)
 	TObjectPtr<UAvaPlayableTransition> PlayableTransition;
 
+	/** Keep track of all instance event entries so we can all send them in a single message to ensure order. */
+	TArray<FAvaPlaybackTransitionEventEntry> InstanceEventEntries;
+
 	/** Re-entry guard for stop() function. */
 	bool bIsStopping = false;
 
-	/** Indicate if the "finish" event has been sent. There is 2 possible paths to send the finish event but we only want to send once. */
+	/** Indicate if the "finish" event has been sent. There is 2 possible paths to send the finish event, but we only want to send once. */
 	bool bIsFinishedSent = false;
 };
\ No newline at end of file
diff --git a/Engine/Plugins/VirtualProduction/Avalanche/Source/AvalancheMedia/Private/Rundown/AvaRundownPagePlayer.cpp b/Engine/Plugins/VirtualProduction/Avalanche/Source/AvalancheMedia/Private/Rundown/AvaRundownPagePlayer.cpp
--- a/Engine/Plugins/VirtualProduction/Avalanche/Source/AvalancheMedia/Private/Rundown/AvaRundownPagePlayer.cpp
+++ b/Engine/Plugins/VirtualProduction/Avalanche/Source/AvalancheMedia/Private/Rundown/AvaRundownPagePlayer.cpp
@@ -403,12 +403,17 @@ void UAvaRundownPagePlayer::HandleOnPlayableSequenceEvent(UAvaPlayable* InPlayab
 		using namespace UE::AvaPlayback::Utils;
 		if (InEventType == EAvaPlayableSequenceEventType::Started)
 		{
-			UE_LOG(LogAvaRundown, Verbose, TEXT("%s Rundown Page %d: Sequence Started \"%s\"."), *GetBriefFrameInfo(), PageId, *InSequenceLabel.ToString());
+			UE_LOG(LogAvaRundown, Verbose,
+				TEXT("%s Rundown Page %d (playable %s): Sequence Started \"%s\"."),
+				*GetBriefFrameInfo(), PageId, *InPlayable->GetInstanceId().ToString(), *InSequenceLabel.ToString());
 		}
 
 		if (InEventType == EAvaPlayableSequenceEventType::Finished)
 		{
-			UE_LOG(LogAvaRundown, Verbose, TEXT("%s Rundown Page %d: Sequence Finished \"%s\"."), *GetBriefFrameInfo(), PageId, *InSequenceLabel.ToString());
+			UE_LOG(LogAvaRundown, Verbose,
+				TEXT("%s Rundown Page %d (playable %s): Sequence Finished \"%s\"."),
+				*GetBriefFrameInfo(), PageId, *InPlayable->GetInstanceId().ToString(), *InSequenceLabel.ToString());
+
 			RundownWeak->NotifyPageSequenceFinished(PageId);
 		}
 	}
diff --git a/Engine/Plugins/VirtualProduction/Avalanche/Source/AvalancheMediaEditor/Private/Broadcast/AvaBroadcastEditor.cpp b/Engine/Plugins/VirtualProduction/Avalanche/Source/AvalancheMediaEditor/Private/Broadcast/AvaBroadcastEditor.cpp
--- a/Engine/Plugins/VirtualProduction/Avalanche/Source/AvalancheMediaEditor/Private/Broadcast/AvaBroadcastEditor.cpp
+++ b/Engine/Plugins/VirtualProduction/Avalanche/Source/AvalancheMediaEditor/Private/Broadcast/AvaBroadcastEditor.cpp
@@ -359,7 +359,9 @@ void FAvaBroadcastEditor::FillPlayToolBar(FToolBarBuilder& ToolBarBuilder)
 		// to reset all the servers.
 		ToolBarBuilder.AddToolBarButton(FUIAction(FExecuteAction::CreateLambda([]
 				{
-					IAvaPlaybackClient& PlaybackClient = IAvaMediaModule::Get().GetPlaybackClient(); 
+					IAvaPlaybackClient& PlaybackClient = IAvaMediaModule::Get().GetPlaybackClient();
+					// Stop all transitions.
+					PlaybackClient.RequestPlayableTransitionStop(FGuid(), NAME_None);
 					// Stop and unload all playbacks.
 					PlaybackClient.RequestPlayback(FGuid(), FSoftObjectPath(), FString(), EAvaPlaybackAction::Unload);
 					// Stop all broadcast channels.
