[Motion Design] Fix grammar error in comment.

diff --git a/Engine/Plugins/VirtualProduction/Avalanche/Source/AvalancheMedia/Private/Playable/Transition/AvaPlayableRemoteTransition.h b/Engine/Plugins/VirtualProduction/Avalanche/Source/AvalancheMedia/Private/Playable/Transition/AvaPlayableRemoteTransition.h
--- a/Engine/Plugins/VirtualProduction/Avalanche/Source/AvalancheMedia/Private/Playable/Transition/AvaPlayableRemoteTransition.h
+++ b/Engine/Plugins/VirtualProduction/Avalanche/Source/AvalancheMedia/Private/Playable/Transition/AvaPlayableRemoteTransition.h
@@ -72,7 +72,7 @@ protected:
 		}
 	}
 
-	/** Returns true if the given events where marked as broadcast for the given instance. */
+	/** Returns true if the given events were marked as broadcast for the given instance. */
 	bool HasEventBeenBroadcast(const FGuid& InInstanceId, EAvaPlayableTransitionEventFlags InFlags) const
 	{
 		if (const EAvaPlayableTransitionEventFlags* Flags = InstanceBroadcastEvent.Find(InInstanceId))
@@ -99,6 +99,6 @@ protected:
 	
 	TMap<FString, FRemoteStatusInfo> RemoteStatusPerServer;
 
-	// Keep track of the events that where locally broadcast to avoid sending them more than once.
+	// Keep track of the events that were locally broadcast to avoid sending them more than once.
 	TMap<FGuid, EAvaPlayableTransitionEventFlags> InstanceBroadcastEvent;
 };
\ No newline at end of file
