Motion Design: The rundown editor controller list now has a description column.

diff --git a/Engine/Plugins/VirtualProduction/Avalanche/Source/AvalancheMediaEditor/Private/Rundown/DetailsView/RemoteControl/Controllers/AvaRundownRCControllerItem.cpp b/Engine/Plugins/VirtualProduction/Avalanche/Source/AvalancheMediaEditor/Private/Rundown/DetailsView/RemoteControl/Controllers/AvaRundownRCControllerItem.cpp
--- a/Engine/Plugins/VirtualProduction/Avalanche/Source/AvalancheMediaEditor/Private/Rundown/DetailsView/RemoteControl/Controllers/AvaRundownRCControllerItem.cpp
+++ b/Engine/Plugins/VirtualProduction/Avalanche/Source/AvalancheMediaEditor/Private/Rundown/DetailsView/RemoteControl/Controllers/AvaRundownRCControllerItem.cpp
@@ -55,6 +55,7 @@ FAvaRundownRCControllerItem::FAvaRundownRCControllerItem(int32 InInstanceIndex,
 			: InController->DisplayName;
 		
 		DisplayNameText = FText::FromName(DisplayName);
+		DescriptionText = InController->Description;
 
 		{
 			// Build formatted tool tip lines with information from the controller.
diff --git a/Engine/Plugins/VirtualProduction/Avalanche/Source/AvalancheMediaEditor/Private/Rundown/DetailsView/RemoteControl/Controllers/AvaRundownRCControllerItem.h b/Engine/Plugins/VirtualProduction/Avalanche/Source/AvalancheMediaEditor/Private/Rundown/DetailsView/RemoteControl/Controllers/AvaRundownRCControllerItem.h
--- a/Engine/Plugins/VirtualProduction/Avalanche/Source/AvalancheMediaEditor/Private/Rundown/DetailsView/RemoteControl/Controllers/AvaRundownRCControllerItem.h
+++ b/Engine/Plugins/VirtualProduction/Avalanche/Source/AvalancheMediaEditor/Private/Rundown/DetailsView/RemoteControl/Controllers/AvaRundownRCControllerItem.h
@@ -21,6 +21,7 @@ public:
 	TSharedRef<ITableRow> CreateWidget(TSharedRef<SAvaRundownRCControllerPanel> InControllerPanel, const TSharedRef<STableViewBase>& InOwnerTable) const;
 
 	FText GetDisplayName() const { return DisplayNameText; }
+	FText GetDescriptionText() const { return DescriptionText; }
 	FText GetToolTipText() const { return ToolTipText; }
 	const FNodeWidgets& GetNodeWidgets() const { return NodeWidgets; }
 	
@@ -35,6 +36,7 @@ private:
 	int32 DisplayIndex = 0;
 	FName AssetName;
 	FText DisplayNameText;
+	FText DescriptionText;
 	FText ToolTipText;
 	FNodeWidgets NodeWidgets;
 
diff --git a/Engine/Plugins/VirtualProduction/Avalanche/Source/AvalancheMediaEditor/Private/Rundown/DetailsView/RemoteControl/Controllers/SAvaRundownRCControllerItemRow.cpp b/Engine/Plugins/VirtualProduction/Avalanche/Source/AvalancheMediaEditor/Private/Rundown/DetailsView/RemoteControl/Controllers/SAvaRundownRCControllerItemRow.cpp
--- a/Engine/Plugins/VirtualProduction/Avalanche/Source/AvalancheMediaEditor/Private/Rundown/DetailsView/RemoteControl/Controllers/SAvaRundownRCControllerItemRow.cpp
+++ b/Engine/Plugins/VirtualProduction/Avalanche/Source/AvalancheMediaEditor/Private/Rundown/DetailsView/RemoteControl/Controllers/SAvaRundownRCControllerItemRow.cpp
@@ -36,6 +36,20 @@ TSharedRef<SWidget> SAvaRundownRCControllerItemRow::GenerateWidgetForColumn(cons
 					]
 				];
 		}
+		else if (InColumnName == SAvaRundownRCControllerPanel::DescriptionColumnName)
+		{
+			return SNew(SScissorRectBox)
+				[
+					SNew(SBox)
+					.VAlign(EVerticalAlignment::VAlign_Center)
+					.Padding(3.f, 2.f, 3.f, 2.f)
+					[
+						SNew(STextBlock)
+						.Text(ItemPtr->GetDescriptionText())
+						.ToolTipText(ItemPtr->GetToolTipText())
+					]
+				];
+		}
 		else if (InColumnName == SAvaRundownRCControllerPanel::ValueColumnName)
 		{
 			if (ItemPtr->GetNodeWidgets().ValueWidget.IsValid())
diff --git a/Engine/Plugins/VirtualProduction/Avalanche/Source/AvalancheMediaEditor/Private/Rundown/DetailsView/RemoteControl/Controllers/SAvaRundownRCControllerPanel.cpp b/Engine/Plugins/VirtualProduction/Avalanche/Source/AvalancheMediaEditor/Private/Rundown/DetailsView/RemoteControl/Controllers/SAvaRundownRCControllerPanel.cpp
--- a/Engine/Plugins/VirtualProduction/Avalanche/Source/AvalancheMediaEditor/Private/Rundown/DetailsView/RemoteControl/Controllers/SAvaRundownRCControllerPanel.cpp
+++ b/Engine/Plugins/VirtualProduction/Avalanche/Source/AvalancheMediaEditor/Private/Rundown/DetailsView/RemoteControl/Controllers/SAvaRundownRCControllerPanel.cpp
@@ -21,6 +21,7 @@
 #define LOCTEXT_NAMESPACE "SAvaRundownRCControllerPanel"
 
 const FName SAvaRundownRCControllerPanel::ControllerColumnName = "ControllerColumn";
+const FName SAvaRundownRCControllerPanel::DescriptionColumnName = "DescriptionColumn";
 const FName SAvaRundownRCControllerPanel::ValueColumnName = "ValueColumn";
 
 FAvaRundownRCControllerHeaderRowExtensionDelegate SAvaRundownRCControllerPanel::HeaderRowExtensionDelegate;
@@ -56,10 +57,13 @@ void SAvaRundownRCControllerPanel::Construct(const FArguments& InArgs, const TSh
 				.CanSelectGeneratedColumn(true)
 				+ SHeaderRow::Column(ControllerColumnName)
 				.DefaultLabel(LOCTEXT("Controller", "Controller"))
+				.FillWidth(0.1f)
+				+ SHeaderRow::Column(DescriptionColumnName)
+				.DefaultLabel(LOCTEXT("Description", "Description"))
 				.FillWidth(0.2f)
 				+ SHeaderRow::Column(ValueColumnName)
 				.DefaultLabel(LOCTEXT("Value", "Value"))
-				.FillWidth(0.8f)
+				.FillWidth(0.7f)
 			)
 		]
 	];
diff --git a/Engine/Plugins/VirtualProduction/Avalanche/Source/AvalancheMediaEditor/Private/Rundown/DetailsView/RemoteControl/Controllers/SAvaRundownRCControllerPanel.h b/Engine/Plugins/VirtualProduction/Avalanche/Source/AvalancheMediaEditor/Private/Rundown/DetailsView/RemoteControl/Controllers/SAvaRundownRCControllerPanel.h
--- a/Engine/Plugins/VirtualProduction/Avalanche/Source/AvalancheMediaEditor/Private/Rundown/DetailsView/RemoteControl/Controllers/SAvaRundownRCControllerPanel.h
+++ b/Engine/Plugins/VirtualProduction/Avalanche/Source/AvalancheMediaEditor/Private/Rundown/DetailsView/RemoteControl/Controllers/SAvaRundownRCControllerPanel.h
@@ -36,6 +36,7 @@ class SAvaRundownRCControllerPanel : public SCompoundWidget
 {
 public:
 	static const FName ControllerColumnName;
+	static const FName DescriptionColumnName;
 	static const FName ValueColumnName;
 
 	static FAvaRundownRCControllerHeaderRowExtensionDelegate& GetHeaderRowExtensionDelegate() { return HeaderRowExtensionDelegate; }
