Motion Design: Fixed a potential crash when a mesh presents no vertices to the alignment tools.

diff --git a/Engine/Plugins/VirtualProduction/Avalanche/Source/AvalancheViewport/Private/AvaScreenAlignmentUtils.cpp b/Engine/Plugins/VirtualProduction/Avalanche/Source/AvalancheViewport/Private/AvaScreenAlignmentUtils.cpp
--- a/Engine/Plugins/VirtualProduction/Avalanche/Source/AvalancheViewport/Private/AvaScreenAlignmentUtils.cpp
+++ b/Engine/Plugins/VirtualProduction/Avalanche/Source/AvalancheViewport/Private/AvaScreenAlignmentUtils.cpp
@@ -271,6 +271,11 @@ namespace UE::AvaViewport::Private
 		// so recalculate vertex positions and the new center and adjust again until it's right. Similar to the Newton-Raphson method.
 		const int32 NumVertices = InActorInfo.World.Vertices.Num();
 
+		if (NumVertices == 0)
+		{
+			return InActorInfo.CameraLocal.Location;
+		}
+
 		const FVector ActorCenterLocationOffset = InActorInfo.World.Center - InActorInfo.World.Location;
 
 		TArray<FVector3f> ActorVertices;
