Motion Design: More potential bugs fixed in relation to the alignment panel.

diff --git a/Engine/Plugins/VirtualProduction/Avalanche/Source/AvalancheViewport/Private/AvaScreenAlignmentUtils.cpp b/Engine/Plugins/VirtualProduction/Avalanche/Source/AvalancheViewport/Private/AvaScreenAlignmentUtils.cpp
--- a/Engine/Plugins/VirtualProduction/Avalanche/Source/AvalancheViewport/Private/AvaScreenAlignmentUtils.cpp
+++ b/Engine/Plugins/VirtualProduction/Avalanche/Source/AvalancheViewport/Private/AvaScreenAlignmentUtils.cpp
@@ -301,7 +301,7 @@ namespace UE::AvaViewport::Private
 
 			for (int32 Index = 0; Index < NumVertices; ++Index)
 			{
-				FVector NewLocationActorVertex = ActorWorldTransform.TransformPosition(InActorInfo.ActorLocal.Vertices[static_cast<uint8>(Index)]);
+				FVector NewLocationActorVertex = ActorWorldTransform.TransformPosition(InActorInfo.ActorLocal.Vertices[Index]);
 				InCoordinateConverter->WorldPositionToViewportPosition(NewLocationActorVertex, VertexScreenCoordinate, VertexDistance);
 				ActorVertices[Index] = {VertexScreenCoordinate.X, VertexScreenCoordinate.Y, static_cast<float>(VertexDistance)};
 			}
@@ -556,7 +556,7 @@ namespace UE::AvaViewport::Private
 				break;
 		}		
 
-		const int32 NumActors = InActors.Num();
+		const int32 NumActors = ActorInfos.Num();
 
 		float Min, Max;
 
@@ -662,7 +662,7 @@ namespace UE::AvaViewport::Private
 				break;
 		}
 
-		const int32 NumActors = InActors.Num();
+		const int32 NumActors = ActorInfos.Num();
 
 		float LoopMin = ActorInfos[0].Screen.GetComponent(TAvaAlignment<InAlignmentType>::Min);
 		float LoopMax = ActorInfos[0].Screen.GetComponent(TAvaAlignment<InAlignmentType>::Max);
