Motion Design Modifiers: Material parameter and global opacity modifiers now use an array of structs instead of maps to store their data for compatibility with Remote Comtrol.

diff --git a/Engine/Plugins/VirtualProduction/ActorModifierCore/Source/ActorModifierCore/Public/Modifiers/ActorModifierCoreBase.h b/Engine/Plugins/VirtualProduction/ActorModifierCore/Source/ActorModifierCore/Public/Modifiers/ActorModifierCoreBase.h
--- a/Engine/Plugins/VirtualProduction/ActorModifierCore/Source/ActorModifierCore/Public/Modifiers/ActorModifierCoreBase.h
+++ b/Engine/Plugins/VirtualProduction/ActorModifierCore/Source/ActorModifierCore/Public/Modifiers/ActorModifierCoreBase.h
@@ -352,7 +352,7 @@ protected:
 	ACTORMODIFIERCORE_API bool RemoveExtension(const FName& InExtensionType);
 
 	/** Logs modifier message if in profiling mode or forced */
-	void LogModifier(const FString& InLog, bool bInForce = false, EActorModifierCoreStatus InType = EActorModifierCoreStatus::Success) const;
+	ACTORMODIFIERCORE_API void LogModifier(const FString& InLog, bool bInForce = false, EActorModifierCoreStatus InType = EActorModifierCoreStatus::Success) const;
 
 	void DeferInitializeModifier();
 
diff --git a/Engine/Plugins/VirtualProduction/Avalanche/Source/AvalancheModifiers/Private/Modifiers/AvaGlobalOpacityModifier.cpp b/Engine/Plugins/VirtualProduction/Avalanche/Source/AvalancheModifiers/Private/Modifiers/AvaGlobalOpacityModifier.cpp
--- a/Engine/Plugins/VirtualProduction/Avalanche/Source/AvalancheModifiers/Private/Modifiers/AvaGlobalOpacityModifier.cpp
+++ b/Engine/Plugins/VirtualProduction/Avalanche/Source/AvalancheModifiers/Private/Modifiers/AvaGlobalOpacityModifier.cpp
@@ -15,7 +15,16 @@ UAvaGlobalOpacityModifier::UAvaGlobalOpacityModifier()
 #if WITH_EDITOR
 	bShowMaterialParameters = false;
 #endif
-	MaterialParameters.ScalarParameters.Add(UDynamicMaterialModel::GlobalOpacityParameterName, GlobalOpacity);
+
+	FAvaMaterialParameterMapScalar* ScalarParameter = MaterialParameters.FindScalarParameter(
+		UDynamicMaterialModel::GlobalOpacityParameterName, 
+		/* Create if needed */ true
+	);
+	
+	if (ScalarParameter)
+	{
+		ScalarParameter->Value = GlobalOpacity;
+	}
 }
 
 #if WITH_EDITOR
@@ -49,9 +58,20 @@ void UAvaGlobalOpacityModifier::OnGlobalOpacityChanged()
 {
 	GlobalOpacity = FMath::Clamp<float>(GlobalOpacity, UE_SMALL_NUMBER * 2, 1.f);
 
-	float& GlobalOpacityRef = MaterialParameters.ScalarParameters.FindOrAdd(UDynamicMaterialModel::GlobalOpacityParameterName);
-	GlobalOpacityRef = GlobalOpacity;
-	OnMaterialParametersChanged();
+	FAvaMaterialParameterMapScalar* ScalarParameter = MaterialParameters.FindScalarParameter(
+		UDynamicMaterialModel::GlobalOpacityParameterName,
+		/* Create if needed */ true
+	);
+
+	if (ScalarParameter)
+	{
+		ScalarParameter->Value = GlobalOpacity;
+		OnMaterialParametersChanged();
+	}
+	else
+	{
+		LogModifier(TEXT("Unable to find or add the global opacity material parameter."), /* Force */ false, EActorModifierCoreStatus::Error);
+	}
 }
 
 void UAvaGlobalOpacityModifier::OnActorMaterialAdded(UMaterialInstanceDynamic* InAdded)
diff --git a/Engine/Plugins/VirtualProduction/Avalanche/Source/AvalancheModifiers/Private/Modifiers/AvaMaterialParameterModifier.cpp b/Engine/Plugins/VirtualProduction/Avalanche/Source/AvalancheModifiers/Private/Modifiers/AvaMaterialParameterModifier.cpp
--- a/Engine/Plugins/VirtualProduction/Avalanche/Source/AvalancheModifiers/Private/Modifiers/AvaMaterialParameterModifier.cpp
+++ b/Engine/Plugins/VirtualProduction/Avalanche/Source/AvalancheModifiers/Private/Modifiers/AvaMaterialParameterModifier.cpp
@@ -12,102 +12,271 @@
 
 #define LOCTEXT_NAMESPACE "AvaMaterialParameterModifier"
 
-void FAvaMaterialParameterMap::MatchKeys(const FAvaMaterialParameterMap& InParameterMap)
+namespace UE::Avalanche::Modifiers::Private
 {
-	// Remove key not tracked anymore, preserve other values
-	for (TMap<FName, float>::TIterator It = ScalarParameters.CreateIterator(); It; ++It)
+	template<typename InArrayType>
+	void MatchKeys(const TArray<InArrayType>& InSource, TArray<InArrayType>& InTarget)
 	{
-		if (!InParameterMap.ScalarParameters.Contains(It->Key))
+		// Cache source parameter names
+		TSet<FName> SourceParameterNames;
+		SourceParameterNames.Reserve(InSource.Num());
+
+		for (const InArrayType& Parameter : InSource)
+		{
+			SourceParameterNames.Add(Parameter.Name);
+		}
+
+		using TParameterIterator = TArray<InArrayType>::TIterator;
+
+		// Remove parameters no longer tracked
+		for (TParameterIterator It = InTarget.CreateIterator(); It; ++It)
+		{
+			if (!SourceParameterNames.Contains(It->Name))
+			{
+				It.RemoveCurrent();
+			}
+		}
+
+		// Cache target parameter names
+		TSet<FName> TargetParameterNames;
+		TargetParameterNames.Reserve(InTarget.Num());
+
+		for (const InArrayType& Parameter : InTarget)
+		{
+			TargetParameterNames.Add(Parameter.Name);
+		}
+
+		using TConstParameterIterator = TArray<InArrayType>::TConstIterator;
+
+		// Add newly tracked values
+		for (TConstParameterIterator It = InSource.CreateConstIterator(); It; ++It)
 		{
-			It.RemoveCurrent();
+			if (!TargetParameterNames.Contains(It->Name))
+			{
+				InTarget.Add(*It);
+			}
 		}
 	}
-	// add newly tracked key, preserve other values
-	for (const TPair<FName, float>& Pair : InParameterMap.ScalarParameters)
+}
+
+void FAvaMaterialParameterMap::MatchKeys(const FAvaMaterialParameterMap& InParameterMap)
+{
+	UE::Avalanche::Modifiers::Private::MatchKeys(InParameterMap.ScalarParameterStructs,  ScalarParameterStructs);
+	UE::Avalanche::Modifiers::Private::MatchKeys(InParameterMap.VectorParameterStructs,  VectorParameterStructs);
+	UE::Avalanche::Modifiers::Private::MatchKeys(InParameterMap.TextureParameterStructs, TextureParameterStructs);
+}
+
+void FAvaMaterialParameterMap::PushParametersTo(UMaterialInstanceDynamic* InMaterial)
+{
+	if (!IsValid(InMaterial))
+	{
+		return;
+	}
+
+	// Apply Scalar
+	for (const FAvaMaterialParameterMapScalar& ScalarPair : ScalarParameterStructs)
 	{
-		ScalarParameters.FindOrAdd(Pair.Key);
+		InMaterial->SetScalarParameterValue(ScalarPair.Name, ScalarPair.Value);
 	}
 
-	// Remove key not tracked anymore, preserve other values
-	for (TMap<FName, FLinearColor>::TIterator It = VectorParameters.CreateIterator(); It; ++It)
+	// Apply Color
+	for (const FAvaMaterialParameterMapVector& VectorPair : VectorParameterStructs)
 	{
-		if (!InParameterMap.VectorParameters.Contains(It->Key))
-		{
-			It.RemoveCurrent();
-		}
+		InMaterial->SetVectorParameterValue(VectorPair.Name, VectorPair.Value);
 	}
-	// add newly tracked key, preserve other values
-	for (const TPair<FName, FLinearColor>& Pair : InParameterMap.VectorParameters)
+
+	// Apply Texture
+	for (const FAvaMaterialParameterMapTexture& TexturePair : TextureParameterStructs)
+	{
+		InMaterial->SetTextureParameterValue(TexturePair.Name, TexturePair.Value);
+	}
+}
+
+void FAvaMaterialParameterMap::PullParametersFrom(UMaterialInstanceDynamic* InMaterial)
+{
+	if (!IsValid(InMaterial))
+	{
+		return;
+	}
+
+	// Get Scalar
+	for (FAvaMaterialParameterMapScalar& ScalarPair : ScalarParameterStructs)
+	{
+		ScalarPair.Value = InMaterial->K2_GetScalarParameterValue(ScalarPair.Name);
+	}
+
+	// Get Color
+	for (FAvaMaterialParameterMapVector& VectorPair : VectorParameterStructs)
+	{
+		VectorPair.Value = InMaterial->K2_GetVectorParameterValue(VectorPair.Name);
+	}
+
+	// Get Texture
+	for (FAvaMaterialParameterMapTexture& TexturePair : TextureParameterStructs)
+	{
+		TexturePair.Value = InMaterial->K2_GetTextureParameterValue(TexturePair.Name);
+	}
+}
+
+const FAvaMaterialParameterMapScalar* FAvaMaterialParameterMap::FindScalarParameter(FName InParameterName) const
+{
+	for (const FAvaMaterialParameterMapScalar& Parameter : ScalarParameterStructs)
 	{
-		VectorParameters.FindOrAdd(Pair.Key);
+		if (Parameter.Name == InParameterName)
+		{
+			return &Parameter;
+		}
 	}
 
-	// Remove key not tracked anymore, preserve other values
-	for (TMap<FName, TObjectPtr<UTexture>>::TIterator It = TextureParameters.CreateIterator(); It; ++It)
+	return nullptr;
+}
+
+FAvaMaterialParameterMapScalar* FAvaMaterialParameterMap::FindScalarParameter(FName InParameterName, bool bInCreateIfMissing)
+{
+	for (FAvaMaterialParameterMapScalar& Parameter : ScalarParameterStructs)
 	{
-		if (!InParameterMap.TextureParameters.Contains(It->Key))
+		if (Parameter.Name == InParameterName)
 		{
-			It.RemoveCurrent();
+			return &Parameter;
 		}
 	}
-	// add newly tracked key, preserve other values
-	for (const TPair<FName, TObjectPtr<UTexture>>& Pair : InParameterMap.TextureParameters)
+
+	if (!bInCreateIfMissing)
 	{
-		TextureParameters.FindOrAdd(Pair.Key);
+		return nullptr;
 	}
+
+	ScalarParameterStructs.Emplace(InParameterName, 0.f);
+
+	return &ScalarParameterStructs.Last();
 }
 
-void FAvaMaterialParameterMap::Set(UMaterialInstanceDynamic* InMaterial)
+const FAvaMaterialParameterMapVector* FAvaMaterialParameterMap::FindVectorParameter(FName InParameterName) const
 {
-	if (!IsValid(InMaterial))
+	for (const FAvaMaterialParameterMapVector& Parameter : VectorParameterStructs)
 	{
-		return;
+		if (Parameter.Name == InParameterName)
+		{
+			return &Parameter;
+		}
 	}
 
-	// Apply scalar
-	for (const TPair<FName, float>& Pair : ScalarParameters)
+	return nullptr;
+}
+
+FAvaMaterialParameterMapVector* FAvaMaterialParameterMap::FindVectorParameter(FName InParameterName, bool bInCreateIfMissing)
+{
+	for (FAvaMaterialParameterMapVector& Parameter : VectorParameterStructs)
 	{
-		InMaterial->SetScalarParameterValue(Pair.Key, Pair.Value);
+		if (Parameter.Name == InParameterName)
+		{
+			return &Parameter;
+		}
 	}
 
-	// Apply color
-	for (const TPair<FName, FLinearColor>& Pair : VectorParameters)
+	if (!bInCreateIfMissing)
 	{
-		InMaterial->SetVectorParameterValue(Pair.Key, Pair.Value);
+		return nullptr;
 	}
 
-	// Apply TextureParameters
-	for (const TPair<FName, TObjectPtr<UTexture>>& Pair : TextureParameters)
+	VectorParameterStructs.Emplace(InParameterName, FLinearColor::Black);
+
+	return &VectorParameterStructs.Last();
+}
+
+const FAvaMaterialParameterMapTexture* FAvaMaterialParameterMap::FindTextureParameter(FName InParameterName) const
+{
+	for (const FAvaMaterialParameterMapTexture& Parameter : TextureParameterStructs)
 	{
-		InMaterial->SetTextureParameterValue(Pair.Key, Pair.Value);
+		if (Parameter.Name == InParameterName)
+		{
+			return &Parameter;
+		}
 	}
+
+	return nullptr;
 }
 
-void FAvaMaterialParameterMap::Get(UMaterialInstanceDynamic* InMaterial)
+FAvaMaterialParameterMapTexture* FAvaMaterialParameterMap::FindTextureParameter(FName InParameterName, bool bInCreateIfMissing)
 {
-	if (!IsValid(InMaterial))
+	for (FAvaMaterialParameterMapTexture& Parameter : TextureParameterStructs)
 	{
-		return;
+		if (Parameter.Name == InParameterName)
+		{
+			return &Parameter;
+		}
 	}
 
-	// Get scalar
-	for (TPair<FName, float>& Pair : ScalarParameters)
+	if (!bInCreateIfMissing)
 	{
-		Pair.Value = InMaterial->K2_GetScalarParameterValue(Pair.Key);
+		return nullptr;
 	}
 
-	// Get Color
-	for (TPair<FName, FLinearColor>& Pair : VectorParameters)
+	TextureParameterStructs.Emplace(InParameterName, nullptr);
+
+	return &TextureParameterStructs.Last();
+}
+
+#if WITH_EDITOR
+void FAvaMaterialParameterMap::UpdateDeprecatedParameterProperties()
+{
+	PRAGMA_DISABLE_DEPRECATION_WARNINGS
+
+	for (const TPair<FName, float>& ScalarPair : ScalarParameters_DEPRECATED)
 	{
-		Pair.Value = InMaterial->K2_GetVectorParameterValue(Pair.Key);
+		FAvaMaterialParameterMapScalar* ScalarParameter = FindScalarParameter(ScalarPair.Key, /* Create if doesn't exist */ false);
+
+		// Only add old values if new ones don't already exist.
+		if (!ScalarParameter)
+		{
+			ScalarParameter = FindScalarParameter(ScalarPair.Key, /* Create if doesn't exist */ true);
+			
+			if (ScalarParameter)
+			{
+				ScalarParameter->Value = ScalarPair.Value;
+			}
+		}
 	}
 
-	// Get Texture
-	for (TPair<FName, TObjectPtr<UTexture>>& Pair : TextureParameters)
+	for (const TPair<FName, FLinearColor>& VectorPair : VectorParameters_DEPRECATED)
+	{
+		FAvaMaterialParameterMapVector* VectorParameter = FindVectorParameter(VectorPair.Key, /* Create if doesn't exist */ false);
+
+		// Only add old values if new ones don't already exist.
+		if (!VectorParameter)
+		{
+			VectorParameter = FindVectorParameter(VectorPair.Key, /* Create if doesn't exist */ true);
+
+			if (VectorParameter)
+			{
+				VectorParameter->Value = VectorPair.Value;
+			}
+		}
+	}
+
+	for (const TPair<FName, TObjectPtr<UTexture>>& TexturePair : TextureParameters_DEPRECATED)
 	{
-		Pair.Value = InMaterial->K2_GetTextureParameterValue(Pair.Key);
+		FAvaMaterialParameterMapTexture* TextureParameter = FindTextureParameter(TexturePair.Key, /* Create if doesn't exist */ false);
+
+		// Only add old values if new ones don't already exist.
+		if (!TextureParameter)
+		{
+			TextureParameter = FindTextureParameter(TexturePair.Key, /* Create if doesn't exist */ true);
+
+			if (TextureParameter)
+			{
+				TextureParameter->Value = TexturePair.Value;
+			}
+		}
 	}
+
+	ScalarParameters_DEPRECATED.Empty();
+	VectorParameters_DEPRECATED.Empty();
+	TextureParameters_DEPRECATED.Empty();
+
+	PRAGMA_ENABLE_DEPRECATION_WARNINGS
 }
+#endif
 
 UAvaMaterialParameterModifier::UAvaMaterialParameterModifier()
 {
@@ -208,7 +377,7 @@ void UAvaMaterialParameterModifier::Apply()
 		if (UMaterialInstanceDynamic* MID = MaterialParameterPair.Key.Get())
 		{
 			// Set State
-			MaterialParameters.Set(MID);
+			MaterialParameters.PushParametersTo(MID);
 		}
 	}
 
@@ -318,6 +487,25 @@ TSet<UMaterialInstanceDynamic*> UAvaMaterialParameterModifier::GetComponentDynam
 	return Materials;
 }
 
+#if WITH_EDITOR
+void UAvaMaterialParameterModifier::PostLoad()
+{
+	Super::PostLoad();
+
+	UpdateDeprecatedParameterProperties();
+}
+
+void UAvaMaterialParameterModifier::UpdateDeprecatedParameterProperties()
+{
+	MaterialParameters.UpdateDeprecatedParameterProperties();
+
+	for (TPair<TWeakObjectPtr<UMaterialInstanceDynamic>, FAvaMaterialParameterMap>& SavedMaterialPair : SavedMaterialParameters)
+	{
+		SavedMaterialPair.Value.UpdateDeprecatedParameterProperties();
+	}
+}
+#endif
+
 void UAvaMaterialParameterModifier::SetMaterialParameters(const FAvaMaterialParameterMap& InParameterMap)
 {
 	MaterialParameters = InParameterMap;
@@ -363,7 +551,7 @@ void UAvaMaterialParameterModifier::SaveMaterialParameters()
 	{
 		if (UMaterialInstanceDynamic* MID = MaterialParameterPair.Key.Get())
 		{
-			MaterialParameterPair.Value.Get(MID);
+			MaterialParameterPair.Value.PullParametersFrom(MID);
 		}
 	}
 }
@@ -375,7 +563,7 @@ void UAvaMaterialParameterModifier::RestoreMaterialParameters()
 	{
 		if (UMaterialInstanceDynamic* MID = MaterialParameterPair.Key.Get())
 		{
-			MaterialParameterPair.Value.Set(MID);
+			MaterialParameterPair.Value.PushParametersTo(MID);
 		}
 	}
 }
@@ -409,7 +597,7 @@ void UAvaMaterialParameterModifier::ScanActorMaterials()
 			if (bMaterialAdded)
 			{
 				// Save original values
-                ParameterMap.Get(MID);
+                ParameterMap.PushParametersTo(MID);
 				OnActorMaterialAdded(MID);
 			}
 
diff --git a/Engine/Plugins/VirtualProduction/Avalanche/Source/AvalancheModifiers/Public/Modifiers/AvaMaterialParameterModifier.h b/Engine/Plugins/VirtualProduction/Avalanche/Source/AvalancheModifiers/Public/Modifiers/AvaMaterialParameterModifier.h
--- a/Engine/Plugins/VirtualProduction/Avalanche/Source/AvalancheModifiers/Public/Modifiers/AvaMaterialParameterModifier.h
+++ b/Engine/Plugins/VirtualProduction/Avalanche/Source/AvalancheModifiers/Public/Modifiers/AvaMaterialParameterModifier.h
@@ -3,6 +3,7 @@
 #pragma once
 
 #include "Modifiers/ActorModifierArrangeBaseModifier.h"
+#include "UObject/NameTypes.h"
 #include "UObject/ObjectPtr.h"
 #include "AvaMaterialParameterModifier.generated.h"
 
@@ -10,27 +11,113 @@ class UMaterialInstanceDynamic;
 class UTexture;
 
 USTRUCT(BlueprintType)
-struct FAvaMaterialParameterMap
+struct FAvaMaterialParameterMapScalar
 {
 	GENERATED_BODY()
 
-	UPROPERTY(EditInstanceOnly, BlueprintReadWrite, Category="MaterialParameter")
-	TMap<FName, float> ScalarParameters;
+	UPROPERTY(EditInstanceOnly, BlueprintReadWrite, Category = "MaterialParameter")
+	FName Name;
+
+	UPROPERTY(EditInstanceOnly, BlueprintReadWrite, Category = "MaterialParameter")
+	float Value;
+};
+
+USTRUCT(BlueprintType)
+struct FAvaMaterialParameterMapVector
+{
+	GENERATED_BODY()
+
+	UPROPERTY(EditInstanceOnly, BlueprintReadWrite, Category = "MaterialParameter")
+	FName Name;
+	
+	UPROPERTY(EditInstanceOnly, BlueprintReadWrite, Category = "MaterialParameter")
+	FLinearColor Value;
+};
+
+USTRUCT(BlueprintType)
+struct FAvaMaterialParameterMapTexture
+{
+	GENERATED_BODY()
+
+	UPROPERTY(EditInstanceOnly, BlueprintReadWrite, Category = "MaterialParameter")
+	FName Name;
+
+	UPROPERTY(EditInstanceOnly, BlueprintReadWrite, Category = "MaterialParameter")
+	TObjectPtr<UTexture> Value;
+};
+
+USTRUCT(BlueprintType)
+struct FAvaMaterialParameterMap
+{
+	GENERATED_BODY()
 
-	UPROPERTY(EditInstanceOnly, BlueprintReadWrite, Category="MaterialParameter")
-	TMap<FName, FLinearColor> VectorParameters;
+	friend class UAvaMaterialParameterModifier;
 
-	UPROPERTY(EditInstanceOnly, BlueprintReadWrite, Category="MaterialParameter")
-	TMap<FName, TObjectPtr<UTexture>> TextureParameters;
+	PRAGMA_DISABLE_DEPRECATION_WARNINGS
+	FAvaMaterialParameterMap() = default;
+	FAvaMaterialParameterMap(const FAvaMaterialParameterMap&) = default;
+	FAvaMaterialParameterMap(FAvaMaterialParameterMap&&) = default;
+	FAvaMaterialParameterMap& operator=(const FAvaMaterialParameterMap&) = default;
+	FAvaMaterialParameterMap& operator=(FAvaMaterialParameterMap&&) = default;
+	PRAGMA_ENABLE_DEPRECATION_WARNINGS
 
 	/** Matches the input parameter key map and removes all unused keys, does not touch current values */
 	void MatchKeys(const FAvaMaterialParameterMap& InParameterMap);
 
-	/** Apply those parameters on this Material Designer Instance */
-	void Set(UMaterialInstanceDynamic* InMaterial);
+	/** Apply our parameters to this Material Designer Instance */
+	void PushParametersTo(UMaterialInstanceDynamic* InMaterial);
 
-	/** Read those parameters from this Material Designer Instance and save them */
-	void Get(UMaterialInstanceDynamic* InMaterial);
+	/** Read parameters from this Material Designer Instance and save them here. */
+	void PullParametersFrom(UMaterialInstanceDynamic* InMaterial);
+
+	/** Returns the scalar parameter struct with the given name. */
+	AVALANCHEMODIFIERS_API const FAvaMaterialParameterMapScalar* FindScalarParameter(FName InParameterName) const;
+
+	/** Returns the scalar parameter struct with the given name, potentially creating it if it doesn't exist. */
+	AVALANCHEMODIFIERS_API FAvaMaterialParameterMapScalar* FindScalarParameter(FName InParameterName, bool bInCreateIfMissing = false);
+
+	/** Returns the vector parameter struct with the given name. */
+	AVALANCHEMODIFIERS_API const FAvaMaterialParameterMapVector* FindVectorParameter(FName InParameterName) const;
+
+	/** Returns the vector parameter struct with the given name, potentially creating it if it doesn't exist. */
+	AVALANCHEMODIFIERS_API FAvaMaterialParameterMapVector* FindVectorParameter(FName InParameterName, bool bInCreateIfMissing = false);
+
+	/** Returns the texture parameter struct with the given name. */
+	AVALANCHEMODIFIERS_API const FAvaMaterialParameterMapTexture* FindTextureParameter(FName InParameterName) const;
+
+	/** Returns the texture parameter struct with the given name, potentially creating it if it doesn't exist. */
+	AVALANCHEMODIFIERS_API FAvaMaterialParameterMapTexture* FindTextureParameter(FName InParameterName, bool bInCreateIfMissing = false);
+
+private:
+#if WITH_EDITOR
+	void UpdateDeprecatedParameterProperties();
+#endif
+
+	UPROPERTY(EditInstanceOnly, BlueprintReadWrite, Category = "MaterialParameter",
+		meta = (DisplayName = "Scalar Parameters", AllowPrivateAccess, TitleProperty = "Name", ShowOnlyInnerProperties))
+	TArray<FAvaMaterialParameterMapScalar> ScalarParameterStructs;
+
+	UPROPERTY(EditInstanceOnly, BlueprintReadWrite, Category = "MaterialParameter",
+		meta = (DisplayName = "Vector Parameters", AllowPrivateAccess, TitleProperty = "Name", ShowOnlyInnerProperties))
+	TArray<FAvaMaterialParameterMapVector> VectorParameterStructs;
+
+	UPROPERTY(EditInstanceOnly, BlueprintReadWrite, Category = "MaterialParameter",
+		meta = (DisplayName = "Texture Parameters", AllowPrivateAccess, TitleProperty = "Name", ShowOnlyInnerProperties))
+	TArray<FAvaMaterialParameterMapTexture> TextureParameterStructs;
+
+#if WITH_EDITORONLY_DATA
+	UE_DEPRECATED(5.7, "ScalarParameters has been deprecated. Please use ScalarParameterStructs")
+	UPROPERTY(meta = (DeprecatedProperty, DeprecationMessage = "This property has been refactored to use an array and struct, instead of a map."))
+	TMap<FName, float> ScalarParameters_DEPRECATED;
+
+	UE_DEPRECATED(5.7, "VectorParameters has been deprecated. Please use VectorParameterStructs")
+	UPROPERTY(meta = (DeprecatedProperty, DeprecationMessage = "This property has been refactored to use an array and struct, instead of a map."))
+	TMap<FName, FLinearColor> VectorParameters_DEPRECATED;
+
+	UE_DEPRECATED(5.7, "TextureParameters has been deprecated. Please use TextureParameterStructs")
+	UPROPERTY(meta = (DeprecatedProperty, DeprecationMessage = "This property has been refactored to use an array and struct, instead of a map."))
+	TMap<FName, TObjectPtr<UTexture>> TextureParameters_DEPRECATED;
+#endif
 };
 
 /** This modifier sets specified dynamic materials parameters on an actor and its children */
@@ -64,6 +151,7 @@ protected:
 	//~ Begin UObject
 #if WITH_EDITOR
 	virtual void PostEditChangeProperty(FPropertyChangedEvent& PropertyChangedEvent) override;
+	virtual void PostLoad() override;
 #endif
 	//~ End UObject
 
@@ -127,4 +215,9 @@ protected:
 	UPROPERTY(Transient)
 	bool bShowMaterialParameters = true;
 #endif
+
+private:
+#if WITH_EDITOR
+	void UpdateDeprecatedParameterProperties();
+#endif
 };
