Motion Designer Modifiers: Material parameter redux

diff --git a/Engine/Plugins/VirtualProduction/Avalanche/Source/AvalancheModifiers/Private/Modifiers/AvaMaterialParameterModifier.cpp b/Engine/Plugins/VirtualProduction/Avalanche/Source/AvalancheModifiers/Private/Modifiers/AvaMaterialParameterModifier.cpp
--- a/Engine/Plugins/VirtualProduction/Avalanche/Source/AvalancheModifiers/Private/Modifiers/AvaMaterialParameterModifier.cpp
+++ b/Engine/Plugins/VirtualProduction/Avalanche/Source/AvalancheModifiers/Private/Modifiers/AvaMaterialParameterModifier.cpp
@@ -597,7 +597,7 @@ void UAvaMaterialParameterModifier::ScanActorMaterials()
 			if (bMaterialAdded)
 			{
 				// Save original values
-                ParameterMap.PushParametersTo(MID);
+                ParameterMap.PullParametersFrom(MID);
 				OnActorMaterialAdded(MID);
 			}
 
