MotionDesign : GeometryMask
- Fixed mask not resolving/applying on level load for 5.6.x levels
- Fixed mask canvas texture not refreshing when leaving PIE

diff --git a/Engine/Plugins/VirtualProduction/Avalanche/Source/AvalancheMask/Private/Handling/AvaMaskAvaShapeMaterialCollectionHandle.cpp b/Engine/Plugins/VirtualProduction/Avalanche/Source/AvalancheMask/Private/Handling/AvaMaskAvaShapeMaterialCollectionHandle.cpp
--- a/Engine/Plugins/VirtualProduction/Avalanche/Source/AvalancheMask/Private/Handling/AvaMaskAvaShapeMaterialCollectionHandle.cpp
+++ b/Engine/Plugins/VirtualProduction/Avalanche/Source/AvalancheMask/Private/Handling/AvaMaskAvaShapeMaterialCollectionHandle.cpp
@@ -11,9 +11,10 @@
 #include "DynamicMeshes/AvaShapeDynMeshBase.h"
 #include "GameFramework/Actor.h"
 #include "Handling/AvaHandleUtilities.h"
-#include "IAvaMaskMaterialHandle.h"
 #include "Handling/AvaObjectHandleSubsystem.h"
 #include "Materials/MaterialInterface.h"
+#include "IAvaMaskMaterialHandle.h"
+#include "Materials/Material.h"
 
 namespace UE::AvaMask::Private
 {
@@ -86,23 +87,22 @@ TArray<TSharedPtr<IAvaMaskMaterialHandle>> FAvaMaskAvaShapeMaterialCollectionHan
 		RefreshMaterials();
 	}
 	
-	for (TPair<int32, TSharedPtr<IAvaMaskMaterialHandle>>& SlotMaterialHandle : MaterialHandles)
+	for (const TPair<int32, TSharedPtr<IAvaMaskMaterialHandle>>& SlotMaterialHandle : MaterialHandles)
 	{
 		// Account for empty/unoccupied slot
-		if (TSharedPtr<IAvaMaskMaterialHandle>& MaterialHandle = SlotMaterialHandle.Value)
+		TSharedPtr<IAvaMaskMaterialHandle> MaterialHandle = SlotMaterialHandle.Value;
+		if (!UE::Ava::Internal::IsHandleValid(MaterialHandle))
 		{
-			if (!UE::Ava::Internal::IsHandleValid(MaterialHandle))
+			if (RefreshMaterials())
 			{
-				if (RefreshMaterials())
-				{
-					return GetMaterialHandles();
-				}
-
-				return { };
+				MaterialHandles.GenerateValueArray(MtlHandles);
+				break;
 			}
-		
-			MtlHandles.Emplace(MaterialHandle);
+
+			return { };
 		}
+		
+		MtlHandles.Emplace(MaterialHandle);
 	}
 	
 	return MtlHandles;
@@ -395,7 +395,10 @@ bool FAvaMaskAvaShapeMaterialCollectionHandle::RefreshMaterials()
 				}
 				else if (!Material)
 				{
-					MaterialHandles.Emplace(InMeshIdx, nullptr);
+					// Use default material to avoid invalid handle
+					UMaterial* DefaultMaterial = LoadObject<UMaterial>(nullptr, TEXT("/Script/Engine.Material'/Engine/EngineMaterials/WorldGridMaterial.WorldGridMaterial'"));
+					MaterialHandle = MaterialHandles.Emplace(InMeshIdx, HandleSubsystem->MakeHandle<IAvaMaskMaterialHandle>(DefaultMaterial, UE::AvaMask::Internal::HandleTag));
+					OccupiedMaterialSlots[InMeshIdx] = true;
 				}
 				else
 				{
diff --git a/Engine/Plugins/VirtualProduction/Avalanche/Source/AvalancheMask/Private/Handling/AvaMaterialInstanceHandle.cpp b/Engine/Plugins/VirtualProduction/Avalanche/Source/AvalancheMask/Private/Handling/AvaMaterialInstanceHandle.cpp
--- a/Engine/Plugins/VirtualProduction/Avalanche/Source/AvalancheMask/Private/Handling/AvaMaterialInstanceHandle.cpp
+++ b/Engine/Plugins/VirtualProduction/Avalanche/Source/AvalancheMask/Private/Handling/AvaMaterialInstanceHandle.cpp
@@ -4,6 +4,7 @@
 
 #include "AvaMaskMaterialReference.h"
 #include "Engine/Texture.h"
+#include "Materials/Material.h"
 #include "Materials/MaterialInstanceDynamic.h"
 #include "Materials/MaterialInterface.h"
 #include "UObject/Package.h"
@@ -51,7 +52,7 @@ void FAvaMaterialInstanceHandle::SetParentMaterial(UMaterialInterface* InMateria
 bool FAvaMaterialInstanceHandle::IsValid() const
 {
 	// @note: MaterialInstance doesn't have to be valid - this handle should deal with MID creation
-	return WeakParentMaterial.IsValid();
+	return WeakParentMaterial.IsValid() || WeakMaterialInstance.IsValid();
 }
 
 bool FAvaMaterialInstanceHandle::IsSupported(const FAvaMaskMaterialReference& InInstance, FName InTag)
@@ -66,6 +67,11 @@ UMaterialInterface* FAvaMaterialInstanceHandle::GetParentMaterial() const
 		return ParentMaterial;
 	}
 
+	if (UMaterialInstanceDynamic* MID = WeakMaterialInstance.Get())
+	{
+		return MID->GetBaseMaterial();
+	}
+
 	return nullptr;
 }
 
diff --git a/Engine/Plugins/VirtualProduction/Avalanche/Source/AvalancheMask/Private/Handling/AvaMaterialInstanceHandle.h b/Engine/Plugins/VirtualProduction/Avalanche/Source/AvalancheMask/Private/Handling/AvaMaterialInstanceHandle.h
--- a/Engine/Plugins/VirtualProduction/Avalanche/Source/AvalancheMask/Private/Handling/AvaMaterialInstanceHandle.h
+++ b/Engine/Plugins/VirtualProduction/Avalanche/Source/AvalancheMask/Private/Handling/AvaMaterialInstanceHandle.h
@@ -19,7 +19,9 @@ class FAvaMaterialInstanceHandle
 {
 public:
 	UE_AVA_INHERITS_WITH_SUPER(FAvaMaterialInstanceHandle, IAvaMaterialHandle);
-	
+
+	static bool IsSupported(const FAvaMaskMaterialReference& InInstance, FName InTag = NAME_None);
+
 	explicit FAvaMaterialInstanceHandle(const TWeakObjectPtr<UMaterialInterface>& InWeakParentMaterial);
 	virtual ~FAvaMaterialInstanceHandle() override = default;
 
@@ -34,14 +36,11 @@ public:
 	virtual bool IsValid() const override;
 	// ~End IAvaObjectHandle
 
-	static bool IsSupported(const FAvaMaskMaterialReference& InInstance, FName InTag = NAME_None);
-
 protected:
 	virtual UMaterialInterface* GetParentMaterial() const;
 	virtual UMaterialInstanceDynamic* GetMaterialInstance();
 	virtual UMaterialInstanceDynamic* GetOrCreateMaterialInstance();
-	
-protected:
+
 	TWeakObjectPtr<UMaterialInterface> WeakParentMaterial;
 	TWeakObjectPtr<UMaterialInstanceDynamic> WeakMaterialInstance;
 };
diff --git a/Engine/Plugins/VirtualProduction/Avalanche/Source/AvalancheMask/Private/Mask2D/AvaMask2DBaseModifier.cpp b/Engine/Plugins/VirtualProduction/Avalanche/Source/AvalancheMask/Private/Mask2D/AvaMask2DBaseModifier.cpp
--- a/Engine/Plugins/VirtualProduction/Avalanche/Source/AvalancheMask/Private/Mask2D/AvaMask2DBaseModifier.cpp
+++ b/Engine/Plugins/VirtualProduction/Avalanche/Source/AvalancheMask/Private/Mask2D/AvaMask2DBaseModifier.cpp
@@ -358,7 +358,7 @@ void UAvaMask2DBaseModifier::SavePreState()
 			const UActorModifierCoreBase* MaskModifier = UE::ActorModifierCore::Utilities::FindFirstActorModifierByClass(InActor, UAvaMask2DBaseModifier::StaticClass());
 
 			// If first modifier in the hierarchy is not this modifier, then skip actor
-			if (!MaskModifier || MaskModifier == this)
+			if (MaskModifier == this)
 			{
 				SaveActorPreState(InActor, NewActorDatas.Emplace(InActor, ActorData.FindOrAdd(InActor)));
 			}
@@ -399,6 +399,8 @@ void UAvaMask2DBaseModifier::RestorePreState()
 		
 		RemoveFromActor(Actor);
 	}
+
+	ScannedPrimitiveComponents.Empty();
 }
 
 void UAvaMask2DBaseModifier::Apply()
@@ -417,7 +419,7 @@ void UAvaMask2DBaseModifier::SaveActorPreState(AActor* InActor, FAvaMask2DActorD
 			{
 				TSharedPtr<IAvaMaskMaterialCollectionHandle> Handle = GetObjectHandleSubsystem()->MakeHandle<IAvaMaskMaterialCollectionHandle>(InActor, UE::AvaMask::Internal::HandleTag);
 
-				Handle->ForEachMaterialHandle([this, &InActorData](
+				Handle->ForEachMaterialHandle([this, &InActorData, InActor](
 					const FSoftComponentReference& InComponent,
 					const int32 InIdx,
 					const bool bInIsSlotOccupied,
@@ -425,25 +427,26 @@ void UAvaMask2DBaseModifier::SaveActorPreState(AActor* InActor, FAvaMask2DActorD
 					{
 						if (InMaterialHandle.IsValid())
 						{
-							if (const TObjectPtr<UMaterialInterface>* Material = InActorData.OriginalMaterials.Find(FAvaMask2DComponentMaterialPath(InComponent, InIdx)))
+							const FAvaMask2DComponentMaterialPath Component(InComponent, InIdx);
+							if (const TObjectPtr<UMaterialInterface>* MaterialPtr = InActorData.OriginalMaterials.Find(Component))
 							{
-								// Refresh parent material to create MID from
-								InMaterialHandle->SetParentMaterial(Material->Get());
+								if (UMaterialInterface* Material = MaterialPtr->Get())
+								{
+									// Refresh parent material to create MID from
+									InMaterialHandle->SetParentMaterial(Material);
+								}
 							}
+
+							// Save parent material otherwise can't restore original due to material dynamic being transient
+							UMaterialInterface* Material = InMaterialHandle->GetMaterial();
+							Material = Material && Material->IsA<UMaterialInstanceDynamic>() ? Material->GetBaseMaterial() : Material;
+							InActorData.OriginalMaterials.Add(FAvaMask2DComponentMaterialPath(InComponent, InIdx), Material);
 						}
 
 						return true;
 					}
 				);
 
-				Handle->ForEachMaterial([&InActorData](const FSoftComponentReference& InComponent, const int32 InIdx, UMaterialInterface* InMaterial)
-				{
-					// Save parent material if material is transient otherwise can't restore original
-					UMaterialInterface* Material = InMaterial && InMaterial->HasAnyFlags(RF_Transient) ? InMaterial->GetBaseMaterial() : InMaterial;
-					InActorData.OriginalMaterials.Add(FAvaMask2DComponentMaterialPath(InComponent, InIdx), Material);
-					return true;
-				});
-
 				if (!Handle->OnSourceMaterialsChanged().IsBoundToObject(this))
 				{
 					Handle->OnSourceMaterialsChanged().BindUObject(this, &UAvaMask2DBaseModifier::OnMaterialsChanged);
diff --git a/Engine/Plugins/VirtualProduction/Avalanche/Source/AvalancheMask/Private/Mask2D/AvaMask2DReadModifier.cpp b/Engine/Plugins/VirtualProduction/Avalanche/Source/AvalancheMask/Private/Mask2D/AvaMask2DReadModifier.cpp
--- a/Engine/Plugins/VirtualProduction/Avalanche/Source/AvalancheMask/Private/Mask2D/AvaMask2DReadModifier.cpp
+++ b/Engine/Plugins/VirtualProduction/Avalanche/Source/AvalancheMask/Private/Mask2D/AvaMask2DReadModifier.cpp
@@ -17,6 +17,10 @@
 #include "Materials/MaterialInstanceDynamic.h"
 #include "Modifiers/ActorModifierCoreStack.h"
 
+#if WITH_EDITOR
+#include "Editor.h"
+#endif
+
 #define LOCTEXT_NAMESPACE "AvaMask2DReadModifier"
 
 #if WITH_EDITOR
@@ -51,6 +55,12 @@ void UAvaMask2DReadModifier::OnMaterialCompiled(UMaterialInterface* InMaterial)
 		MarkModifierDirty();
 	}
 }
+
+void UAvaMask2DReadModifier::OnPIEEnded(bool bInIsSimulating)
+{
+	// Used to refresh mask canvas
+	MarkModifierDirty();
+}
 #endif
 
 void UAvaMask2DReadModifier::Apply()
@@ -172,6 +182,7 @@ void UAvaMask2DReadModifier::OnModifierEnabled(EActorModifierCoreEnableReason In
 
 #if WITH_EDITOR
 	UMaterial::OnMaterialCompilationFinished().AddUObject(this, &UAvaMask2DReadModifier::OnMaterialCompiled);
+	FEditorDelegates::EndPIE.AddUObject(this, &UAvaMask2DReadModifier::OnPIEEnded);
 #endif
 }
 
@@ -181,6 +192,7 @@ void UAvaMask2DReadModifier::OnModifierDisabled(EActorModifierCoreDisableReason
 
 #if WITH_EDITOR
 	UMaterial::OnMaterialCompilationFinished().RemoveAll(this);
+	FEditorDelegates::EndPIE.RemoveAll(this);
 #endif
 }
 
diff --git a/Engine/Plugins/VirtualProduction/Avalanche/Source/AvalancheMask/Public/Mask2D/AvaMask2DReadModifier.h b/Engine/Plugins/VirtualProduction/Avalanche/Source/AvalancheMask/Public/Mask2D/AvaMask2DReadModifier.h
--- a/Engine/Plugins/VirtualProduction/Avalanche/Source/AvalancheMask/Public/Mask2D/AvaMask2DReadModifier.h
+++ b/Engine/Plugins/VirtualProduction/Avalanche/Source/AvalancheMask/Public/Mask2D/AvaMask2DReadModifier.h
@@ -77,5 +77,6 @@ private:
 	static const TAvaPropertyChangeDispatcher<UAvaMask2DReadModifier> PropertyChangeDispatcher;
 
 	void OnMaterialCompiled(UMaterialInterface* InMaterial);
+	void OnPIEEnded(bool bInIsSimulating);
 #endif
 };
