Motion Design: The spawn defaults buttons are now localized.

diff --git a/Engine/Plugins/VirtualProduction/Avalanche/Source/AvalancheEditor/Private/Scene/SAvaSceneDefaultActorResponses.cpp b/Engine/Plugins/VirtualProduction/Avalanche/Source/AvalancheEditor/Private/Scene/SAvaSceneDefaultActorResponses.cpp
--- a/Engine/Plugins/VirtualProduction/Avalanche/Source/AvalancheEditor/Private/Scene/SAvaSceneDefaultActorResponses.cpp
+++ b/Engine/Plugins/VirtualProduction/Avalanche/Source/AvalancheEditor/Private/Scene/SAvaSceneDefaultActorResponses.cpp
@@ -142,7 +142,7 @@ void SAvaSceneDefaultActorResponses::Construct(const FArguments& InArgs, UWorld*
 			.Padding(10.f, 10.f)
 			[
 				SNew(SButton)
-				.Text(INVTEXT("Spawn"))
+				.Text(LOCTEXT("Spawn", "Spawn"))
 				.ButtonStyle(FAppStyle::Get(), "PrimaryButton")
 				.OnClicked(this, &SAvaSceneDefaultActorResponses::OnOkayClick)
 				.IsEnabled(this, &SAvaSceneDefaultActorResponses::IsOkayEnabled)
@@ -152,7 +152,7 @@ void SAvaSceneDefaultActorResponses::Construct(const FArguments& InArgs, UWorld*
 			.Padding(10.f, 10.f)
 			[
 				SNew(SButton)
-				.Text(INVTEXT("Cancel"))
+				.Text(LOCTEXT("Cancel", "Cancel"))
 				.ButtonStyle(FAppStyle::Get(), "Button")
 				.OnClicked(this, &SAvaSceneDefaultActorResponses::OnCancelClick)
 			]
