Motion Design Broadcast: Fix data loss where uobject classes that are not found would cause JSON serialization to fail skipping to load the rest of the properties.

diff --git a/Engine/Plugins/VirtualProduction/Avalanche/Source/AvalancheMedia/Private/Broadcast/AvaBroadcastSerialization.cpp b/Engine/Plugins/VirtualProduction/Avalanche/Source/AvalancheMedia/Private/Broadcast/AvaBroadcastSerialization.cpp
--- a/Engine/Plugins/VirtualProduction/Avalanche/Source/AvalancheMedia/Private/Broadcast/AvaBroadcastSerialization.cpp
+++ b/Engine/Plugins/VirtualProduction/Avalanche/Source/AvalancheMedia/Private/Broadcast/AvaBroadcastSerialization.cpp
@@ -8,6 +8,7 @@
 #include "Formatters/XmlArchiveOutputFormatter.h"
 #include "JsonObjectConverter.h"
 #include "Misc/FileHelper.h"
+#include "Misc/PackageName.h"
 #include "Misc/Paths.h"
 
 namespace UE::AvaMedia::BroadcastSerialization::Private
@@ -145,6 +146,40 @@ namespace UE::AvaMedia::BroadcastSerialization::Private
 			break;
 		}
 	}
+
+	// Custom import callback for UObjects (i.e. primarily Media Outputs) so that if it fails, it can continue serializing the rest of the properties
+	bool OnCustomImportJsonObjectValue(const TSharedPtr<FJsonValue>& InJsonValue, FProperty* InProperty, void* InValue)
+	{
+		if (!InJsonValue.IsValid() || InJsonValue->Type != EJson::Object)
+		{
+			return false;
+		}
+
+		const FObjectProperty* const ObjectProperty = CastField<FObjectProperty>(InProperty);
+		if (!ObjectProperty)
+		{
+			return false;
+		}
+
+		const TSharedPtr<FJsonObject> Object = InJsonValue->AsObject();
+
+		const FString ClassString = Object->GetStringField(TEXT("_ClassName"));
+		if (ClassString.IsEmpty())
+		{
+			ObjectProperty->SetObjectPropertyValue(InValue, nullptr);
+			return true; // Logic overriden, the built-in logic won't be executed
+		}
+
+		const UClass* const FoundClass = FPackageName::IsShortPackageName(ClassString) ? FindFirstObject<UClass>(*ClassString) : LoadClass<UObject>(nullptr, *ClassString);
+		if (!FoundClass || FoundClass->HasAnyClassFlags(CLASS_Abstract))
+		{
+			ObjectProperty->SetObjectPropertyValue(InValue, nullptr);
+			return true; // Logic overriden, the built-in logic won't be executed
+		}
+
+		// Object should be valid, execute built-in logic
+		return false;
+	}
 }
 
 bool FAvaBroadcastSerialization::SaveBroadcastToJson(const UAvaBroadcast* InBroadcast, const FString& InFilename)
@@ -236,8 +271,14 @@ bool FAvaBroadcastSerialization::LoadBroadcastFromJson(const FString& InFilename
 		}
 		LoadReferencedObject(OutBroadcast, ReferenceObject);
 	}
-	
-	return FJsonObjectConverter::JsonObjectToUStruct(BroadcastJson.ToSharedRef(), OutBroadcast->GetClass(), OutBroadcast);;
+
+	FJsonObjectConverter::CustomImportCallback CustomImporter;
+	CustomImporter.BindStatic(&OnCustomImportJsonObjectValue);
+
+	constexpr int64 CheckFlags = 0;
+	constexpr int64 SkipFlags = 0;
+	constexpr bool bStrictMode = false;
+	return FJsonObjectConverter::JsonObjectToUStruct(BroadcastJson.ToSharedRef(), OutBroadcast->GetClass(), OutBroadcast, CheckFlags, SkipFlags, bStrictMode, /*FailReason*/nullptr, &CustomImporter);
 }
 
 bool FAvaBroadcastSerialization::SaveBroadcastToXml(UAvaBroadcast* InBroadcast, const FString& InFilename)
diff --git a/Engine/Plugins/VirtualProduction/Avalanche/Source/AvalancheMedia/Private/Broadcast/Channel/AvaBroadcastOutputChannel.cpp b/Engine/Plugins/VirtualProduction/Avalanche/Source/AvalancheMedia/Private/Broadcast/Channel/AvaBroadcastOutputChannel.cpp
--- a/Engine/Plugins/VirtualProduction/Avalanche/Source/AvalancheMedia/Private/Broadcast/Channel/AvaBroadcastOutputChannel.cpp
+++ b/Engine/Plugins/VirtualProduction/Avalanche/Source/AvalancheMedia/Private/Broadcast/Channel/AvaBroadcastOutputChannel.cpp
@@ -221,7 +221,10 @@ FAvaBroadcastOutputChannel::~FAvaBroadcastOutputChannel()
 #if WITH_EDITOR
 	for (UMediaOutput* MediaOutput : MediaOutputs)
 	{
-		MediaOutput->OnOutputModified().RemoveAll(this);
+		if (MediaOutput)
+		{
+			MediaOutput->OnOutputModified().RemoveAll(this);
+		}
 	}
 #endif
 }
