MotionDesign : Mask
- Fixed Masked text material keeps resetting when level is restarted
- Fixed mask warning about invalid or null object provided to MakeHandle

diff --git a/Engine/Plugins/VirtualProduction/Avalanche/Source/AvalancheMask/Private/Handling/AvaMaskDesignedMaterialHandle.h b/Engine/Plugins/VirtualProduction/Avalanche/Source/AvalancheMask/Private/Handling/AvaMaskDesignedMaterialHandle.h
--- a/Engine/Plugins/VirtualProduction/Avalanche/Source/AvalancheMask/Private/Handling/AvaMaskDesignedMaterialHandle.h
+++ b/Engine/Plugins/VirtualProduction/Avalanche/Source/AvalancheMask/Private/Handling/AvaMaskDesignedMaterialHandle.h
@@ -59,6 +59,7 @@ public:
 	virtual void CopyParametersFrom(UMaterialInstance* InSourceMaterial) override { Super::CopyParametersFrom(InSourceMaterial); }
 	virtual bool HasRequiredParameters(TArray<FString>& OutMissingParameterNames) override;
 	virtual void SetParentMaterial(UMaterialInterface* InMaterial) override { Super::SetParentMaterial(InMaterial); }
+	virtual UMaterialInterface* GetParentMaterial() override { return Super::GetParentMaterial(); }
 	// ~End IAvaMaterialHandle
 
 	// ~Begin IAvaObjectHandle
diff --git a/Engine/Plugins/VirtualProduction/Avalanche/Source/AvalancheMask/Private/Handling/AvaMaskMaterialInstanceHandle.h b/Engine/Plugins/VirtualProduction/Avalanche/Source/AvalancheMask/Private/Handling/AvaMaskMaterialInstanceHandle.h
--- a/Engine/Plugins/VirtualProduction/Avalanche/Source/AvalancheMask/Private/Handling/AvaMaskMaterialInstanceHandle.h
+++ b/Engine/Plugins/VirtualProduction/Avalanche/Source/AvalancheMask/Private/Handling/AvaMaskMaterialInstanceHandle.h
@@ -58,6 +58,7 @@ public:
 	virtual UMaterialInterface* GetMaterial() override { return FAvaMaterialInstanceHandle::GetMaterial(); }
 	virtual void CopyParametersFrom(UMaterialInstance* InSourceMaterial) override { FAvaMaterialInstanceHandle::CopyParametersFrom(InSourceMaterial); }
 	virtual void SetParentMaterial(UMaterialInterface* InMaterial) override { FAvaMaterialInstanceHandle::SetParentMaterial(InMaterial); }
+	virtual UMaterialInterface* GetParentMaterial() override { return FAvaMaterialInstanceHandle::GetParentMaterial(); }
 	// ~End IAvaMaterialHandle
 
 	// ~Begin IAvaObjectHandle	
diff --git a/Engine/Plugins/VirtualProduction/Avalanche/Source/AvalancheMask/Private/Handling/AvaMaskParametricMaterialHandle.h b/Engine/Plugins/VirtualProduction/Avalanche/Source/AvalancheMask/Private/Handling/AvaMaskParametricMaterialHandle.h
--- a/Engine/Plugins/VirtualProduction/Avalanche/Source/AvalancheMask/Private/Handling/AvaMaskParametricMaterialHandle.h
+++ b/Engine/Plugins/VirtualProduction/Avalanche/Source/AvalancheMask/Private/Handling/AvaMaskParametricMaterialHandle.h
@@ -54,6 +54,7 @@ public:
 	virtual UMaterialInterface* GetMaterial() override { return FAvaParametricMaterialHandle::GetMaterial(); }
 	virtual void CopyParametersFrom(UMaterialInstance* InSourceMaterial) override { FAvaParametricMaterialHandle::CopyParametersFrom(InSourceMaterial); }
 	virtual void SetParentMaterial(UMaterialInterface* InMaterial) override { FAvaParametricMaterialHandle::SetParentMaterial(InMaterial); }
+	virtual UMaterialInterface* GetParentMaterial() override { return FAvaParametricMaterialHandle::GetParentMaterial(); }
 	// ~End IAvaMaterialHandle
 
 	// ~Begin IAvaObjectHandle
diff --git a/Engine/Plugins/VirtualProduction/Avalanche/Source/AvalancheMask/Private/Handling/AvaMaskText3DActorMaterialCollectionHandle.cpp b/Engine/Plugins/VirtualProduction/Avalanche/Source/AvalancheMask/Private/Handling/AvaMaskText3DActorMaterialCollectionHandle.cpp
--- a/Engine/Plugins/VirtualProduction/Avalanche/Source/AvalancheMask/Private/Handling/AvaMaskText3DActorMaterialCollectionHandle.cpp
+++ b/Engine/Plugins/VirtualProduction/Avalanche/Source/AvalancheMask/Private/Handling/AvaMaskText3DActorMaterialCollectionHandle.cpp
@@ -12,6 +12,7 @@
 #include "Handling/AvaHandleUtilities.h"
 #include "IAvaMaskMaterialCollectionHandle.h"
 #include "IAvaMaskMaterialHandle.h"
+#include "Materials/Material.h"
 #include "Materials/MaterialInterface.h"
 #include "Text3DComponent.h"
 
@@ -94,8 +95,19 @@ TArray<TSharedPtr<IAvaMaskMaterialHandle>> FAvaMaskText3DActorMaterialCollection
 
 			MaterialHandles.SetNum(MaterialCount);
 			int32 Index = 0;
-			MaterialExtension->ForEachMaterial([this, HandleSubsystem, &Index](const UE::Text3D::Material::FMaterialParameters& InParameters, UMaterialInterface* InMaterial)->bool
+			UMaterial* DefaultMaterial = nullptr;
+			MaterialExtension->ForEachMaterial([this, HandleSubsystem, &Index, &DefaultMaterial](const UE::Text3D::Material::FMaterialParameters& InParameters, UMaterialInterface* InMaterial)->bool
 			{
+				if (!InMaterial)
+				{
+					if (!DefaultMaterial)
+					{
+						DefaultMaterial = LoadObject<UMaterial>(nullptr, TEXT("/Script/Engine.Material'/Engine/EngineMaterials/WorldGridMaterial.WorldGridMaterial'"));
+					}
+
+					InMaterial = DefaultMaterial;		
+				}
+
 				if (!MaterialHandles[Index] || MaterialHandles[Index]->GetMaterial() != InMaterial)
 				{
 					MaterialHandles[Index] = HandleSubsystem->MakeHandle<IAvaMaskMaterialHandle>(InMaterial, UE::AvaMask::Internal::HandleTag);
@@ -548,7 +560,7 @@ void FAvaMaskText3DActorMaterialCollectionHandle::OnTextPostUpdate(UText3DCompon
 				UMaterialInterface* Material = MaterialHandles[Index]->GetMaterial();
 				if (Material != InMaterial)
 				{
-					if (!Material || Material->GetBaseMaterial() != InMaterial->GetBaseMaterial())
+					if (!Material || !InMaterial || Material->GetBaseMaterial() != InMaterial->GetBaseMaterial())
 					{
 						bMatch = false;
 					}
diff --git a/Engine/Plugins/VirtualProduction/Avalanche/Source/AvalancheMask/Private/Handling/AvaMaterialInstanceHandle.cpp b/Engine/Plugins/VirtualProduction/Avalanche/Source/AvalancheMask/Private/Handling/AvaMaterialInstanceHandle.cpp
--- a/Engine/Plugins/VirtualProduction/Avalanche/Source/AvalancheMask/Private/Handling/AvaMaterialInstanceHandle.cpp
+++ b/Engine/Plugins/VirtualProduction/Avalanche/Source/AvalancheMask/Private/Handling/AvaMaterialInstanceHandle.cpp
@@ -36,6 +36,21 @@ UMaterialInterface* FAvaMaterialInstanceHandle::GetMaterial()
 	return WeakParentMaterial.Get();
 }
 
+UMaterialInterface* FAvaMaterialInstanceHandle::GetParentMaterial()
+{
+	if (UMaterialInterface* ParentMaterial = WeakParentMaterial.Get())
+	{
+		return ParentMaterial;
+	}
+
+	if (UMaterialInstanceDynamic* MID = WeakMaterialInstance.Get())
+	{
+		return MID->GetBaseMaterial();
+	}
+
+	return nullptr;
+}
+
 void FAvaMaterialInstanceHandle::CopyParametersFrom(UMaterialInstance* InSourceMaterial)
 {
 	if (UMaterialInstanceDynamic* MaterialInstance = GetMaterialInstance())
@@ -60,21 +75,6 @@ bool FAvaMaterialInstanceHandle::IsSupported(const FAvaMaskMaterialReference& In
 	return !!InInstance.GetTypedObject<UMaterialInterface>();
 }
 
-UMaterialInterface* FAvaMaterialInstanceHandle::GetParentMaterial() const
-{
-	if (UMaterialInterface* ParentMaterial = WeakParentMaterial.Get())
-	{
-		return ParentMaterial;
-	}
-
-	if (UMaterialInstanceDynamic* MID = WeakMaterialInstance.Get())
-	{
-		return MID->GetBaseMaterial();
-	}
-
-	return nullptr;
-}
-
 UMaterialInstanceDynamic* FAvaMaterialInstanceHandle::GetMaterialInstance()
 {
 	if (UMaterialInstanceDynamic* MaterialInstance = WeakMaterialInstance.Get())
diff --git a/Engine/Plugins/VirtualProduction/Avalanche/Source/AvalancheMask/Private/Handling/AvaMaterialInstanceHandle.h b/Engine/Plugins/VirtualProduction/Avalanche/Source/AvalancheMask/Private/Handling/AvaMaterialInstanceHandle.h
--- a/Engine/Plugins/VirtualProduction/Avalanche/Source/AvalancheMask/Private/Handling/AvaMaterialInstanceHandle.h
+++ b/Engine/Plugins/VirtualProduction/Avalanche/Source/AvalancheMask/Private/Handling/AvaMaterialInstanceHandle.h
@@ -28,6 +28,7 @@ public:
 	// ~Begin IAvaMaterialHandle
 	virtual FString GetMaterialName() override;
 	virtual UMaterialInterface* GetMaterial() override;
+	virtual UMaterialInterface* GetParentMaterial() override;
 	virtual void CopyParametersFrom(UMaterialInstance* InSourceMaterial) override;
 	virtual void SetParentMaterial(UMaterialInterface* InMaterial) override;
 	// ~End IAvaMaterialHandle
@@ -37,7 +38,6 @@ public:
 	// ~End IAvaObjectHandle
 
 protected:
-	virtual UMaterialInterface* GetParentMaterial() const;
 	virtual UMaterialInstanceDynamic* GetMaterialInstance();
 	virtual UMaterialInstanceDynamic* GetOrCreateMaterialInstance();
 
diff --git a/Engine/Plugins/VirtualProduction/Avalanche/Source/AvalancheMask/Private/Handling/IAvaMaterialHandle.h b/Engine/Plugins/VirtualProduction/Avalanche/Source/AvalancheMask/Private/Handling/IAvaMaterialHandle.h
--- a/Engine/Plugins/VirtualProduction/Avalanche/Source/AvalancheMask/Private/Handling/IAvaMaterialHandle.h
+++ b/Engine/Plugins/VirtualProduction/Avalanche/Source/AvalancheMask/Private/Handling/IAvaMaterialHandle.h
@@ -18,6 +18,7 @@ public:
 
 	virtual FString GetMaterialName() = 0;
 	virtual UMaterialInterface* GetMaterial() = 0;
+	virtual UMaterialInterface* GetParentMaterial() = 0;
 	virtual void SetParentMaterial(UMaterialInterface* InMaterial) = 0;
 
 	virtual void CopyParametersFrom(UMaterialInstance* InSourceMaterial) = 0;
diff --git a/Engine/Plugins/VirtualProduction/Avalanche/Source/AvalancheMask/Private/Mask2D/AvaMask2DBaseModifier.cpp b/Engine/Plugins/VirtualProduction/Avalanche/Source/AvalancheMask/Private/Mask2D/AvaMask2DBaseModifier.cpp
--- a/Engine/Plugins/VirtualProduction/Avalanche/Source/AvalancheMask/Private/Mask2D/AvaMask2DBaseModifier.cpp
+++ b/Engine/Plugins/VirtualProduction/Avalanche/Source/AvalancheMask/Private/Mask2D/AvaMask2DBaseModifier.cpp
@@ -438,8 +438,7 @@ void UAvaMask2DBaseModifier::SaveActorPreState(AActor* InActor, FAvaMask2DActorD
 							}
 
 							// Save parent material otherwise can't restore original due to material dynamic being transient
-							UMaterialInterface* Material = InMaterialHandle->GetMaterial();
-							Material = Material && Material->IsA<UMaterialInstanceDynamic>() ? Material->GetBaseMaterial() : Material;
+							UMaterialInterface* Material = InMaterialHandle->GetParentMaterial();
 							InActorData.OriginalMaterials.Add(FAvaMask2DComponentMaterialPath(InComponent, InIdx), Material);
 						}
 
@@ -910,13 +909,21 @@ void UAvaMask2DBaseModifier::OnMaterialsChanged(UObject* InMaterialOwner, const
 {
 	if (!InMaterialHandles.IsEmpty())
 	{
-		if (const AActor* OwningActor = InMaterialOwner->GetTypedOuter<AActor>())
+		if (AActor* OwningActor = InMaterialOwner->GetTypedOuter<AActor>())
 		{
 			const UActorModifierCoreBase* MaskModifier = UE::ActorModifierCore::Utilities::FindFirstActorModifierByClass(OwningActor, UAvaMask2DBaseModifier::StaticClass());
 
 			// Only trigger update if first modifier found in the hierarchy above us is this modifier
-			if (MaskModifier == this)
+			if (MaskModifier == this && ActorData.Contains(OwningActor))
 			{
+				FInstancedStruct* HandleData = MaterialCollectionHandleData.Find(OwningActor);
+				const TSharedPtr<IAvaMaskMaterialCollectionHandle>* CollectionHandle = MaterialCollectionHandles.Find(OwningActor);
+
+				if (HandleData && CollectionHandle)
+				{
+					(*CollectionHandle)->SaveOriginalState(*HandleData);
+				}
+
 				MarkModifierDirty();
 			}
 		}
