Motion Design: added ticker section bubble for ticker properties

diff --git a/Engine/Plugins/VirtualProduction/Avalanche/Source/AvalancheEditor/Private/AvaEditorModule.cpp b/Engine/Plugins/VirtualProduction/Avalanche/Source/AvalancheEditor/Private/AvaEditorModule.cpp
--- a/Engine/Plugins/VirtualProduction/Avalanche/Source/AvalancheEditor/Private/AvaEditorModule.cpp
+++ b/Engine/Plugins/VirtualProduction/Avalanche/Source/AvalancheEditor/Private/AvaEditorModule.cpp
@@ -276,6 +276,12 @@ void FAvaEditorModule::RegisterPropertyEditorCategories()
 		Section->AddCategory("CameraOptions");
 		Section->AddCategory("CurrentCameraSettings");
 	}
+
+	// AvaTicker Sections
+	{
+		Section = PropertyModule.FindOrCreateSection("AvaTickerComponent", "Ticker", LOCTEXT("Ticker", "Ticker"));
+		Section->AddCategory("Ticker");
+	}
 }
 
 void FAvaEditorModule::RegisterCustomLayouts()
