Motion Design: Remove compatibility rule on translucent priority

diff --git a/Engine/Plugins/VirtualProduction/Avalanche/Source/AvalancheModifiers/Private/Modifiers/AvaTranslucentPriorityModifier.cpp b/Engine/Plugins/VirtualProduction/Avalanche/Source/AvalancheModifiers/Private/Modifiers/AvaTranslucentPriorityModifier.cpp
--- a/Engine/Plugins/VirtualProduction/Avalanche/Source/AvalancheModifiers/Private/Modifiers/AvaTranslucentPriorityModifier.cpp
+++ b/Engine/Plugins/VirtualProduction/Avalanche/Source/AvalancheModifiers/Private/Modifiers/AvaTranslucentPriorityModifier.cpp
@@ -31,30 +31,6 @@ void UAvaTranslucentPriorityModifier::OnModifierCDOSetup(FActorModifierCoreMetad
 	InMetadata.SetDisplayName(LOCTEXT("ModifierDisplayName", "Translucent Sort Priority"));
 	InMetadata.SetDescription(LOCTEXT("ModifierDescription", "Set the translucency sort priority of a primitive component based on different modes"));
 #endif
-
-	InMetadata.SetCompatibilityRule([this](const AActor* InActor)
-	{
-		if (!InActor)
-		{
-			return false;
-		}
-
-		const bool bResult = ForEachComponent<UPrimitiveComponent>([](const UPrimitiveComponent* InComponent)->bool
-		{
-			/** Stop when we have a primitive component otherwise keep going */
-			if (InComponent)
-			{
-				return false;
-			}
-
-			return true;
-		}
-		, EActorModifierCoreComponentType::All
-		, EActorModifierCoreLookup::SelfAndAllChildren
-		, InActor);
-
-		return !bResult;
-	});
 }
 
 void UAvaTranslucentPriorityModifier::OnModifierAdded(EActorModifierCoreEnableReason InReason)
