Motion Design: Fixed issue where translucent priority modifier when not finding an actor item would put these items at the back of the sorted list. Instead, if an item is not found (e.g. an actor spawned at runtime) it will try finding its item, and if fails to do so, tries to find its place in the tree by looking for its attach parent item.

diff --git a/Engine/Plugins/VirtualProduction/Avalanche/Source/AvalancheModifiers/Private/Shared/AvaTranslucentPriorityModifierShared.cpp b/Engine/Plugins/VirtualProduction/Avalanche/Source/AvalancheModifiers/Private/Shared/AvaTranslucentPriorityModifierShared.cpp
--- a/Engine/Plugins/VirtualProduction/Avalanche/Source/AvalancheModifiers/Private/Shared/AvaTranslucentPriorityModifierShared.cpp
+++ b/Engine/Plugins/VirtualProduction/Avalanche/Source/AvalancheModifiers/Private/Shared/AvaTranslucentPriorityModifierShared.cpp
@@ -250,11 +250,22 @@ TArray<FAvaTranslucentPriorityModifierComponentState> UAvaTranslucentPriorityMod
 		TSharedPtr<IAvaOutliner> AvaOutliner = FAvaOutlinerUtils::EditorGetOutliner(World);
 		if (AvaOutliner.IsValid())
 		{
-			SortedComponentStates.StableSort([AvaOutliner, bInvertedTree](const FAvaTranslucentPriorityModifierComponentState& InComponentA, const FAvaTranslucentPriorityModifierComponentState& InComponentB)->bool
+			auto FindItem =
+				[&AvaOutliner, World](AActor* InActor)->FAvaOutlinerItemPtr
+				{
+					FAvaOutlinerItemPtr Item;
+					while (InActor && !Item.IsValid())
+					{
+						Item = AvaOutliner->FindItem(InActor);
+						InActor = InActor->GetAttachParentActor();
+					}
+					return Item;
+				};
+
+			SortedComponentStates.StableSort([&FindItem, bInvertedTree](const FAvaTranslucentPriorityModifierComponentState& InComponentA, const FAvaTranslucentPriorityModifierComponentState& InComponentB)->bool
 			{
-				const FAvaOutlinerItemPtr OutlinerItemA = AvaOutliner->FindItem(InComponentA.GetOwningActor());
-				const FAvaOutlinerItemPtr OutlinerItemB = AvaOutliner->FindItem(InComponentB.GetOwningActor());
-
+				const FAvaOutlinerItemPtr OutlinerItemA = FindItem(InComponentA.GetOwningActor());
+				const FAvaOutlinerItemPtr OutlinerItemB = FindItem(InComponentB.GetOwningActor());
 				return UE::AvaOutliner::CompareOutlinerItemOrder(OutlinerItemB, OutlinerItemA) ^ bInvertedTree;
 			});
 		}
@@ -266,11 +277,22 @@ TArray<FAvaTranslucentPriorityModifierComponentState> UAvaTranslucentPriorityMod
 			{
 				const FAvaSceneTree& SceneTree = SceneInterface->GetSceneTree();
 
-				SortedComponentStates.StableSort([&SceneTree, World, bInvertedTree](const FAvaTranslucentPriorityModifierComponentState& InComponentA, const FAvaTranslucentPriorityModifierComponentState& InComponentB)->bool
+				auto FindTreeNode =
+					[&SceneTree, World](AActor* InActor)->const FAvaSceneTreeNode*
+					{
+						const FAvaSceneTreeNode* TreeNode = nullptr;
+						while (InActor && !TreeNode)
+						{
+							TreeNode = SceneTree.FindTreeNode(FAvaSceneItem(InActor, World));
+							InActor = InActor->GetAttachParentActor();
+						}
+						return TreeNode;
+					};
+
+				SortedComponentStates.StableSort([&FindTreeNode, bInvertedTree](const FAvaTranslucentPriorityModifierComponentState& InComponentA, const FAvaTranslucentPriorityModifierComponentState& InComponentB)->bool
 				{
-					const FAvaSceneTreeNode* SceneItemA = SceneTree.FindTreeNode(FAvaSceneItem(InComponentA.GetOwningActor(), World));
-					const FAvaSceneTreeNode* SceneItemB = SceneTree.FindTreeNode(FAvaSceneItem(InComponentB.GetOwningActor(), World));
-
+					const FAvaSceneTreeNode* SceneItemA = FindTreeNode(InComponentA.GetOwningActor());
+					const FAvaSceneTreeNode* SceneItemB = FindTreeNode(InComponentB.GetOwningActor());
 					return FAvaSceneTree::CompareTreeItemOrder(SceneItemB, SceneItemA) ^ bInvertedTree;
 				});
 			}
