Motion Design: fix issue where mask read failing on one actors bails out completely and doesn't apply read to the rest of the actors

diff --git a/Engine/Plugins/VirtualProduction/Avalanche/Source/AvalancheMask/Private/Mask2D/AvaMask2DReadModifier.cpp b/Engine/Plugins/VirtualProduction/Avalanche/Source/AvalancheMask/Private/Mask2D/AvaMask2DReadModifier.cpp
--- a/Engine/Plugins/VirtualProduction/Avalanche/Source/AvalancheMask/Private/Mask2D/AvaMask2DReadModifier.cpp
+++ b/Engine/Plugins/VirtualProduction/Avalanche/Source/AvalancheMask/Private/Mask2D/AvaMask2DReadModifier.cpp
@@ -67,21 +67,30 @@ void UAvaMask2DReadModifier::Apply()
 {
 	Super::Apply();
 
+	TArray<FText> FailReasons;
 	if (const UGeometryMaskCanvas* Canvas = GetCurrentCanvas())
 	{
 		for (TPair<TWeakObjectPtr<AActor>, FAvaMask2DActorData>& ActorDataPair : ActorData)
 		{
 			if (AActor* Actor = ActorDataPair.Key.Get())
 			{
-				if (!ApplyRead(Actor, ActorDataPair.Value))
+				FText FailReason;
+				if (!ApplyRead(Actor, ActorDataPair.Value, FailReason))
 				{
-					return;
+					FailReasons.Add(MoveTemp(FailReason));
 				}
 			}
 		}
 	}
 
-	Next();
+	if (FailReasons.IsEmpty())
+	{
+		Next();
+	}
+	else
+	{
+		Fail(FText::Join(FText::FromStringView(TEXT(", ")), FailReasons));
+	}
 }
 
 EBlendMode UAvaMask2DReadModifier::GetBlendMode() const
@@ -90,6 +99,12 @@ EBlendMode UAvaMask2DReadModifier::GetBlendMode() const
 }
 
 bool UAvaMask2DReadModifier::ApplyRead(AActor* InActor, FAvaMask2DActorData& InActorData)
+{
+	FText FailReason;
+	return ApplyRead(InActor, InActorData, FailReason);
+}
+
+bool UAvaMask2DReadModifier::ApplyRead(AActor* InActor, FAvaMask2DActorData& InActorData, FText& OutFailReason)
 {
 	// Only add read/write component to actors with primitives
 	if (!ActorSupportsMaskReadWrite(InActor))
@@ -151,10 +166,8 @@ bool UAvaMask2DReadModifier::ApplyRead(AActor* InActor, FAvaMask2DActorData& InA
 		}
 
 		// Only check for valid materials here, after ApplyModifiedState has occured once
-		FText FailReason;
-		if (!MaterialCollectionHandle->ValidateMaterials(FailReason))
+		if (!MaterialCollectionHandle->ValidateMaterials(OutFailReason))
 		{
-			Fail(FailReason);
 			return false;
 		}
 	}
diff --git a/Engine/Plugins/VirtualProduction/Avalanche/Source/AvalancheMask/Public/Mask2D/AvaMask2DReadModifier.h b/Engine/Plugins/VirtualProduction/Avalanche/Source/AvalancheMask/Public/Mask2D/AvaMask2DReadModifier.h
--- a/Engine/Plugins/VirtualProduction/Avalanche/Source/AvalancheMask/Public/Mask2D/AvaMask2DReadModifier.h
+++ b/Engine/Plugins/VirtualProduction/Avalanche/Source/AvalancheMask/Public/Mask2D/AvaMask2DReadModifier.h
@@ -65,8 +65,11 @@ protected:
 	virtual void SetupMaskComponent(UActorComponent* InComponent) override;
 	void SetupMaskReadComponent(IGeometryMaskReadInterface* InMaskReader);
 
+	UE_DEPRECATED(5.7, "Use ApplyRead that takes in a fail reason")
 	bool ApplyRead(AActor* InActor, FAvaMask2DActorData& InActorData);
 
+	bool ApplyRead(AActor* InActor, FAvaMask2DActorData& InActorData, FText& OutFailReason);
+
 	/** Base opacity/alpha to use in Read mode */
 	UPROPERTY(EditInstanceOnly, BlueprintReadWrite, Getter, Setter, Category = "Mask2D", meta = (ClampMin = 0.0, ClampMax = 1.0, AllowPrivateAccess = "true"))
 	float BaseOpacity = 0.0f;
