Motion Design Modifiers: Missing default struct values

diff --git a/Engine/Plugins/VirtualProduction/Avalanche/Source/AvalancheModifiers/Public/Modifiers/AvaMaterialParameterModifier.h b/Engine/Plugins/VirtualProduction/Avalanche/Source/AvalancheModifiers/Public/Modifiers/AvaMaterialParameterModifier.h
--- a/Engine/Plugins/VirtualProduction/Avalanche/Source/AvalancheModifiers/Public/Modifiers/AvaMaterialParameterModifier.h
+++ b/Engine/Plugins/VirtualProduction/Avalanche/Source/AvalancheModifiers/Public/Modifiers/AvaMaterialParameterModifier.h
@@ -16,10 +16,10 @@ struct FAvaMaterialParameterMapScalar
 	GENERATED_BODY()
 
 	UPROPERTY(EditInstanceOnly, BlueprintReadWrite, Category = "MaterialParameter")
-	FName Name;
+	FName Name = NAME_None;
 
 	UPROPERTY(EditInstanceOnly, BlueprintReadWrite, Category = "MaterialParameter")
-	float Value;
+	float Value = 0.f;
 };
 
 USTRUCT(BlueprintType)
@@ -28,10 +28,10 @@ struct FAvaMaterialParameterMapVector
 	GENERATED_BODY()
 
 	UPROPERTY(EditInstanceOnly, BlueprintReadWrite, Category = "MaterialParameter")
-	FName Name;
+	FName Name = NAME_None;
 	
 	UPROPERTY(EditInstanceOnly, BlueprintReadWrite, Category = "MaterialParameter")
-	FLinearColor Value;
+	FLinearColor Value = FLinearColor::Black;
 };
 
 USTRUCT(BlueprintType)
@@ -40,10 +40,10 @@ struct FAvaMaterialParameterMapTexture
 	GENERATED_BODY()
 
 	UPROPERTY(EditInstanceOnly, BlueprintReadWrite, Category = "MaterialParameter")
-	FName Name;
+	FName Name = NAME_None;
 
 	UPROPERTY(EditInstanceOnly, BlueprintReadWrite, Category = "MaterialParameter")
-	TObjectPtr<UTexture> Value;
+	TObjectPtr<UTexture> Value = nullptr;
 };
 
 USTRUCT(BlueprintType)
