MotionDesign : ActorModifier
- Fixed translucent sort priority incorrect increments causing components under the same actor to have different sorting values

diff --git a/Engine/Plugins/VirtualProduction/Avalanche/Source/AvalancheModifiers/Private/Modifiers/AvaTranslucentPriorityModifier.cpp b/Engine/Plugins/VirtualProduction/Avalanche/Source/AvalancheModifiers/Private/Modifiers/AvaTranslucentPriorityModifier.cpp
--- a/Engine/Plugins/VirtualProduction/Avalanche/Source/AvalancheModifiers/Private/Modifiers/AvaTranslucentPriorityModifier.cpp
+++ b/Engine/Plugins/VirtualProduction/Avalanche/Source/AvalancheModifiers/Private/Modifiers/AvaTranslucentPriorityModifier.cpp
@@ -188,6 +188,16 @@ void UAvaTranslucentPriorityModifier::Apply()
 				continue;
 			}
 
+			const bool bHasIteratedOnce = PreviousComponentOwner != nullptr;
+			const bool bShouldIncreaseStep = Mode == EAvaTranslucentPriorityModifierMode::AutoCameraDistance || PreviousComponentOwner != Component->GetOwner();
+
+			if (bHasIteratedOnce && bShouldIncreaseStep)
+			{
+				TranslucentSortPriority += GlobalStep;
+			}
+
+			PreviousComponentOwner = Component->GetOwner();
+
 			// This modifier handles this component
 			if (ComponentModifier == this)
 			{
@@ -205,15 +215,6 @@ void UAvaTranslucentPriorityModifier::Apply()
 
 				ComponentModifier->MarkModifierDirty();
 			}
-			
-			// Only increase when new actor detected or mode is camera
-			if (Mode == EAvaTranslucentPriorityModifierMode::AutoCameraDistance
-				|| PreviousComponentOwner != Component->GetOwner())
-			{
-				TranslucentSortPriority += GlobalStep;
-			}
-
-			PreviousComponentOwner = Component->GetOwner();
 		}
 	}
 
