MotionDesign : GeometryMask
- Fixed invalid/empty material slot on masked primitive components causing issues, this new empty handle will assign a MID based on default mask material on apply and restore the slot to none when original state is restored, this will allow primitive components with material slots set to none to still work with mask

diff --git a/Engine/Plugins/VirtualProduction/Avalanche/Source/AvalancheMask/Private/Handling/AvaMaskAvaShapeMaterialCollectionHandle.cpp b/Engine/Plugins/VirtualProduction/Avalanche/Source/AvalancheMask/Private/Handling/AvaMaskAvaShapeMaterialCollectionHandle.cpp
--- a/Engine/Plugins/VirtualProduction/Avalanche/Source/AvalancheMask/Private/Handling/AvaMaskAvaShapeMaterialCollectionHandle.cpp
+++ b/Engine/Plugins/VirtualProduction/Avalanche/Source/AvalancheMask/Private/Handling/AvaMaskAvaShapeMaterialCollectionHandle.cpp
@@ -393,13 +393,6 @@ bool FAvaMaskAvaShapeMaterialCollectionHandle::RefreshMaterials()
 						OccupiedMaterialSlots[InMeshIdx] = true;
 					}
 				}
-				else if (!Material)
-				{
-					// Use default material to avoid invalid handle
-					UMaterial* DefaultMaterial = LoadObject<UMaterial>(nullptr, TEXT("/Script/Engine.Material'/Engine/EngineMaterials/WorldGridMaterial.WorldGridMaterial'"));
-					MaterialHandle = MaterialHandles.Emplace(InMeshIdx, HandleSubsystem->MakeHandle<IAvaMaskMaterialHandle>(DefaultMaterial, UE::AvaMask::Internal::HandleTag));
-					OccupiedMaterialSlots[InMeshIdx] = true;
-				}
 				else
 				{
 					MaterialHandle = MaterialHandles.Emplace(InMeshIdx, HandleSubsystem->MakeHandle<IAvaMaskMaterialHandle>(Material, UE::AvaMask::Internal::HandleTag));
diff --git a/Engine/Plugins/VirtualProduction/Avalanche/Source/AvalancheMask/Private/Handling/AvaMaskEmptyMaterialHandle.cpp b/Engine/Plugins/VirtualProduction/Avalanche/Source/AvalancheMask/Private/Handling/AvaMaskEmptyMaterialHandle.cpp
new file mode 100644
--- /dev/null
+++ b/Engine/Plugins/VirtualProduction/Avalanche/Source/AvalancheMask/Private/Handling/AvaMaskEmptyMaterialHandle.cpp
@@ -0,0 +1,132 @@
+// Copyright Epic Games, Inc. All Rights Reserved.
+
+#include "Handling/AvaMaskEmptyMaterialHandle.h"
+
+#include "AvaMaskMaterialInstanceHandle.h"
+#include "AvaMaskMaterialReference.h"
+#include "Materials/Material.h"
+#include "Materials/MaterialInstanceDynamic.h"
+#include "UObject/Package.h"
+#include "UObject/UObjectGlobals.h"
+
+FAvaMaskEmptyMaterialHandle::FAvaMaskEmptyMaterialHandle()
+	: FAvaMaterialInstanceHandle(nullptr)
+{
+}
+
+const EBlendMode FAvaMaskEmptyMaterialHandle::GetBlendMode()
+{
+	if (const UMaterialInterface* Material = GetMaterial())
+	{
+		return Material->GetBlendMode();
+	}
+
+	return EBlendMode::BLEND_Translucent;
+}
+
+void FAvaMaskEmptyMaterialHandle::SetBlendMode(const EBlendMode InBlendMode)
+{
+	const EBlendMode TargetBlendMode = UE::AvaMask::Internal::GetTargetBlendMode(GetBlendMode(), InBlendMode);
+	if (UMaterialInstanceDynamic* MaterialInstance = GetMaterialInstance())
+	{
+		MaterialInstance->BasePropertyOverrides.bOverride_BlendMode = true;
+		MaterialInstance->BasePropertyOverrides.BlendMode = TargetBlendMode;
+		MaterialInstance->BlendMode = TargetBlendMode;
+#if WITH_EDITOR
+		MaterialInstance->UpdateCachedData();
+#endif
+	}
+}
+
+void FAvaMaskEmptyMaterialHandle::SetChildMaterial(UMaterialInstanceDynamic* InChildMaterial)
+{
+	// do nothing
+}
+
+bool FAvaMaskEmptyMaterialHandle::SaveOriginalState(const FStructView& InHandleData)
+{
+	// do nothing
+	return true;
+}
+
+bool FAvaMaskEmptyMaterialHandle::ApplyOriginalState(const FStructView& InHandleData, TUniqueFunction<void(UMaterialInterface*)>&& InMaterialSetter)
+{
+	InMaterialSetter(nullptr);
+	return true;
+}
+
+bool FAvaMaskEmptyMaterialHandle::ApplyModifiedState(const FAvaMask2DSubjectParameters& InModifiedParameters, const FStructView& InHandleData,
+	TUniqueFunction<void(UMaterialInterface*)>&& InMaterialSetter)
+{
+	if (UMaterialInterface* MaterialInstance = GetOrCreateMaterialInstance())
+	{
+		SetBlendMode(InModifiedParameters.MaterialParameters.BlendMode);
+		SetMaskParameters(
+			InModifiedParameters.MaterialParameters.Texture
+			, InModifiedParameters.MaterialParameters.BaseOpacity
+			, InModifiedParameters.MaterialParameters.Channel
+			, InModifiedParameters.MaterialParameters.bInvert
+		);
+		InMaterialSetter(MaterialInstance);
+	}
+
+	return true;
+}
+
+FString FAvaMaskEmptyMaterialHandle::GetMaterialName()
+{
+	return FAvaMaterialInstanceHandle::GetMaterialName();
+}
+
+UMaterialInterface* FAvaMaskEmptyMaterialHandle::GetMaterial()
+{
+	return FAvaMaterialInstanceHandle::GetMaterial();
+}
+
+void FAvaMaskEmptyMaterialHandle::CopyParametersFrom(UMaterialInstance* InSourceMaterial)
+{
+	FAvaMaterialInstanceHandle::CopyParametersFrom(InSourceMaterial);
+}
+
+void FAvaMaskEmptyMaterialHandle::SetParentMaterial(UMaterialInterface* InMaterial)
+{
+	// do nothing
+}
+
+UMaterialInterface* FAvaMaskEmptyMaterialHandle::GetParentMaterial()
+{
+	return nullptr;
+}
+
+bool FAvaMaskEmptyMaterialHandle::IsValid() const
+{
+	return true;
+}
+
+bool FAvaMaskEmptyMaterialHandle::IsSupported(const FAvaMaskMaterialReference& InInstance, FName InTag)
+{
+	return InTag == UE::AvaMask::Internal::HandleTag && InInstance.GetObject() == nullptr;
+}
+
+UMaterialInstanceDynamic* FAvaMaskEmptyMaterialHandle::GetMaterialInstance()
+{
+	return FAvaMaterialInstanceHandle::GetMaterialInstance();
+}
+
+UMaterialInstanceDynamic* FAvaMaskEmptyMaterialHandle::GetOrCreateMaterialInstance()
+{
+	if (UMaterialInstanceDynamic* MaterialInstance = GetMaterialInstance())
+	{
+		return MaterialInstance;
+	}
+
+	UMaterialInstanceDynamic* MaterialInstance = nullptr;
+	if (UMaterial* DefaultMaskMaterial = LoadObject<UMaterial>(nullptr, TEXT("/Script/Engine.Material'/Avalanche/MaskResources/M_AvalancheMaskDefaultTranslucent.M_AvalancheMaskDefaultTranslucent'")))
+	{
+		UObject* Outer = GetTransientPackage();
+		MaterialInstance = UMaterialInstanceDynamic::Create(DefaultMaskMaterial, Outer);
+		WeakMaterialInstance = MaterialInstance;
+	}
+
+	return MaterialInstance;
+}
diff --git a/Engine/Plugins/VirtualProduction/Avalanche/Source/AvalancheMask/Private/Handling/AvaMaskEmptyMaterialHandle.h b/Engine/Plugins/VirtualProduction/Avalanche/Source/AvalancheMask/Private/Handling/AvaMaskEmptyMaterialHandle.h
new file mode 100644
--- /dev/null
+++ b/Engine/Plugins/VirtualProduction/Avalanche/Source/AvalancheMask/Private/Handling/AvaMaskEmptyMaterialHandle.h
@@ -0,0 +1,67 @@
+// Copyright Epic Games, Inc. All Rights Reserved.
+
+#pragma once
+
+#include "Engine/EngineTypes.h"
+#include "Handling/AvaMaterialInstanceHandle.h"
+#include "IAvaMaskMaterialHandle.h"
+#include "Materials/MaterialParameters.h"
+
+#include "AvaMaskEmptyMaterialHandle.generated.h"
+
+class UMaterialInstanceDynamic;
+class UMaterialInterface;
+struct FAvaMaskMaterialReference;
+
+USTRUCT()
+struct FAvaMaskEmptyMaterialHandleData
+{
+	GENERATED_BODY()
+};
+
+/**
+ * This handle is used for primitive slots that have no materials "None"
+ * It will allow mask to apply properly and restore material to "None" 
+ */
+class FAvaMaskEmptyMaterialHandle
+	: public FAvaMaterialInstanceHandle
+	, public TAvaMaskMaterialHandle<FAvaMaskEmptyMaterialHandleData>
+{
+public:
+	UE_AVA_INHERITS_WITH_SUPER(FAvaMaskEmptyMaterialHandle, FAvaMaterialInstanceHandle);
+
+	static bool IsSupported(const FAvaMaskMaterialReference& InInstance, FName InTag = NAME_None);
+
+	explicit FAvaMaskEmptyMaterialHandle();
+	virtual ~FAvaMaskEmptyMaterialHandle() override = default;
+
+	//~ Begin IAvaMaskMaterialHandle
+	virtual const EBlendMode GetBlendMode() override;
+	virtual void SetBlendMode(const EBlendMode InBlendMode) override;
+	virtual void SetChildMaterial(UMaterialInstanceDynamic* InChildMaterial) override;
+	virtual bool SaveOriginalState(const FStructView& InHandleData) override;
+	virtual bool ApplyOriginalState(
+		const FStructView& InHandleData
+		, TUniqueFunction<void(UMaterialInterface*)>&& InMaterialSetter) override;
+	virtual bool ApplyModifiedState(
+		const FAvaMask2DSubjectParameters& InModifiedParameters
+		, const FStructView& InHandleData
+		, TUniqueFunction<void(UMaterialInterface*)>&& InMaterialSetter) override;
+	//~ End IAvaMaskMaterialHandle
+
+	//~ Begin IAvaMaterialHandle
+	virtual FString GetMaterialName() override;
+	virtual UMaterialInterface* GetMaterial() override;
+	virtual void CopyParametersFrom(UMaterialInstance* InSourceMaterial) override;
+	virtual void SetParentMaterial(UMaterialInterface* InMaterial) override;
+	virtual UMaterialInterface* GetParentMaterial() override;
+	//~ End IAvaMaterialHandle
+
+	//~ Begin IAvaObjectHandle	
+	virtual bool IsValid() const override;
+	//~ End IAvaObjectHandle
+
+protected:
+	virtual UMaterialInstanceDynamic* GetMaterialInstance() override;
+	virtual UMaterialInstanceDynamic* GetOrCreateMaterialInstance() override;
+};
diff --git a/Engine/Plugins/VirtualProduction/Avalanche/Source/AvalancheMask/Private/Handling/AvaMaskText3DActorMaterialCollectionHandle.cpp b/Engine/Plugins/VirtualProduction/Avalanche/Source/AvalancheMask/Private/Handling/AvaMaskText3DActorMaterialCollectionHandle.cpp
--- a/Engine/Plugins/VirtualProduction/Avalanche/Source/AvalancheMask/Private/Handling/AvaMaskText3DActorMaterialCollectionHandle.cpp
+++ b/Engine/Plugins/VirtualProduction/Avalanche/Source/AvalancheMask/Private/Handling/AvaMaskText3DActorMaterialCollectionHandle.cpp
@@ -98,16 +98,6 @@ TArray<TSharedPtr<IAvaMaskMaterialHandle>> FAvaMaskText3DActorMaterialCollection
 			UMaterial* DefaultMaterial = nullptr;
 			MaterialExtension->ForEachMaterial([this, HandleSubsystem, &Index, &DefaultMaterial](const UE::Text3D::Material::FMaterialParameters& InParameters, UMaterialInterface* InMaterial)->bool
 			{
-				if (!InMaterial)
-				{
-					if (!DefaultMaterial)
-					{
-						DefaultMaterial = LoadObject<UMaterial>(nullptr, TEXT("/Script/Engine.Material'/Engine/EngineMaterials/WorldGridMaterial.WorldGridMaterial'"));
-					}
-
-					InMaterial = DefaultMaterial;		
-				}
-
 				if (!MaterialHandles[Index] || MaterialHandles[Index]->GetMaterial() != InMaterial)
 				{
 					MaterialHandles[Index] = HandleSubsystem->MakeHandle<IAvaMaskMaterialHandle>(InMaterial, UE::AvaMask::Internal::HandleTag);
diff --git a/Engine/Plugins/VirtualProduction/Avalanche/Source/AvalancheMask/Private/Handling/AvaObjectHandleSubsystem.cpp b/Engine/Plugins/VirtualProduction/Avalanche/Source/AvalancheMask/Private/Handling/AvaObjectHandleSubsystem.cpp
--- a/Engine/Plugins/VirtualProduction/Avalanche/Source/AvalancheMask/Private/Handling/AvaObjectHandleSubsystem.cpp
+++ b/Engine/Plugins/VirtualProduction/Avalanche/Source/AvalancheMask/Private/Handling/AvaObjectHandleSubsystem.cpp
@@ -5,6 +5,7 @@
 #include "AvaHandleUtilities.h"
 #include "AvaMaskActorMaterialCollectionHandle.h"
 #include "AvaMaskAvaShapeMaterialCollectionHandle.h"
+#include "AvaMaskEmptyMaterialHandle.h"
 #include "AvaMaskLog.h"
 #include "AvaMaskMaterialReference.h"
 #include "AvaMaskMediaPlateMaterialHandle.h"
@@ -108,5 +109,12 @@ void UAvaObjectHandleSubsystem::FindObjectHandleFactories()
 			{
 				return MakeShared<FAvaMaskMaterialInstanceHandle>(InMaterial.GetTypedObject<UMaterialInterface>());
 			}});
+
+		ObjectHandleFactories.Add({
+			FAvaMaskEmptyMaterialHandle::IsSupported
+			, [](const FAvaMaskMaterialReference& InMaterial) -> TSharedPtr<IAvaMaskMaterialHandle>
+			{
+				return MakeShared<FAvaMaskEmptyMaterialHandle>();
+			}});
 	}
 }
diff --git a/Engine/Plugins/VirtualProduction/Avalanche/Source/AvalancheMask/Private/Handling/AvaObjectHandleSubsystem.h b/Engine/Plugins/VirtualProduction/Avalanche/Source/AvalancheMask/Private/Handling/AvaObjectHandleSubsystem.h
--- a/Engine/Plugins/VirtualProduction/Avalanche/Source/AvalancheMask/Private/Handling/AvaObjectHandleSubsystem.h
+++ b/Engine/Plugins/VirtualProduction/Avalanche/Source/AvalancheMask/Private/Handling/AvaObjectHandleSubsystem.h
@@ -28,11 +28,6 @@ public:
 				&&	!std::is_const_v<ObjectType> && TModels_V<CStaticClassProvider, std::decay_t<ObjectType>>)>
     TSharedPtr<HandleType> MakeHandle(ObjectType* InInstance, FName InTag = NAME_None)
     {
-		if (!InInstance)
-		{
-			UE_LOG(LogAvaMask, Warning, TEXT("Invalid or null object provided to MakeHandle"));
-			return nullptr;
-		}
     	return StaticCastSharedPtr<HandleType>(MakeHandleDirect(FAvaMaskMaterialReference(InInstance), InTag));
     }
 
