Motion Design: fixed issue with local array containing raw ptrs to behavior instances. The behavior instance array is reallocated in specific situations, causing these raw ptrs to be left dangling.

diff --git a/Engine/Plugins/VirtualProduction/Avalanche/Source/AvalancheTransition/Private/Execution/AvaTransitionExecutor.cpp b/Engine/Plugins/VirtualProduction/Avalanche/Source/AvalancheTransition/Private/Execution/AvaTransitionExecutor.cpp
--- a/Engine/Plugins/VirtualProduction/Avalanche/Source/AvalancheTransition/Private/Execution/AvaTransitionExecutor.cpp
+++ b/Engine/Plugins/VirtualProduction/Avalanche/Source/AvalancheTransition/Private/Execution/AvaTransitionExecutor.cpp
@@ -10,6 +10,7 @@
 #include "Behavior/IAvaTransitionBehavior.h"
 #include "Engine/World.h"
 #include "Execution/AvaTransitionExecutorBuilder.h"
+#include "Misc/EnumerateRange.h"
 
 DEFINE_LOG_CATEGORY_STATIC(LogAvaTransitionExecutor, Log, All);
 
@@ -52,10 +53,10 @@ void FAvaTransitionExecutor::Setup()
 	struct FLayerInfo
 	{
 		// Enter Instances found for a given Layer
-		TArray<FAvaTransitionBehaviorInstance*> EnterInstances;
+		TArray<int32, TInlineAllocator<1>> EnterInstances;
 
 		// Exit Instances found for a given Layer
-		TArray<FAvaTransitionBehaviorInstance*> ExitInstances;
+		TArray<int32, TInlineAllocator<1>> ExitInstances;
 
 		// The accumulated Transition Type for a given Layer (e.g. combinations could be In, Out or In | Out)
 		EAvaTransitionType TransitionType;
@@ -70,22 +71,22 @@ void FAvaTransitionExecutor::Setup()
 		TagLayerInfo.Reserve(Instances.Num());
 
 		// Gather the Layer info for each Instance
-		for (FAvaTransitionBehaviorInstance& Instance : Instances)
+		for (TEnumerateRef<FAvaTransitionBehaviorInstance> Instance : EnumerateRange(Instances))
 		{
-			for (const FAvaTag* Tag : Instance.GetTransitionLayer().GetTags())
+			for (const FAvaTag* Tag : Instance->GetTransitionLayer().GetTags())
 			{
 				FLayerInfo& LayerInfo = TagLayerInfo.FindOrAdd(*Tag);
-				LayerInfo.TransitionType |= Instance.GetTransitionType();
-				LayerInfo.TransitionLayer = Instance.GetTransitionLayer();
+				LayerInfo.TransitionType |= Instance->GetTransitionType();
+				LayerInfo.TransitionLayer = Instance->GetTransitionLayer();
 
-				switch (Instance.GetTransitionType())
+				switch (Instance->GetTransitionType())
 				{
 				case EAvaTransitionType::In:
-					LayerInfo.EnterInstances.AddUnique(&Instance);
+					LayerInfo.EnterInstances.AddUnique(Instance.GetIndex());
 					break;
 
 				case EAvaTransitionType::Out:
-					LayerInfo.ExitInstances.AddUnique(&Instance);
+					LayerInfo.ExitInstances.AddUnique(Instance.GetIndex());
 					break;
 				}
 			}
@@ -115,22 +116,24 @@ void FAvaTransitionExecutor::Setup()
 			continue;
 		}
 
-		for (FAvaTransitionBehaviorInstance* ExitInstance : Pair.Value.ExitInstances)
+		for (int32 ExitInstanceIndex : Pair.Value.ExitInstances)
 		{
-			check(ExitInstance->GetTransitionType() == EAvaTransitionType::Out);
+			FAvaTransitionBehaviorInstance& ExitInstance = Instances[ExitInstanceIndex];
+			check(ExitInstance.GetTransitionType() == EAvaTransitionType::Out);
 
-			if (FAvaTransitionScene* TransitionScene = ExitInstance->GetTransitionContext().GetTransitionScene())
+			if (FAvaTransitionScene* TransitionScene = ExitInstance.GetTransitionContext().GetTransitionScene())
 			{
 				bool bSceneReused = false;
 
-				const UAvaTransitionTree* TransitionTree = ExitInstance->GetTransitionTree();
+				const UAvaTransitionTree* TransitionTree = ExitInstance.GetTransitionTree();
 				if (TransitionTree && TransitionTree->GetInstancingMode() == EAvaTransitionInstancingMode::Reuse)
 				{
 					// Scene reused if the Tree Instancing Mode is set to reuse, and if there is an Enter Instance with matching Tree (i.e. Level)
 					bSceneReused |= Pair.Value.EnterInstances.ContainsByPredicate(
-						[TransitionTree](const FAvaTransitionBehaviorInstance* InEnterInstance)
+						[TransitionTree, &Instances = Instances](int32 InEnterInstanceIndex)
 						{
-							return InEnterInstance && InEnterInstance->GetTransitionTree() == TransitionTree;
+							const FAvaTransitionBehaviorInstance& EnterInstance = Instances[InEnterInstanceIndex];
+							return EnterInstance.GetTransitionTree() == TransitionTree;
 						});
 				}
 
