MotionDesign : GeometryMask
- Fixed invalid/empty material slot on masked primitive components causing issues, this new empty handle will assign a MID based on default mask material on apply and restore the slot to none when original state is restored, this will allow primitive components with material slots set to none to still work with mask
- Fixed a regression introduced previously where slots with a transient MID based on a valid material would not be saved and empty handle would apply when it should not since slot would be none but saved original material is not

diff --git a/Engine/Plugins/VirtualProduction/Avalanche/Source/AvalancheMask/Private/Handling/AvaHandleUtilities.cpp b/Engine/Plugins/VirtualProduction/Avalanche/Source/AvalancheMask/Private/Handling/AvaHandleUtilities.cpp
--- a/Engine/Plugins/VirtualProduction/Avalanche/Source/AvalancheMask/Private/Handling/AvaHandleUtilities.cpp
+++ b/Engine/Plugins/VirtualProduction/Avalanche/Source/AvalancheMask/Private/Handling/AvaHandleUtilities.cpp
@@ -3,6 +3,8 @@
 #include "Handling/AvaHandleUtilities.h"
 
 #include "Engine/Engine.h"
+#include "Materials/Material.h"
+#include "Materials/MaterialInterface.h"
 #include "Handling/AvaObjectHandleSubsystem.h"
 
 namespace UE::Ava::Internal
@@ -11,4 +13,27 @@ namespace UE::Ava::Internal
 	{
 		return GEngine->GetEngineSubsystem<UAvaObjectHandleSubsystem>();
 	}
+
+	bool IsDefaultMaterial(UMaterialInterface* InMaterial)
+	{
+		if (!InMaterial)
+		{
+			return false;
+		}
+
+		if (InMaterial->GetPathName() == UAvaObjectHandleSubsystem::DefaultMaterialPath)
+		{
+			return true;
+		}
+
+		if (const UMaterialInterface* BaseMaterial = InMaterial->GetBaseMaterial())
+		{
+			if (BaseMaterial->GetPathName() == UAvaObjectHandleSubsystem::DefaultMaterialPath)
+			{
+				return true;
+			}
+		}
+
+		return false;
+	}
 }
diff --git a/Engine/Plugins/VirtualProduction/Avalanche/Source/AvalancheMask/Private/Handling/AvaHandleUtilities.h b/Engine/Plugins/VirtualProduction/Avalanche/Source/AvalancheMask/Private/Handling/AvaHandleUtilities.h
--- a/Engine/Plugins/VirtualProduction/Avalanche/Source/AvalancheMask/Private/Handling/AvaHandleUtilities.h
+++ b/Engine/Plugins/VirtualProduction/Avalanche/Source/AvalancheMask/Private/Handling/AvaHandleUtilities.h
@@ -14,6 +14,8 @@ namespace UE::Ava::Internal
 {
 	UAvaObjectHandleSubsystem* GetObjectHandleSubsystem();
 
+	bool IsDefaultMaterial(UMaterialInterface* InMaterial);
+
 	template <typename HandleType
 		UE_REQUIRES(std::is_base_of_v<IAvaObjectHandle, HandleType>)>
 	bool IsHandleValid(const TSharedPtr<HandleType>& InHandle)
diff --git a/Engine/Plugins/VirtualProduction/Avalanche/Source/AvalancheMask/Private/Handling/AvaMaskActorMaterialCollectionHandle.cpp b/Engine/Plugins/VirtualProduction/Avalanche/Source/AvalancheMask/Private/Handling/AvaMaskActorMaterialCollectionHandle.cpp
--- a/Engine/Plugins/VirtualProduction/Avalanche/Source/AvalancheMask/Private/Handling/AvaMaskActorMaterialCollectionHandle.cpp
+++ b/Engine/Plugins/VirtualProduction/Avalanche/Source/AvalancheMask/Private/Handling/AvaMaskActorMaterialCollectionHandle.cpp
@@ -58,7 +58,7 @@ namespace UE::AvaMask::Private
 					else
 					{
 						ComponentMaterials[SlotIdx] = Material;
-						ComponentMaterialHandles[SlotIdx] = HandleSubsystem->MakeHandle<IAvaMaskMaterialHandle>(Material, UE::AvaMask::Internal::HandleTag);
+						ComponentMaterialHandles[SlotIdx] = HandleSubsystem->MakeMaterialHandle<IAvaMaskMaterialHandle>(Material, UE::AvaMask::Internal::HandleTag);
 					}
 				}
 			}
@@ -534,7 +534,7 @@ bool FAvaMaskActorMaterialCollectionHandle::RefreshMaterials()
 			else
 			{
 				ComponentMaterials[SlotIdx] = Material;
-				ComponentMaterialHandles[SlotIdx] = HandleSubsystem->MakeHandle<IAvaMaskMaterialHandle>(Material, UE::AvaMask::Internal::HandleTag);
+				ComponentMaterialHandles[SlotIdx] = HandleSubsystem->MakeMaterialHandle<IAvaMaskMaterialHandle>(Material, UE::AvaMask::Internal::HandleTag);
 			}
 		}
 
diff --git a/Engine/Plugins/VirtualProduction/Avalanche/Source/AvalancheMask/Private/Handling/AvaMaskAvaShapeMaterialCollectionHandle.cpp b/Engine/Plugins/VirtualProduction/Avalanche/Source/AvalancheMask/Private/Handling/AvaMaskAvaShapeMaterialCollectionHandle.cpp
--- a/Engine/Plugins/VirtualProduction/Avalanche/Source/AvalancheMask/Private/Handling/AvaMaskAvaShapeMaterialCollectionHandle.cpp
+++ b/Engine/Plugins/VirtualProduction/Avalanche/Source/AvalancheMask/Private/Handling/AvaMaskAvaShapeMaterialCollectionHandle.cpp
@@ -395,7 +395,7 @@ bool FAvaMaskAvaShapeMaterialCollectionHandle::RefreshMaterials()
 				}
 				else
 				{
-					MaterialHandle = MaterialHandles.Emplace(InMeshIdx, HandleSubsystem->MakeHandle<IAvaMaskMaterialHandle>(Material, UE::AvaMask::Internal::HandleTag));
+					MaterialHandle = MaterialHandles.Emplace(InMeshIdx, HandleSubsystem->MakeMaterialHandle<IAvaMaskMaterialHandle>(Material, UE::AvaMask::Internal::HandleTag));
 					if (!ensureAlways(UE::Ava::Internal::IsHandleValid(MaterialHandle)))
 					{
 						return false;
diff --git a/Engine/Plugins/VirtualProduction/Avalanche/Source/AvalancheMask/Private/Handling/AvaMaskEmptyMaterialHandle.cpp b/Engine/Plugins/VirtualProduction/Avalanche/Source/AvalancheMask/Private/Handling/AvaMaskEmptyMaterialHandle.cpp
deleted file mode 100644
diff --git a/Engine/Plugins/VirtualProduction/Avalanche/Source/AvalancheMask/Private/Handling/AvaMaskEmptyMaterialHandle.h b/Engine/Plugins/VirtualProduction/Avalanche/Source/AvalancheMask/Private/Handling/AvaMaskEmptyMaterialHandle.h
deleted file mode 100644
diff --git a/Engine/Plugins/VirtualProduction/Avalanche/Source/AvalancheMask/Private/Handling/AvaMaskMaterialInstanceHandle.h b/Engine/Plugins/VirtualProduction/Avalanche/Source/AvalancheMask/Private/Handling/AvaMaskMaterialInstanceHandle.h
--- a/Engine/Plugins/VirtualProduction/Avalanche/Source/AvalancheMask/Private/Handling/AvaMaskMaterialInstanceHandle.h
+++ b/Engine/Plugins/VirtualProduction/Avalanche/Source/AvalancheMask/Private/Handling/AvaMaskMaterialInstanceHandle.h
@@ -19,7 +19,7 @@ struct FAvaMaskMaterialInstanceHandleData
 	GENERATED_BODY()
 
 	UPROPERTY()
-	TWeakObjectPtr<UMaterialInterface> OriginalMaterial;
+	TObjectPtr<UMaterialInterface> OriginalMaterial;
 
 	UPROPERTY()
 	FAvaMask2DMaterialParameters OriginalMaskMaterialParameters;
diff --git a/Engine/Plugins/VirtualProduction/Avalanche/Source/AvalancheMask/Private/Handling/AvaMaskText3DActorMaterialCollectionHandle.cpp b/Engine/Plugins/VirtualProduction/Avalanche/Source/AvalancheMask/Private/Handling/AvaMaskText3DActorMaterialCollectionHandle.cpp
--- a/Engine/Plugins/VirtualProduction/Avalanche/Source/AvalancheMask/Private/Handling/AvaMaskText3DActorMaterialCollectionHandle.cpp
+++ b/Engine/Plugins/VirtualProduction/Avalanche/Source/AvalancheMask/Private/Handling/AvaMaskText3DActorMaterialCollectionHandle.cpp
@@ -100,7 +100,7 @@ TArray<TSharedPtr<IAvaMaskMaterialHandle>> FAvaMaskText3DActorMaterialCollection
 			{
 				if (!MaterialHandles[Index] || MaterialHandles[Index]->GetMaterial() != InMaterial)
 				{
-					MaterialHandles[Index] = HandleSubsystem->MakeHandle<IAvaMaskMaterialHandle>(InMaterial, UE::AvaMask::Internal::HandleTag);
+					MaterialHandles[Index] = HandleSubsystem->MakeMaterialHandle<IAvaMaskMaterialHandle>(InMaterial, UE::AvaMask::Internal::HandleTag);
 				}
 
 				Index++;
diff --git a/Engine/Plugins/VirtualProduction/Avalanche/Source/AvalancheMask/Private/Handling/AvaMaterialInstanceHandle.cpp b/Engine/Plugins/VirtualProduction/Avalanche/Source/AvalancheMask/Private/Handling/AvaMaterialInstanceHandle.cpp
--- a/Engine/Plugins/VirtualProduction/Avalanche/Source/AvalancheMask/Private/Handling/AvaMaterialInstanceHandle.cpp
+++ b/Engine/Plugins/VirtualProduction/Avalanche/Source/AvalancheMask/Private/Handling/AvaMaterialInstanceHandle.cpp
@@ -9,11 +9,18 @@
 #include "Materials/MaterialInterface.h"
 #include "UObject/Package.h"
 
-FAvaMaterialInstanceHandle::FAvaMaterialInstanceHandle(
-	const TWeakObjectPtr<UMaterialInterface>& InWeakParentMaterial)
-	: WeakParentMaterial(InWeakParentMaterial)
-	, WeakMaterialInstance(Cast<UMaterialInstanceDynamic>(InWeakParentMaterial))
+FAvaMaterialInstanceHandle::FAvaMaterialInstanceHandle(const TWeakObjectPtr<UMaterialInterface>& InWeakParentMaterial)
 {
+	// Use base material as parent if material is a MID
+	if (UMaterialInstanceDynamic* MID = Cast<UMaterialInstanceDynamic>(InWeakParentMaterial.Get()))
+	{
+		WeakMaterialInstance = MID;
+		WeakParentMaterial = MID->GetBaseMaterial();
+	}
+	else
+	{
+		WeakParentMaterial = InWeakParentMaterial;
+	}
 }
 
 FString FAvaMaterialInstanceHandle::GetMaterialName()
@@ -61,7 +68,22 @@ void FAvaMaterialInstanceHandle::CopyParametersFrom(UMaterialInstance* InSourceM
 
 void FAvaMaterialInstanceHandle::SetParentMaterial(UMaterialInterface* InMaterial)
 {
-	WeakParentMaterial = InMaterial;
+	if (UMaterialInstanceDynamic* MID = Cast<UMaterialInstanceDynamic>(InMaterial))
+	{
+		WeakParentMaterial = MID->GetBaseMaterial();
+	}
+	else
+	{
+		WeakParentMaterial = InMaterial;
+	}
+
+	if (UMaterialInstanceDynamic* MID = WeakMaterialInstance.Get())
+	{
+		if (MID->GetBaseMaterial() != WeakParentMaterial.Get())
+		{
+			WeakMaterialInstance.Reset();
+		}
+	}
 }
 
 bool FAvaMaterialInstanceHandle::IsValid() const
diff --git a/Engine/Plugins/VirtualProduction/Avalanche/Source/AvalancheMask/Private/Handling/AvaObjectHandleSubsystem.cpp b/Engine/Plugins/VirtualProduction/Avalanche/Source/AvalancheMask/Private/Handling/AvaObjectHandleSubsystem.cpp
--- a/Engine/Plugins/VirtualProduction/Avalanche/Source/AvalancheMask/Private/Handling/AvaObjectHandleSubsystem.cpp
+++ b/Engine/Plugins/VirtualProduction/Avalanche/Source/AvalancheMask/Private/Handling/AvaObjectHandleSubsystem.cpp
@@ -5,7 +5,6 @@
 #include "AvaHandleUtilities.h"
 #include "AvaMaskActorMaterialCollectionHandle.h"
 #include "AvaMaskAvaShapeMaterialCollectionHandle.h"
-#include "AvaMaskEmptyMaterialHandle.h"
 #include "AvaMaskLog.h"
 #include "AvaMaskMaterialReference.h"
 #include "AvaMaskMediaPlateMaterialHandle.h"
@@ -20,10 +19,13 @@
 #include "Materials/MaterialInterface.h"
 #include "UObject/UObjectIterator.h"
 
+FString UAvaObjectHandleSubsystem::DefaultMaterialPath = TEXT("/Avalanche/MaskResources/M_AvalancheMaskDefaultTranslucent.M_AvalancheMaskDefaultTranslucent");
+
 void UAvaObjectHandleSubsystem::Initialize(FSubsystemCollectionBase& Collection)
 {
 	Super::Initialize(Collection);
 
+	DefaultMaterial = TSoftObjectPtr<UMaterialInterface>(FSoftObjectPath(DefaultMaterialPath));
 	FindObjectHandleFactories();
 }
 
@@ -109,12 +111,5 @@ void UAvaObjectHandleSubsystem::FindObjectHandleFactories()
 			{
 				return MakeShared<FAvaMaskMaterialInstanceHandle>(InMaterial.GetTypedObject<UMaterialInterface>());
 			}});
-
-		ObjectHandleFactories.Add({
-			FAvaMaskEmptyMaterialHandle::IsSupported
-			, [](const FAvaMaskMaterialReference& InMaterial) -> TSharedPtr<IAvaMaskMaterialHandle>
-			{
-				return MakeShared<FAvaMaskEmptyMaterialHandle>();
-			}});
 	}
 }
diff --git a/Engine/Plugins/VirtualProduction/Avalanche/Source/AvalancheMask/Private/Handling/AvaObjectHandleSubsystem.h b/Engine/Plugins/VirtualProduction/Avalanche/Source/AvalancheMask/Private/Handling/AvaObjectHandleSubsystem.h
--- a/Engine/Plugins/VirtualProduction/Avalanche/Source/AvalancheMask/Private/Handling/AvaObjectHandleSubsystem.h
+++ b/Engine/Plugins/VirtualProduction/Avalanche/Source/AvalancheMask/Private/Handling/AvaObjectHandleSubsystem.h
@@ -6,9 +6,13 @@
 #include "AvaMaskMaterialReference.h"
 #include "IAvaObjectHandle.h"
 #include "Subsystems/EngineSubsystem.h"
+#include "UObject/SoftObjectPtr.h"
 #include "AvaObjectHandleSubsystem.generated.h"
 
-class UAvaObjectHandleSubsystem;
+class IAvaMaskMaterialCollectionHandle;
+class IAvaMaskMaterialHandle;
+class UMaterialInterface;
+class UObject;
 
 /** Responsible for providing Handlers for a given UObject. */
 UCLASS()
@@ -17,20 +21,53 @@ class UAvaObjectHandleSubsystem : public UEngineSubsystem
 	GENERATED_BODY()
 
 public:
-	// ~Begin USubsystem
+	static FString DefaultMaterialPath;
+
+	//~ Begin USubsystem
 	virtual void Initialize(FSubsystemCollectionBase& Collection) override;
-	// ~End USubsystem
+	//~ End USubsystem
 
 	TSharedPtr<IAvaObjectHandle> MakeHandleDirect(const FAvaMaskMaterialReference& InInstance, FName InTag = NAME_None);
 
-	template <typename HandleType, typename ObjectType = UObject
-		UE_REQUIRES(!std::is_const_v<HandleType> && std::is_base_of_v<IAvaObjectHandle, HandleType>
-				&&	!std::is_const_v<ObjectType> && TModels_V<CStaticClassProvider, std::decay_t<ObjectType>>)>
-    TSharedPtr<HandleType> MakeHandle(ObjectType* InInstance, FName InTag = NAME_None)
+	/**
+	 * Used to create material collection handles, Owner provided must be valid
+	 */
+	template <typename HandleType
+		UE_REQUIRES(!std::is_const_v<HandleType> && std::is_base_of_v<IAvaMaskMaterialCollectionHandle, HandleType>)>
+    TSharedPtr<HandleType> MakeCollectionHandle(UObject* InOwner, FName InTag = NAME_None)
     {
-    	return StaticCastSharedPtr<HandleType>(MakeHandleDirect(FAvaMaskMaterialReference(InInstance), InTag));
+		return MakeHandle<HandleType>(InOwner, InTag);
     }
 
+	/**
+	 * Used to create material handles, material instance can be null, in that case the default material will be used
+	 */
+	template <typename HandleType
+		UE_REQUIRES(!std::is_const_v<HandleType> && std::is_base_of_v<IAvaMaskMaterialHandle, HandleType>)>
+	TSharedPtr<HandleType> MakeMaterialHandle(UMaterialInterface* InInstance, FName InTag = NAME_None)
+	{
+		if (!InInstance)
+		{
+			InInstance = DefaultMaterial.LoadSynchronous();
+		}
+
+		return MakeHandle<HandleType>(InInstance, InTag);
+	}
+
+	template <typename HandleType, typename ObjectType = UObject
+	UE_REQUIRES(!std::is_const_v<HandleType> && std::is_base_of_v<IAvaObjectHandle, HandleType>
+			&&	!std::is_const_v<ObjectType> && TModels_V<CStaticClassProvider, std::decay_t<ObjectType>>)>
+	TSharedPtr<HandleType> MakeHandle(ObjectType* InInstance, FName InTag = NAME_None)
+	{
+		if (!InInstance)
+		{
+			UE_LOG(LogAvaMask, Warning, TEXT("Invalid or null instance object provided to MakeHandle"));
+			return nullptr;
+		}
+
+		return StaticCastSharedPtr<HandleType>(MakeHandleDirect(FAvaMaskMaterialReference(InInstance), InTag));
+	}
+
 	template<typename HandleType UE_REQUIRES(!std::is_const_v<HandleType> && std::is_base_of_v<IAvaObjectHandle, HandleType>)>
 	TSharedPtr<HandleType> MakeHandle(const FAvaMaskMaterialReference& InInstance, FName InTag = NAME_None)
 	{
@@ -44,4 +81,7 @@ private:
 	using FMakeHandleFunction = TFunction<TSharedPtr<IAvaObjectHandle>(const FAvaMaskMaterialReference&)>;
 
 	TArray<TPair<FIsSupportedFunction, FMakeHandleFunction>> ObjectHandleFactories;
+
+	UPROPERTY()
+	TSoftObjectPtr<UMaterialInterface> DefaultMaterial;
 };
diff --git a/Engine/Plugins/VirtualProduction/Avalanche/Source/AvalancheMask/Private/Handling/IAvaMaskMaterialCollectionHandle.h b/Engine/Plugins/VirtualProduction/Avalanche/Source/AvalancheMask/Private/Handling/IAvaMaskMaterialCollectionHandle.h
--- a/Engine/Plugins/VirtualProduction/Avalanche/Source/AvalancheMask/Private/Handling/IAvaMaskMaterialCollectionHandle.h
+++ b/Engine/Plugins/VirtualProduction/Avalanche/Source/AvalancheMask/Private/Handling/IAvaMaskMaterialCollectionHandle.h
@@ -2,6 +2,8 @@
 
 #pragma once
 
+#include "AvaHandleUtilities.h"
+#include "AvaObjectHandleSubsystem.h"
 #include "Containers/Array.h"
 #include "Engine/EngineTypes.h"
 #include "IAvaMaskMaterialHandle.h"
@@ -238,6 +240,12 @@ bool TAvaMaskMaterialCollectionHandle<HandleDataType>::ApplyOriginalState(const
 					MaterialHandleDataView
 					, [InComponent, InSlotIdx, this](UMaterialInterface* InMaterial)
 					{
+						// If default mask material is used, restore back to "none" since this material is used for empty slots
+						if (UE::Ava::Internal::IsDefaultMaterial(InMaterial))
+						{
+							InMaterial = nullptr;
+						}
+
 						SetMaterial(InComponent, InSlotIdx, InMaterial);
 					});
 			}
diff --git a/Engine/Plugins/VirtualProduction/Avalanche/Source/AvalancheMask/Private/Mask2D/AvaMask2DBaseModifier.cpp b/Engine/Plugins/VirtualProduction/Avalanche/Source/AvalancheMask/Private/Mask2D/AvaMask2DBaseModifier.cpp
--- a/Engine/Plugins/VirtualProduction/Avalanche/Source/AvalancheMask/Private/Mask2D/AvaMask2DBaseModifier.cpp
+++ b/Engine/Plugins/VirtualProduction/Avalanche/Source/AvalancheMask/Private/Mask2D/AvaMask2DBaseModifier.cpp
@@ -302,6 +302,7 @@ void UAvaMask2DBaseModifier::OnModifierAdded(EActorModifierCoreEnableReason InRe
 	if (InReason == EActorModifierCoreEnableReason::Load)
 	{
 		LocalParamsToCanvas();
+		RestoreActors();
 	}
 	else
 	{
@@ -348,59 +349,39 @@ void UAvaMask2DBaseModifier::SavePreState()
 	Super::SavePreState();
 
 	ScannedPrimitiveComponents.Empty();
-
-	TMap<TWeakObjectPtr<AActor>, FAvaMask2DActorData> NewActorDatas;
-	NewActorDatas.Reserve(ActorData.Num());
+	TSet<AActor*> KeepActors;
 
 	ForEachActor<AActor>(
-		[this, &NewActorDatas](AActor* InActor)
+		[this, &KeepActors](AActor* InActor)
 		{
 			const UActorModifierCoreBase* MaskModifier = UE::ActorModifierCore::Utilities::FindFirstActorModifierByClass(InActor, UAvaMask2DBaseModifier::StaticClass());
 
 			// If first modifier in the hierarchy is not this modifier, then skip actor
 			if (MaskModifier == this)
 			{
-				SaveActorPreState(InActor, NewActorDatas.Emplace(InActor, ActorData.FindOrAdd(InActor)));
+				SaveActorPreState(InActor, ActorData.FindOrAdd(InActor));
+				KeepActors.Add(InActor);
 			}
 
 			return true;
 		}
 		, EActorModifierCoreLookup::SelfAndAllChildren);
 
-	ActorData = NewActorDatas;
+	// Cleanup
+	for (TMap<TWeakObjectPtr<AActor>, FAvaMask2DActorData>::TIterator It(ActorData); It; ++It)
+	{
+		if (!KeepActors.Contains(It->Key.Get()))
+		{
+			It.RemoveCurrent();
+		}
+	}
 }
 
 void UAvaMask2DBaseModifier::RestorePreState()
 {
 	Super::RestorePreState();
-	
-	TMap<TWeakObjectPtr<AActor>, FAvaMask2DActorData> ActorDataCopy = ActorData;
-	for (TPair<TWeakObjectPtr<AActor>, FAvaMask2DActorData>& ActorDataPair : ActorDataCopy)
-	{
-		if (AActor* Actor = ActorDataPair.Key.Get())
-		{
-			RestoreActorPreState(Actor, ActorDataPair.Value);
-		}
-	}
-	
-	TArray<AActor*> ActorDataActors;
-	ActorDataActors.Reserve(ActorData.Num());
-	Algo::Transform(ActorData, ActorDataActors, [](const TPair<TWeakObjectPtr<AActor>, FAvaMask2DActorData>& InActorDataPair)
-	{
-		return InActorDataPair.Key.Get();
-	});
-	
-	for (AActor* Actor : ActorDataActors)
-	{
-		if (!Actor)
-		{
-			continue;
-		}
-		
-		RemoveFromActor(Actor);
-	}
 
-	ScannedPrimitiveComponents.Empty();
+	RestoreActors();
 }
 
 void UAvaMask2DBaseModifier::Apply()
@@ -415,36 +396,9 @@ void UAvaMask2DBaseModifier::SaveActorPreState(AActor* InActor, FAvaMask2DActorD
 		UE::Ava::Internal::FindOrAddHandleByLambda(
 			MaterialCollectionHandles
 			, InActor
-			, [this, InActor, &InActorData]()
+			, [this, InActor]()
 			{
-				TSharedPtr<IAvaMaskMaterialCollectionHandle> Handle = GetObjectHandleSubsystem()->MakeHandle<IAvaMaskMaterialCollectionHandle>(InActor, UE::AvaMask::Internal::HandleTag);
-
-				Handle->ForEachMaterialHandle([this, &InActorData, InActor](
-					const FSoftComponentReference& InComponent,
-					const int32 InIdx,
-					const bool bInIsSlotOccupied,
-					const TSharedPtr<IAvaMaskMaterialHandle>& InMaterialHandle)->bool
-					{
-						if (InMaterialHandle.IsValid())
-						{
-							const FAvaMask2DComponentMaterialPath Component(InComponent, InIdx);
-							if (const TObjectPtr<UMaterialInterface>* MaterialPtr = InActorData.OriginalMaterials.Find(Component))
-							{
-								if (UMaterialInterface* Material = MaterialPtr->Get())
-								{
-									// Refresh parent material to create MID from
-									InMaterialHandle->SetParentMaterial(Material);
-								}
-							}
-
-							// Save parent material otherwise can't restore original due to material dynamic being transient
-							UMaterialInterface* Material = InMaterialHandle->GetParentMaterial();
-							InActorData.OriginalMaterials.Add(FAvaMask2DComponentMaterialPath(InComponent, InIdx), Material);
-						}
-
-						return true;
-					}
-				);
+				TSharedPtr<IAvaMaskMaterialCollectionHandle> Handle = GetObjectHandleSubsystem()->MakeCollectionHandle<IAvaMaskMaterialCollectionHandle>(InActor, UE::AvaMask::Internal::HandleTag);
 
 				if (!Handle->OnSourceMaterialsChanged().IsBoundToObject(this))
 				{
@@ -473,17 +427,12 @@ void UAvaMask2DBaseModifier::RestoreActorPreState(AActor* InActor, FAvaMask2DAct
     	UE::Ava::Internal::FindOrAddHandleByLambda(
     		MaterialCollectionHandles
     		, InActor
-    		, [this, InActor, &InActorData]()
+    		, [this, InActor]()
     		{
-    			TSharedPtr<IAvaMaskMaterialCollectionHandle> Handle = GetObjectHandleSubsystem()->MakeHandle<IAvaMaskMaterialCollectionHandle>(InActor, UE::AvaMask::Internal::HandleTag);
+    			TSharedPtr<IAvaMaskMaterialCollectionHandle> Handle = GetObjectHandleSubsystem()->MakeCollectionHandle<IAvaMaskMaterialCollectionHandle>(InActor, UE::AvaMask::Internal::HandleTag);
 
     			Handle->OnSourceMaterialsChanged().Unbind();
 
-    			for (const TPair<FAvaMask2DComponentMaterialPath, TObjectPtr<UMaterialInterface>>& OriginalMaterial : InActorData.OriginalMaterials)
-    			{
-    				Handle->SetMaterial(OriginalMaterial.Key.Component, OriginalMaterial.Key.SlotIdx, OriginalMaterial.Value);
-				}
-
     			return Handle;
     		}
     	);
@@ -930,4 +879,19 @@ void UAvaMask2DBaseModifier::OnMaterialsChanged(UObject* InMaterialOwner, const
 	}
 }
 
+void UAvaMask2DBaseModifier::RestoreActors()
+{
+	TMap<TWeakObjectPtr<AActor>, FAvaMask2DActorData> ActorDataCopy = ActorData;
+	for (TPair<TWeakObjectPtr<AActor>, FAvaMask2DActorData>& ActorDataPair : ActorDataCopy)
+	{
+		if (AActor* Actor = ActorDataPair.Key.Get())
+		{
+			RestoreActorPreState(Actor, ActorDataPair.Value);
+			RemoveFromActor(Actor);
+		}
+	}
+
+	ScannedPrimitiveComponents.Empty();
+}
+
 #undef LOCTEXT_NAMESPACE
diff --git a/Engine/Plugins/VirtualProduction/Avalanche/Source/AvalancheMask/Private/Mask2D/AvaMask2DReadModifier.cpp b/Engine/Plugins/VirtualProduction/Avalanche/Source/AvalancheMask/Private/Mask2D/AvaMask2DReadModifier.cpp
--- a/Engine/Plugins/VirtualProduction/Avalanche/Source/AvalancheMask/Private/Mask2D/AvaMask2DReadModifier.cpp
+++ b/Engine/Plugins/VirtualProduction/Avalanche/Source/AvalancheMask/Private/Mask2D/AvaMask2DReadModifier.cpp
@@ -122,7 +122,7 @@ bool UAvaMask2DReadModifier::ApplyRead(AActor* InActor, FAvaMask2DActorData& InA
 				, InActor
 				, [this, InActor]()
 				{
-					return GetObjectHandleSubsystem()->MakeHandle<IAvaMaskMaterialCollectionHandle>(InActor, UE::AvaMask::Internal::HandleTag);
+					return GetObjectHandleSubsystem()->MakeCollectionHandle<IAvaMaskMaterialCollectionHandle>(InActor, UE::AvaMask::Internal::HandleTag);
 				});
 
 		const UTexture* CanvasTexture = TryResolveCanvasTexture(InActor, InActorData);
diff --git a/Engine/Plugins/VirtualProduction/Avalanche/Source/AvalancheMask/Public/Mask2D/AvaMask2DBaseModifier.h b/Engine/Plugins/VirtualProduction/Avalanche/Source/AvalancheMask/Public/Mask2D/AvaMask2DBaseModifier.h
--- a/Engine/Plugins/VirtualProduction/Avalanche/Source/AvalancheMask/Public/Mask2D/AvaMask2DBaseModifier.h
+++ b/Engine/Plugins/VirtualProduction/Avalanche/Source/AvalancheMask/Public/Mask2D/AvaMask2DBaseModifier.h
@@ -44,7 +44,7 @@ struct FAvaMask2DActorData
 	GENERATED_BODY()
 
 	/** The canvas texture to apply to this actor materials */
-	UPROPERTY()
+	UPROPERTY(Transient)
 	TWeakObjectPtr<UTexture> CanvasTextureWeak;
 
 	/** Original assigned materials, used to create instances from, and restore to when modifier removed. */
@@ -168,6 +168,9 @@ protected:
 	
 	virtual void OnMaterialsChanged(UObject* InMaterialOwner, const TArray<TSharedPtr<IAvaMaskMaterialHandle>>& InMaterialHandles);
 
+	/** Restores modified actors to their original state and clears cache data */
+	void RestoreActors();
+
 #if WITH_EDITOR
 	UFUNCTION(CallInEditor, Category = "Mask2D")
 	void VisualizeMask();
@@ -229,7 +232,7 @@ protected:
 	UPROPERTY(EditAnywhere, Transient, DuplicateTransient, Category = "Mask2D|Shared", NoClear,	meta = (DisplayName = "Canvas", AllowPrivateAccess = "true", ShowInnerProperties, NoResetToDefault))
 	TWeakObjectPtr<UGeometryMaskCanvas> CanvasWeak;
 
-	UPROPERTY(DuplicateTransient)
+	UPROPERTY()
 	TMap<TWeakObjectPtr<UObject>, FInstancedStruct> MaterialCollectionHandleData;
 
 	UPROPERTY(DuplicateTransient)
