Motion Design: fixed issue where dynamic mesh converter wasn't calling OnMeshChanged, causing the dynamic mesh component to not update its bounds and triggering all sorts of visual errors like no shadows casted or the mesh being incorrectly culled before it's fully out of the frustum

diff --git a/Engine/Plugins/VirtualProduction/Avalanche/Source/AvalancheModifiers/Private/Modifiers/AvaDynamicMeshConverterModifier.cpp b/Engine/Plugins/VirtualProduction/Avalanche/Source/AvalancheModifiers/Private/Modifiers/AvaDynamicMeshConverterModifier.cpp
--- a/Engine/Plugins/VirtualProduction/Avalanche/Source/AvalancheModifiers/Private/Modifiers/AvaDynamicMeshConverterModifier.cpp
+++ b/Engine/Plugins/VirtualProduction/Avalanche/Source/AvalancheModifiers/Private/Modifiers/AvaDynamicMeshConverterModifier.cpp
@@ -567,7 +567,27 @@ bool UAvaDynamicMeshConverterModifier::ConvertComponents(TArray<TWeakObjectPtr<U
 
 	ConvertedComponents = MoveTemp(NewConvertedComponents);
 
-	return MeshBuilder.BuildDynamicMesh(DynamicMeshComponent->GetDynamicMesh(), OutMaterialsWeak);
+	UDynamicMesh* const DynamicMesh = DynamicMeshComponent->GetDynamicMesh();
+	if (!IsValid(DynamicMesh))
+	{
+		return false;
+	}
+
+	// Copy ChangeInfo from how it gets built in FCEMeshBuilder::BuildDynamicMesh 
+	FDynamicMeshChangeInfo ChangeInfo;
+	ChangeInfo.Type = EDynamicMeshChangeType::GeneralEdit;
+	ChangeInfo.Flags = EDynamicMeshAttributeChangeFlags::Unknown;
+
+	// Call PreMeshChange to have the same behavior as UDynamicMesh::EditMeshInternal when bDeferChangeEvents is set to false
+	DynamicMesh->OnPreMeshChanged().Broadcast(DynamicMesh, ChangeInfo);
+
+	const bool bResult = MeshBuilder.BuildDynamicMesh(DynamicMesh, OutMaterialsWeak);
+
+	// Call OnMeshChanged and BP Event to have the same behavior as UDynamicMesh::EditMeshInternal when bDeferChangeEvents is set to false
+	DynamicMesh->OnMeshChanged().Broadcast(DynamicMesh, ChangeInfo);
+	DynamicMesh->MeshModifiedBPEvent.Broadcast(DynamicMesh);
+
+	return bResult;
 }
 
 bool UAvaDynamicMeshConverterModifier::HasFlag(EAvaDynamicMeshConverterModifierType InFlag) const
