Motion Design: fix outliner crashing while iterating level stream level actors and registering them. The issue happens because the outliner checks whether these actors are default brushes, etc to skip them. UWorld::GetDefaultPhysicsVolume spawns the default physics volume if it doesn't exist yet, which causes the level actor array to change while still iterating on it. First checking for HasDefaultPhysicsVolume prevents any spawn from happening.

diff --git a/Engine/Plugins/VirtualProduction/Avalanche/Source/AvalancheOutliner/Private/AvaOutliner.cpp b/Engine/Plugins/VirtualProduction/Avalanche/Source/AvalancheOutliner/Private/AvaOutliner.cpp
--- a/Engine/Plugins/VirtualProduction/Avalanche/Source/AvalancheOutliner/Private/AvaOutliner.cpp
+++ b/Engine/Plugins/VirtualProduction/Avalanche/Source/AvalancheOutliner/Private/AvaOutliner.cpp
@@ -160,7 +160,7 @@ bool FAvaOutliner::IsDefaultWorldActorToHide(const UWorld* const InWorld, const
 	}
 
 	// Make sure the Actor is none of these Default World Actors
-	return InActor == InWorld->GetDefaultPhysicsVolume()
+	return (InWorld->HasDefaultPhysicsVolume() && InActor == InWorld->GetDefaultPhysicsVolume())
 		|| InActor == InWorld->GetDefaultBrush()
 		|| InActor == InWorld->GetWorldSettings()
 		|| InActor == InWorld->MyParticleEventManager;
