Motion Design: fix an unexpected issue where a page from a template played in one channel gets taken out by playing another page from the same templated played in a different channel.

diff --git a/Engine/Plugins/VirtualProduction/Avalanche/Source/AvalancheMedia/Private/Rundown/AvaRundown.cpp b/Engine/Plugins/VirtualProduction/Avalanche/Source/AvalancheMedia/Private/Rundown/AvaRundown.cpp
--- a/Engine/Plugins/VirtualProduction/Avalanche/Source/AvalancheMedia/Private/Rundown/AvaRundown.cpp
+++ b/Engine/Plugins/VirtualProduction/Avalanche/Source/AvalancheMedia/Private/Rundown/AvaRundown.cpp
@@ -2310,13 +2310,23 @@ namespace UE::AvaMedia::Rundown::Private
 		for (const TObjectPtr<UAvaRundownPagePlayer>& PagePlayer : InRundown->GetPagePlayers())
 		{
 			// Early filter on preview/channel.
-			if (!PagePlayer
-				|| PagePlayer->bIsPreview != bInIsPreview
-				|| (bInIsPreview && PagePlayer->ChannelName != InPreviewChannelName) )
+			if (!PagePlayer || PagePlayer->bIsPreview != bInIsPreview)
 			{
 				continue;
 			}
-			
+
+			// Preview: Skip if player channel name mismatches the passed in preview channel
+			if (bInIsPreview && PagePlayer->ChannelName != InPreviewChannelName)
+			{
+				continue;
+			}
+
+			// Non-Preview: Skip if player channel mismatch the page's channel name
+			if (!bInIsPreview && PagePlayer->ChannelFName != InPageToPlay.GetChannelName())
+			{
+				continue;
+			}
+
 			const FAvaRundownPage& PlayingPage = InRundown->GetPage(PagePlayer->PageId);
 
 			if (!PlayingPage.IsValidPage())
