Motion Design: fixed issue with rendering rundown pages with MRQ that contain object references that have not been loaded yet

diff --git a/Engine/Plugins/VirtualProduction/Avalanche/Source/AvalancheMRQ/Private/AvaMRQRundownPageSetting.cpp b/Engine/Plugins/VirtualProduction/Avalanche/Source/AvalancheMRQ/Private/AvaMRQRundownPageSetting.cpp
--- a/Engine/Plugins/VirtualProduction/Avalanche/Source/AvalancheMRQ/Private/AvaMRQRundownPageSetting.cpp
+++ b/Engine/Plugins/VirtualProduction/Avalanche/Source/AvalancheMRQ/Private/AvaMRQRundownPageSetting.cpp
@@ -12,6 +12,7 @@
 #include "Misc/CoreMiscDefines.h"
 #include "Misc/PackageName.h"
 #include "MoviePipeline.h"
+#include "Playback/AvaPlaybackUtils.h"
 #include "Rundown/AvaRundown.h"
 
 DEFINE_LOG_CATEGORY_STATIC(LogAvaMRQSequenceData, Log, All);
@@ -117,5 +118,26 @@ void UAvaMRQRundownPageSetting::SetupForPipelineImpl(UMoviePipeline* InPipeline)
 	}
 
 	FAvaRemoteControlRebind::RebindUnboundEntities(RemoteControlPreset, Level);
-	Page.GetRemoteControlValues().ApplyEntityValuesToRemoteControlPreset(RemoteControlPreset);
+
+	const FAvaPlayableRemoteControlValues& RemoteControlValues = Page.GetRemoteControlValues();
+
+	// Synchronously load the object paths stored in the remote control values of the rundown page
+	// No async available as MRQ Setup phase does not do deferred setups
+	{
+		TSet<FSoftObjectPath> ReferencedPaths;
+
+		FAvaPlayableRemoteControlValues::CollectReferencedAssetPaths(RemoteControlValues.EntityValues, ReferencedPaths);
+		FAvaPlayableRemoteControlValues::CollectReferencedAssetPaths(RemoteControlValues.ControllerValues, ReferencedPaths);
+
+		for (const FSoftObjectPath& ReferencedPath : ReferencedPaths)
+		{
+			if (!ReferencedPath.ResolveObject())
+			{
+				// Synchronously load the object
+				ReferencedPath.TryLoad();
+			}
+		}
+	}
+
+	RemoteControlValues.ApplyEntityValuesToRemoteControlPreset(RemoteControlPreset);
 }
