Motion Design: fix crash spawning actors (e.g. via scene state or blueprint) and attaching them to an actor in another level having a translucent priority modifier

diff --git a/Engine/Plugins/VirtualProduction/Avalanche/Source/AvalancheModifiers/Private/Shared/AvaTranslucentPriorityModifierShared.cpp b/Engine/Plugins/VirtualProduction/Avalanche/Source/AvalancheModifiers/Private/Shared/AvaTranslucentPriorityModifierShared.cpp
--- a/Engine/Plugins/VirtualProduction/Avalanche/Source/AvalancheModifiers/Private/Shared/AvaTranslucentPriorityModifierShared.cpp
+++ b/Engine/Plugins/VirtualProduction/Avalanche/Source/AvalancheModifiers/Private/Shared/AvaTranslucentPriorityModifierShared.cpp
@@ -283,7 +283,8 @@ TArray<FAvaTranslucentPriorityModifierComponentState> UAvaTranslucentPriorityMod
 						const FAvaSceneTreeNode* TreeNode = nullptr;
 						while (InActor && !TreeNode)
 						{
-							TreeNode = SceneTree.FindTreeNode(FAvaSceneItem(InActor, World));
+							UObject* const StopOuter = World && InActor->IsIn(World) ? World : nullptr;
+							TreeNode = SceneTree.FindTreeNode(FAvaSceneItem(InActor, StopOuter));
 							InActor = InActor->GetAttachParentActor();
 						}
 						return TreeNode;
