Motion Design: fix rundown refresh assinging 'none' values to the remote control entities

diff --git a/Engine/Plugins/VirtualProduction/Avalanche/Source/AvalancheMediaEditor/Private/Rundown/DetailsView/RemoteControl/Properties/SAvaRundownPageRemoteControlProps.cpp b/Engine/Plugins/VirtualProduction/Avalanche/Source/AvalancheMediaEditor/Private/Rundown/DetailsView/RemoteControl/Properties/SAvaRundownPageRemoteControlProps.cpp
--- a/Engine/Plugins/VirtualProduction/Avalanche/Source/AvalancheMediaEditor/Private/Rundown/DetailsView/RemoteControl/Properties/SAvaRundownPageRemoteControlProps.cpp
+++ b/Engine/Plugins/VirtualProduction/Avalanche/Source/AvalancheMediaEditor/Private/Rundown/DetailsView/RemoteControl/Properties/SAvaRundownPageRemoteControlProps.cpp
@@ -258,7 +258,13 @@ void SAvaRundownPageRemoteControlProps::Refresh(const TArray<int32>& InSelectedP
 						const FAvaPlayableRemoteControlValues& Values = ActivePage->GetRemoteControlValues();
 						FAvaPlayableRemoteControlValues::CollectReferencedAssetPaths(Values.ControllerValues, Assets);
 						FAvaPlayableRemoteControlValues::CollectReferencedAssetPaths(Values.EntityValues, Assets);
-						AsyncAssetLoader->BeginLoadingAssets(Assets.Array());
+
+						// Synchronously load the assets as this refresh assigns the value of these assets to the remote control entities below.
+						// Doing an async load would mean that the relevant entities are likely to get assigned a 'none' value instead of the actual asset.
+						for (FSoftObjectPath& Asset : Assets)
+						{
+							Asset.TryLoad();
+						}
 					}
 				}
 			}
