MotionDesign : Text3D
- Fixed mask compatibility issue with text materials

diff --git a/Engine/Plugins/VirtualProduction/Avalanche/Source/AvalancheMask/Private/Handling/AvaMaskText3DActorMaterialCollectionHandle.cpp b/Engine/Plugins/VirtualProduction/Avalanche/Source/AvalancheMask/Private/Handling/AvaMaskText3DActorMaterialCollectionHandle.cpp
--- a/Engine/Plugins/VirtualProduction/Avalanche/Source/AvalancheMask/Private/Handling/AvaMaskText3DActorMaterialCollectionHandle.cpp
+++ b/Engine/Plugins/VirtualProduction/Avalanche/Source/AvalancheMask/Private/Handling/AvaMaskText3DActorMaterialCollectionHandle.cpp
@@ -61,7 +61,7 @@ TArray<TObjectPtr<UMaterialInterface>> FAvaMaskText3DActorMaterialCollectionHand
 	if (const UText3DComponent* TextComponent = GetComponent())
 	{
 		const UText3DMaterialExtensionBase* MaterialExtension = TextComponent->GetMaterialExtension();
-		Materials.Reserve(MaterialExtension->GetMaterialCount());
+		Materials.Reserve(GetNumMaterials());
 
 		MaterialExtension->ForEachMaterial([&Materials](const UE::Text3D::Material::FMaterialParameters& InParameters, UMaterialInterface* InMaterial)->bool
 		{
@@ -79,7 +79,7 @@ TArray<TSharedPtr<IAvaMaskMaterialHandle>> FAvaMaskText3DActorMaterialCollection
 	{
 		const UText3DMaterialExtensionBase* MaterialExtension = TextComponent->GetMaterialExtension();
 
-		const int32 MaterialCount = MaterialExtension->GetMaterialCount();
+		const int32 MaterialCount = GetNumMaterials();
 
 		if (MaterialHandles.Num() != MaterialCount
 			|| MaterialHandles.ContainsByPredicate([](const TSharedPtr<IAvaMaskMaterialHandle>& InHandle)
@@ -94,17 +94,15 @@ TArray<TSharedPtr<IAvaMaskMaterialHandle>> FAvaMaskText3DActorMaterialCollection
 			}
 
 			MaterialHandles.SetNum(MaterialCount);
-			int32 Index = 0;
-			UMaterial* DefaultMaterial = nullptr;
-			MaterialExtension->ForEachMaterial([this, HandleSubsystem, &Index, &DefaultMaterial](const UE::Text3D::Material::FMaterialParameters& InParameters, UMaterialInterface* InMaterial)->bool
+			int32 MaterialIndex = 0;
+			MaterialExtension->ForEachMaterial([this, HandleSubsystem, &MaterialIndex](const UE::Text3D::Material::FMaterialParameters& InParameters, UMaterialInterface* InMaterial)->bool
 			{
-				if (!MaterialHandles[Index] || MaterialHandles[Index]->GetMaterial() != InMaterial)
+				if (!MaterialHandles[MaterialIndex] || MaterialHandles[MaterialIndex]->GetMaterial() != InMaterial)
 				{
-					MaterialHandles[Index] = HandleSubsystem->MakeMaterialHandle<IAvaMaskMaterialHandle>(InMaterial, UE::AvaMask::Internal::HandleTag);
+					MaterialHandles[MaterialIndex] = HandleSubsystem->MakeMaterialHandle<IAvaMaskMaterialHandle>(InMaterial, UE::AvaMask::Internal::HandleTag);
 				}
 
-				Index++;
-
+				MaterialIndex++;
 				return true; // continue
 			});
 		}
@@ -121,35 +119,28 @@ void FAvaMaskText3DActorMaterialCollectionHandle::SetMaterial(
 	if (UText3DComponent* TextComponent = GetComponent())
 	{
 		UText3DMaterialExtensionBase* MaterialExtension = TextComponent->GetMaterialExtension();
-		const int32 MaterialCount = MaterialExtension->GetMaterialCount();
-		
-		if (!ensure(InSlotIdx < MaterialCount))
-		{
-			UE_LOG(LogAvaMask, Error, TEXT("SlotIdx out of range, should be from 0-%i, was %u"), MaterialCount-1, InSlotIdx);
-			return;
-		}
 
-		const int32 MaterialIndex = InSlotIdx / 4;
-		const EText3DGroupType GroupType = static_cast<EText3DGroupType>(InSlotIdx % 4);
-
-		TArray<FName> OutNames;
-		MaterialExtension->GetMaterialNames(OutNames);
+		int32 MaterialIndex = 0;
+		bool bUpdate = false;
+		MaterialExtension->ForEachMaterial([InSlotIdx, &MaterialIndex, MaterialExtension, InMaterial, &bUpdate](const UE::Text3D::Material::FMaterialParameters& InParams, UMaterialInterface* InCurrentMaterial)
+			{
+				if (MaterialIndex++ == InSlotIdx)
+				{
+					if (InCurrentMaterial != InMaterial)
+					{
+						MaterialExtension->SetMaterial(InParams, InMaterial);
+						bUpdate = true;
+					}
 
-		if (!ensure(OutNames.IsValidIndex(MaterialIndex)))
-		{
-			UE_LOG(LogAvaMask, Error, TEXT("Material index out of range, should be from 0-%i, was %u"), OutNames.Num()-1, MaterialIndex);
-			return;
-		}
+					return false; // Stop
+				}
 
-		UE::Text3D::Material::FMaterialParameters Parameters;
-		Parameters.Group = GroupType;
-		Parameters.Tag = OutNames[MaterialIndex];
-		const UMaterialInterface* CurrentMaterial = MaterialExtension->GetMaterial(Parameters);
+				return true; // continue
+			}
+		);
 
-		if (CurrentMaterial != InMaterial)
+		if (bUpdate)
 		{
-			MaterialExtension->SetMaterial(Parameters, InMaterial);
-
 			constexpr bool bImmediate = true;
 			TextComponent->RequestUpdate(EText3DRendererFlags::Material, bImmediate);
 		}
@@ -178,31 +169,33 @@ void FAvaMaskText3DActorMaterialCollectionHandle::SetMaterials(
 
 		UText3DMaterialExtensionBase* MaterialExtension = TextComponent->GetMaterialExtension();
 
-		TArray<FName> OutNames;
-		MaterialExtension->GetMaterialNames(OutNames);
-
-		bool bRefresh = false;
-		for (int32 Index = 0; Index < MaterialCount; ++Index)
-        {
-			if (InSetToggle[Index])
+		int32 MaterialIndex = 0;
+		bool bUpdate = false;
+		MaterialExtension->ForEachMaterial([&InMaterials, &InSetToggle, &MaterialIndex, MaterialExtension, &bUpdate](const UE::Text3D::Material::FMaterialParameters& InParams, UMaterialInterface* InCurrentMaterial)
 			{
-				const int32 MaterialIndex = Index / 4;
-				const EText3DGroupType GroupType = static_cast<EText3DGroupType>(Index % 4);
-
-				UE::Text3D::Material::FMaterialParameters Parameters;
-				Parameters.Group = GroupType;
-				Parameters.Tag = OutNames[MaterialIndex];
-				const UMaterialInterface* CurrentMaterial = MaterialExtension->GetMaterial(Parameters);
+				if (!InMaterials.IsValidIndex(MaterialIndex))
+				{
+					UE_LOG(LogAvaMask, Warning, TEXT("Expected index %i materials, got %u"), MaterialIndex, InMaterials.Num());
+					return false; // stop
+				}
 
-				if (CurrentMaterial != InMaterials[Index])
+				if (InSetToggle[MaterialIndex])
 				{
-					MaterialExtension->SetMaterial(Parameters, InMaterials[Index]);
-					bRefresh = true;
+					UMaterialInterface* Material = InMaterials[MaterialIndex];
+					
+					if (InCurrentMaterial != Material)
+					{
+						MaterialExtension->SetMaterial(InParams, Material);
+						bUpdate = true;
+					}
 				}
+
+				MaterialIndex++;
+				return true; // continue
 			}
-        }
+		);
 
-		if (bRefresh)
+		if (bUpdate)
 		{
 			constexpr bool bImmediate = true;
 			TextComponent->RequestUpdate(EText3DRendererFlags::Material, bImmediate);
@@ -212,13 +205,20 @@ void FAvaMaskText3DActorMaterialCollectionHandle::SetMaterials(
 
 int32 FAvaMaskText3DActorMaterialCollectionHandle::GetNumMaterials() const
 {
+	int32 MaterialCount = 0;
 	if (const UText3DComponent* TextComponent = GetComponent())
 	{
-		const UText3DMaterialExtensionBase* MaterialExtension = TextComponent->GetMaterialExtension();
-		return MaterialExtension->GetMaterialCount();
+		if (const UText3DMaterialExtensionBase* MaterialExtension = TextComponent->GetMaterialExtension())
+		{
+			MaterialExtension->ForEachMaterial([&MaterialCount](const UE::Text3D::Material::FMaterialParameters&, UMaterialInterface*)
+				{
+					MaterialCount++;
+					return true; // Continue
+				}
+			);
+		}
 	}
-
-	return 0;
+	return MaterialCount;
 }
 
 void FAvaMaskText3DActorMaterialCollectionHandle::ForEachMaterial(
@@ -288,22 +288,14 @@ void FAvaMaskText3DActorMaterialCollectionHandle::MapEachMaterial(
 
 		const UText3DMaterialExtensionBase* MaterialExtension = TextComponent->GetMaterialExtension();
 
-		TArray<FName> OutNames;
-		MaterialExtension->GetMaterialNames(OutNames);
-
-		for (int32 Index = 0; Index < MaterialCount; ++Index)
+		int32 MaterialIndex = 0;
+		MaterialExtension->ForEachMaterial([&InFunction, &ComponentReference, &MaterialIndex, &MappedMaterials, &SetFlags](const UE::Text3D::Material::FMaterialParameters& InParameters, UMaterialInterface* InMaterial)
 		{
-			const int32 MaterialIndex = Index / 4;
-			const EText3DGroupType GroupType = static_cast<EText3DGroupType>(Index % 4);
-
-			UE::Text3D::Material::FMaterialParameters Parameters;
-			Parameters.Group = GroupType;
-			Parameters.Tag = OutNames[MaterialIndex];
-			UMaterialInterface* Material = MaterialExtension->GetMaterial(Parameters);
-
-			MappedMaterials[Index] = InFunction(ComponentReference, Index, Material);
-			SetFlags[Index] = MappedMaterials[Index] != nullptr;
-		}
+			MappedMaterials[MaterialIndex] = InFunction(ComponentReference, MaterialIndex, InMaterial);
+			SetFlags[MaterialIndex] = MappedMaterials[MaterialIndex] != nullptr;
+			MaterialIndex++;
+			return true; // continue
+		});
 
 		SetMaterials(MappedMaterials, SetFlags);
 	}
@@ -328,13 +320,17 @@ void FAvaMaskText3DActorMaterialCollectionHandle::MapEachMaterialHandle(TFunctio
 
 		const UText3DMaterialExtensionBase* MaterialExtension = TextComponent->GetMaterialExtension();
 
-		TArray<FName> OutNames;
-		MaterialExtension->GetMaterialNames(OutNames);
-
-		for (int32 Index = 0; Index < MaterialCount; ++Index)
+		int32 MaterialIndex = 0;
+		MaterialExtension->ForEachMaterial([this, &InFunction, &ComponentReference, &MaterialIndex, &MappedMaterials, &SetFlags](const UE::Text3D::Material::FMaterialParameters& InParameters, UMaterialInterface* InMaterial)
 		{
-			TSharedPtr<IAvaMaskMaterialHandle>& MaterialHandle = MaterialHandles[Index];
-			const TSharedPtr<IAvaMaskMaterialHandle> MappedMaterialHandle = InFunction(ComponentReference, Index, true, MaterialHandle);
+			if (!MaterialHandles.IsValidIndex(MaterialIndex))
+			{
+				UE_LOG(LogAvaMask, Warning, TEXT("Expected index %i materials, got %u"), MaterialIndex, MaterialHandles.Num());
+				return false; // stop
+			}
+				
+			TSharedPtr<IAvaMaskMaterialHandle>& MaterialHandle = MaterialHandles[MaterialIndex];
+			const TSharedPtr<IAvaMaskMaterialHandle> MappedMaterialHandle = InFunction(ComponentReference, MaterialIndex, true, MaterialHandle);
 
 			UMaterialInterface* RemappedMaterial = nullptr;
 			if (MappedMaterialHandle.IsValid())
@@ -346,9 +342,11 @@ void FAvaMaskText3DActorMaterialCollectionHandle::MapEachMaterialHandle(TFunctio
 				RemappedMaterial = MaterialHandle->GetMaterial();
 			}
 
-			MappedMaterials[Index] = RemappedMaterial;
-			SetFlags[Index] = RemappedMaterial != nullptr;
-		}
+			MappedMaterials[MaterialIndex] = RemappedMaterial;
+			SetFlags[MaterialIndex] = RemappedMaterial != nullptr;
+			MaterialIndex++;
+			return true; // continue
+		});
 
 		SetMaterials(MappedMaterials, SetFlags);
 	}
@@ -543,21 +541,17 @@ void FAvaMaskText3DActorMaterialCollectionHandle::OnTextPostUpdate(UText3DCompon
 				return false;
 			}
 
-			int32 Index = 0;
+			int32 MaterialIndex = 0;
 			bool bMatch = true;
-			MaterialExtension->ForEachMaterial([this, &bMatch, &Index](const UE::Text3D::Material::FMaterialParameters& InParameters, UMaterialInterface* InMaterial)
+			MaterialExtension->ForEachMaterial([this, &bMatch, &MaterialIndex](const UE::Text3D::Material::FMaterialParameters& InParameters, UMaterialInterface* InMaterial)
 			{
-				UMaterialInterface* Material = MaterialHandles[Index]->GetMaterial();
+				UMaterialInterface* Material = MaterialHandles[MaterialIndex]->GetMaterial();
 				if (Material != InMaterial)
 				{
-					if (!Material || !InMaterial || Material->GetBaseMaterial() != InMaterial->GetBaseMaterial())
-					{
-						bMatch = false;
-					}
+					bMatch = false;
 				}
 
-				Index++;
-
+				MaterialIndex++;
 				return bMatch;
 			});
 	
diff --git a/Engine/Plugins/VirtualProduction/Avalanche/Source/AvalancheMask/Private/Mask2D/AvaMask2DBaseModifier.cpp b/Engine/Plugins/VirtualProduction/Avalanche/Source/AvalancheMask/Private/Mask2D/AvaMask2DBaseModifier.cpp
--- a/Engine/Plugins/VirtualProduction/Avalanche/Source/AvalancheMask/Private/Mask2D/AvaMask2DBaseModifier.cpp
+++ b/Engine/Plugins/VirtualProduction/Avalanche/Source/AvalancheMask/Private/Mask2D/AvaMask2DBaseModifier.cpp
@@ -887,7 +887,10 @@ void UAvaMask2DBaseModifier::RestoreActors()
 		if (AActor* Actor = ActorDataPair.Key.Get())
 		{
 			RestoreActorPreState(Actor, ActorDataPair.Value);
-			RemoveFromActor(Actor);
+
+			// Do not remove ActorData, only handles
+			MaterialCollectionHandleData.Remove(Actor);
+			MaterialCollectionHandles.Remove(Actor);
 		}
 	}
 
