MotionDesign : Modifiers
- Fixed auto size modifier performance issue due to modifier firing every frame

diff --git a/Engine/Plugins/VirtualProduction/Avalanche/Source/AvalancheModifiers/Private/Modifiers/AvaAutoSizeModifier.cpp b/Engine/Plugins/VirtualProduction/Avalanche/Source/AvalancheModifiers/Private/Modifiers/AvaAutoSizeModifier.cpp
--- a/Engine/Plugins/VirtualProduction/Avalanche/Source/AvalancheModifiers/Private/Modifiers/AvaAutoSizeModifier.cpp
+++ b/Engine/Plugins/VirtualProduction/Avalanche/Source/AvalancheModifiers/Private/Modifiers/AvaAutoSizeModifier.cpp
@@ -13,9 +13,31 @@
 
 #define LOCTEXT_NAMESPACE "AvaAutoSizeModifier"
 
+// todo : Add function to UE::ActorModifier::ActorUtils
+namespace UE::ActorModifier::Private
+{
+	FBox GetActorLocalBounds(AActor* InActor, bool bInIncludeChildren, bool bInSkipHidden)
+	{
+		TSet<TWeakObjectPtr<AActor>> AttachedModifyActors {InActor};
+
+		if (bInIncludeChildren)
+		{
+			TArray<AActor*> AttachedActors;
+			InActor->GetAttachedActors(AttachedActors, /** Reset */false, /** Recurse */true);
+			Algo::Transform(AttachedActors, AttachedModifyActors, [](AActor* InAttachedActor)
+			{
+				return InAttachedActor;
+			});
+		}
+
+		constexpr bool bTransformBox = false;
+		return UE::ActorModifier::ActorUtils::GetActorsBounds(AttachedModifyActors, InActor->GetActorTransform(), bInSkipHidden, bTransformBox);
+	}
+}
+
 bool UAvaAutoSizeModifier::IsModifierDirtyable() const
 {
-	const AActor* const ActorModified = GetModifiedActor();
+	AActor* ActorModified = GetModifiedActor();
 	AActor* const TrackedActor = ReferenceActor.ReferenceActorWeak.Get();
 
 	if (!IsValid(TrackedActor) || !IsValid(ActorModified))
@@ -23,9 +45,13 @@ bool UAvaAutoSizeModifier::IsModifierDirtyable() const
 		return Super::IsModifierDirtyable();
 	}
 
-	const FBox ReferenceActorLocalBounds = UE::ActorModifier::ActorUtils::GetActorsBounds(TrackedActor, bIncludeChildren, true);
+	constexpr bool bSkipHidden = true;
+	const FBox ReferenceActorLocalBounds = UE::ActorModifier::Private::GetActorLocalBounds(TrackedActor, bIncludeChildren, bSkipHidden);
+	const FBox ModifiedActorBounds = UE::ActorModifier::Private::GetActorLocalBounds(ActorModified, /** Recurse */false, /** SkipHidden */false);
 
-	if (ReferenceActorLocalBounds.GetSize().Equals(CachedReferenceBounds.GetSize(), 0.01))
+	// Compute the ratio between reference and modified actor bounds size and compare it to cached version (without X axis, only YZ are considered for Width/Height)
+	const FVector RatioBounds = FVector(0, 1, 1) * (ReferenceActorLocalBounds.GetSize() / FVector::Max(ModifiedActorBounds.GetSize(), FVector(UE_SMALL_NUMBER)));
+	if (RatioBounds.Equals(CachedFollowLocation, 0.01))
 	{
 		return Super::IsModifierDirtyable();
 	}
@@ -145,9 +171,15 @@ void UAvaAutoSizeModifier::Apply()
 		RenderExtension->SetTrackedActorVisibility(TrackedActor, bIncludeChildren);
 	}
 
-	CachedReferenceBounds = UE::ActorModifier::ActorUtils::GetActorsBounds(TrackedActor, bIncludeChildren, true);
+	constexpr bool bSkipHidden = true;
+	CachedReferenceBounds = UE::ActorModifier::Private::GetActorLocalBounds(TrackedActor, bIncludeChildren, bSkipHidden);
+
+	const FBox ModifiedActorBounds = UE::ActorModifier::Private::GetActorLocalBounds(CurrentlyModifiedActor, /** Recurse */false, /** SkipHidden */false);
 
-	const FBox ModifiedActorBounds = UE::ActorModifier::ActorUtils::GetActorsBounds(CurrentlyModifiedActor, false);
+	// Reuse this unused property for now since headers cannot be modified in a hotfix release
+	// Store the ratio between reference and modified actor bounds size (without X axis, only YZ are considered for Width/Height)
+	// todo : rename this property to CachedActorSizeRatio
+	CachedFollowLocation = FVector(0, 1, 1) * (CachedReferenceBounds.GetSize() / FVector::Max(ModifiedActorBounds.GetSize(), FVector(UE_SMALL_NUMBER)));
 
 	FVector ModifiedActorBoundsOrigin;
 	FVector ModifiedActorBoundsExtent;
