MotionDesign : Modifiers
- Fixed visibility and arrange modifier conflict causing children actors to be left in an incorrect state

diff --git a/Engine/Plugins/VirtualProduction/Avalanche/Source/AvalancheModifiers/Private/Modifiers/AvaVisibilityModifier.cpp b/Engine/Plugins/VirtualProduction/Avalanche/Source/AvalancheModifiers/Private/Modifiers/AvaVisibilityModifier.cpp
--- a/Engine/Plugins/VirtualProduction/Avalanche/Source/AvalancheModifiers/Private/Modifiers/AvaVisibilityModifier.cpp
+++ b/Engine/Plugins/VirtualProduction/Avalanche/Source/AvalancheModifiers/Private/Modifiers/AvaVisibilityModifier.cpp
@@ -121,7 +121,16 @@ void UAvaVisibilityModifier::Apply()
 	ChildrenActorsWeak = NewChildrenActorsWeak;
 
 	FActorModifierRenderStateUpdateExtension* RenderStateExtension = GetExtension<FActorModifierRenderStateUpdateExtension>();
-	RenderStateExtension->SetTrackedActorsVisibility(ChildrenActorsWeak);
+
+	// Reset tracked actors state to detect change in visibility later
+	for (const TWeakObjectPtr<AActor>& ChildrenActorWeak : ChildrenActorsWeak)
+	{
+		if (AActor* ChildrenActor = ChildrenActorWeak.Get())
+		{
+			RenderStateExtension->UntrackActorVisibility(ChildrenActor);
+			RenderStateExtension->TrackActorVisibility(ChildrenActor);
+		}
+	}
 
 	Next();
 }
