Motion Design: added tag soft handle container

diff --git a/Engine/Plugins/VirtualProduction/Avalanche/Source/AvalancheTag/Private/AvaTagSoftHandleContainer.cpp b/Engine/Plugins/VirtualProduction/Avalanche/Source/AvalancheTag/Private/AvaTagSoftHandleContainer.cpp
new file mode 100644
--- /dev/null
+++ b/Engine/Plugins/VirtualProduction/Avalanche/Source/AvalancheTag/Private/AvaTagSoftHandleContainer.cpp
@@ -0,0 +1,179 @@
+// Copyright Epic Games, Inc. All Rights Reserved.
+
+#include "AvaTagSoftHandleContainer.h"
+#include "AvaTagCollection.h"
+#include "AvaTagHandle.h"
+#include "AvaTagHandleContainer.h"
+#include "AvaTagList.h"
+#include "AvaTagSoftHandle.h"
+
+FAvaTagSoftHandleContainer::FAvaTagSoftHandleContainer(const FAvaTagHandle& InTagHandle)
+	: Source(InTagHandle.Source)
+	, TagIds({ InTagHandle.TagId })
+{
+}
+
+bool FAvaTagSoftHandleContainer::ContainsTag(const FAvaTagHandle& InTagHandle) const
+{
+	if (ContainsTagHandle(InTagHandle))
+	{
+		return true;
+	}
+
+	const UAvaTagCollection* const ResolvedSource = ResolveSource();
+	if (!ResolvedSource)
+	{
+		return false;
+	}
+
+	// Populate the Tag Set with resolved Tags
+	TSet<FAvaTag> OtherTagSet;
+	{
+		FAvaTagList OtherTagList = InTagHandle.GetTags();
+		if (OtherTagList.Tags.IsEmpty())
+		{
+			return false;
+		}
+
+		OtherTagSet.Reserve(OtherTagList.Tags.Num());
+		for (const FAvaTag* OtherTag : OtherTagList)
+		{
+			OtherTagSet.Add(*OtherTag);
+		}
+	}
+
+	for (const FAvaTagId& TagId : TagIds)
+	{
+		for (const FAvaTag* Tag : ResolvedSource->GetTags(TagId))
+		{
+			if (OtherTagSet.Contains(*Tag))
+			{
+				return true;
+			}
+		}
+	}
+
+	return false;
+}
+
+bool FAvaTagSoftHandleContainer::ContainsTagHandle(const FAvaTagHandle& InTagHandle) const
+{
+	return Source == InTagHandle.Source && TagIds.Contains(InTagHandle.TagId);
+}
+
+FString FAvaTagSoftHandleContainer::ToString() const
+{
+	const UAvaTagCollection* const ResolvedSource = ResolveSource();
+	if (!ResolvedSource)
+	{
+		return FString();
+	}
+
+	TStringBuilder<32> StringBuilder;
+	for (const FAvaTagId& TagId : TagIds)
+	{
+		for (const FAvaTag* Tag : ResolvedSource->GetTags(TagId))
+		{
+			if (StringBuilder.Len() > 0)
+			{
+				StringBuilder.Append(TEXT(", "));
+			}
+			StringBuilder.Append(Tag->ToString());
+		}
+	}
+	return FString(StringBuilder);
+}
+
+void FAvaTagSoftHandleContainer::PostSerialize(const FArchive& Ar)
+{
+	if (Ar.IsSaving())
+	{
+		for (const FAvaTagId& TagId : TagIds)
+		{
+			if (TagId.IsValid())
+			{
+				Ar.MarkSearchableName(FAvaTagId::StaticStruct(), *TagId.ToString());
+			}
+		}
+	}
+}
+
+bool FAvaTagSoftHandleContainer::SerializeFromMismatchedTag(const FPropertyTag& InPropertyTag, FStructuredArchive::FSlot InSlot)
+{
+	if (InPropertyTag.GetType().IsStruct(FAvaTagSoftHandle::StaticStruct()->GetFName()))
+	{
+		FAvaTagSoftHandle TagSoftHandle;
+		FAvaTagSoftHandle::StaticStruct()->SerializeItem(InSlot, &TagSoftHandle, nullptr);
+
+		if (TagSoftHandle.IsValid())
+		{
+			Source = TagSoftHandle.Source;
+			TagIds = { TagSoftHandle.TagId };
+		}
+		return true;
+	}
+
+	if (InPropertyTag.GetType().IsStruct(FAvaTagHandleContainer::StaticStruct()->GetFName()))
+	{
+		FAvaTagHandleContainer TagHandleContainer;
+		FAvaTagHandleContainer::StaticStruct()->SerializeItem(InSlot, &TagHandleContainer, nullptr);
+
+		if (TagHandleContainer.Source)
+		{
+			Source = TagHandleContainer.Source;
+			TagIds = TagHandleContainer.GetTagIds();
+		}
+		return true;
+	}
+
+	return false;
+}
+
+bool FAvaTagSoftHandleContainer::AddTagHandle(const FAvaTagHandle& InTagHandle)
+{
+	if (!Source)
+	{
+		Source = InTagHandle.Source;
+	}
+
+	if (TagIds.Contains(InTagHandle.TagId))
+	{
+		return false;
+	}
+
+	TagIds.Add(InTagHandle.TagId);
+	return true;
+}
+
+bool FAvaTagSoftHandleContainer::RemoveTagHandle(const FAvaTagHandle& InTagHandle)
+{
+	return TagIds.Remove(InTagHandle.TagId) > 0;
+}
+
+TArray<FAvaTag> FAvaTagSoftHandleContainer::ResolveTags() const
+{
+	TArray<FAvaTag> Tags;
+
+	const UAvaTagCollection* const ResolvedSource = ResolveSource();
+	if (!ResolvedSource)
+	{
+		return Tags;
+	}
+
+	Tags.Reserve(TagIds.Num());
+
+	for (const FAvaTagId& TagId : TagIds)
+	{
+		for (const FAvaTag* Tag : ResolvedSource->GetTags(TagId))
+		{
+			Tags.Add(*Tag);
+		}
+	}
+
+	return Tags;
+}
+
+const UAvaTagCollection* FAvaTagSoftHandleContainer::ResolveSource() const
+{
+	return Source.LoadSynchronous();
+}
diff --git a/Engine/Plugins/VirtualProduction/Avalanche/Source/AvalancheTag/Public/AvaTagSoftHandleContainer.h b/Engine/Plugins/VirtualProduction/Avalanche/Source/AvalancheTag/Public/AvaTagSoftHandleContainer.h
new file mode 100644
--- /dev/null
+++ b/Engine/Plugins/VirtualProduction/Avalanche/Source/AvalancheTag/Public/AvaTagSoftHandleContainer.h
@@ -0,0 +1,75 @@
+// Copyright Epic Games, Inc. All Rights Reserved.
+
+#pragma once
+
+#include "AvaTagId.h"
+#include "Containers/ContainersFwd.h"
+#include "AvaTagSoftHandleContainer.generated.h"
+
+#define UE_API AVALANCHETAG_API
+
+class UAvaTagCollection;
+struct FAvaTag;
+struct FAvaTagHandle;
+
+/**
+ * Soft handle to multiple tags in a particular Source.
+ * This should be used when needing to soft reference multiple FAvaTags.
+ */
+USTRUCT(BlueprintType, DisplayName="Motion Design Tag Soft Handle Container")
+struct FAvaTagSoftHandleContainer
+{
+	GENERATED_BODY()
+
+	FAvaTagSoftHandleContainer() = default;
+
+	UE_API explicit FAvaTagSoftHandleContainer(const FAvaTagHandle& InTagHandle);
+
+	/** Returns true if the Tag Handles resolve to same valued FAvaTags, even if the Source or Tag Id is different */
+	UE_API bool ContainsTag(const FAvaTagHandle& InTagHandle) const;
+
+	/** Returns true if the Tag Handles is the exact same as the other (Same Source and Tag Id) */
+	UE_API bool ContainsTagHandle(const FAvaTagHandle& InTagHandle) const;
+
+	UE_API FString ToString() const;
+
+	UE_API void PostSerialize(const FArchive& Ar);
+
+	UE_API bool SerializeFromMismatchedTag(const FPropertyTag& InPropertyTag, FStructuredArchive::FSlot InSlot);
+
+	/** Adds the provided Tag Handle to TagIds. Returns true only if it was added as a new entry to TagIds */
+	UE_API bool AddTagHandle(const FAvaTagHandle& InTagHandle);
+
+	/** Removes the provided Tag Handle from TagIds. Returns true only if it existed and was removed from TagIds */
+	UE_API bool RemoveTagHandle(const FAvaTagHandle& InTagHandle);
+
+	/** Returns an array of resolved tags through this container's tag ids and source tag collection */
+	UE_API TArray<FAvaTag> ResolveTags() const;
+
+	/** Resolves the source soft object reference */
+	const UAvaTagCollection* ResolveSource() const;
+
+	TConstArrayView<FAvaTagId> GetTagIds() const
+	{
+		return TagIds;
+	}
+
+	UPROPERTY(EditAnywhere, Category = "Tag")
+	TSoftObjectPtr<const UAvaTagCollection> Source;
+
+private:
+	UPROPERTY(EditAnywhere, Category = "Tag")
+	TArray<FAvaTagId> TagIds;
+};
+
+template<>
+struct TStructOpsTypeTraits<FAvaTagSoftHandleContainer> : public TStructOpsTypeTraitsBase2<FAvaTagSoftHandleContainer>
+{
+	enum
+	{
+		WithPostSerialize = true,
+		WithStructuredSerializeFromMismatchedTag = true,
+	};
+};
+
+#undef UE_API
diff --git a/Engine/Plugins/VirtualProduction/Avalanche/Source/AvalancheTagEditor/Private/AvaTagEditorModule.cpp b/Engine/Plugins/VirtualProduction/Avalanche/Source/AvalancheTagEditor/Private/AvaTagEditorModule.cpp
--- a/Engine/Plugins/VirtualProduction/Avalanche/Source/AvalancheTagEditor/Private/AvaTagEditorModule.cpp
+++ b/Engine/Plugins/VirtualProduction/Avalanche/Source/AvalancheTagEditor/Private/AvaTagEditorModule.cpp
@@ -7,11 +7,13 @@
 #include "AvaTagHandle.h"
 #include "AvaTagHandleContainer.h"
 #include "AvaTagSoftHandle.h"
+#include "AvaTagSoftHandleContainer.h"
 #include "Customization/AvaTagAliasCustomization.h"
 #include "Customization/AvaTagCollectionCustomization.h"
 #include "Customization/AvaTagHandleCustomization.h"
 #include "Customization/TagCustomizers/AvaTagHandleContainerCustomizer.h"
 #include "Customization/TagCustomizers/AvaTagHandleCustomizer.h"
+#include "Customization/TagCustomizers/AvaTagSoftHandleContainerCustomizer.h"
 #include "Customization/TagCustomizers/AvaTagSoftHandleCustomizer.h"
 #include "Modules/ModuleManager.h"
 #include "PropertyEditorModule.h"
@@ -52,6 +54,9 @@ void FAvalancheTagEditorModule::RegisterCustomizations()
 	PropertyEditorModule.RegisterCustomPropertyTypeLayout(CustomizedTypes.Add_GetRef(FAvaTagHandleContainer::StaticStruct()->GetFName())
 		, UE::AvaTagEditor::Private::CreateTagHandleCustomization<FAvaTagHandleContainerCustomizer>());
 
+	PropertyEditorModule.RegisterCustomPropertyTypeLayout(CustomizedTypes.Add_GetRef(FAvaTagSoftHandleContainer::StaticStruct()->GetFName())
+		, UE::AvaTagEditor::Private::CreateTagHandleCustomization<FAvaTagSoftHandleContainerCustomizer>());
+
 	PropertyEditorModule.RegisterCustomPropertyTypeLayout(CustomizedTypes.Add_GetRef(FAvaTagAlias::StaticStruct()->GetFName())
 		, FOnGetPropertyTypeCustomizationInstance::CreateLambda([]
 		{
diff --git a/Engine/Plugins/VirtualProduction/Avalanche/Source/AvalancheTagEditor/Private/Customization/TagCustomizers/AvaTagSoftHandleContainerCustomizer.cpp b/Engine/Plugins/VirtualProduction/Avalanche/Source/AvalancheTagEditor/Private/Customization/TagCustomizers/AvaTagSoftHandleContainerCustomizer.cpp
new file mode 100644
--- /dev/null
+++ b/Engine/Plugins/VirtualProduction/Avalanche/Source/AvalancheTagEditor/Private/Customization/TagCustomizers/AvaTagSoftHandleContainerCustomizer.cpp
@@ -0,0 +1,65 @@
+// Copyright Epic Games, Inc. All Rights Reserved.
+
+#include "AvaTagSoftHandleContainerCustomizer.h"
+#include "AvaTagCollection.h"
+#include "AvaTagHandle.h"
+#include "AvaTagSoftHandleContainer.h"
+#include "PropertyHandle.h"
+#include "ScopedTransaction.h"
+
+#define LOCTEXT_NAMESPACE "AvaTagSoftHandleContainerCustomizer"
+
+TSharedPtr<IPropertyHandle> FAvaTagSoftHandleContainerCustomizer::GetTagCollectionHandle(const TSharedRef<IPropertyHandle>& InStructHandle) const
+{
+	return InStructHandle->GetChildHandle(GET_MEMBER_NAME_CHECKED(FAvaTagSoftHandleContainer, Source));
+}
+
+const UAvaTagCollection* FAvaTagSoftHandleContainerCustomizer::GetOrLoadTagCollection(const void* InStructRawData) const
+{
+	return static_cast<const FAvaTagSoftHandleContainer*>(InStructRawData)->Source.LoadSynchronous();
+}
+
+void FAvaTagSoftHandleContainerCustomizer::SetTagHandleAdded(const TSharedRef<IPropertyHandle>& InContainerProperty, const FAvaTagHandle& InTagHandle, bool bInAdd) const
+{
+	FScopedTransaction Transaction(bInAdd
+		? LOCTEXT("AddTagHandleInContainer", "Add Tag Handle in Container")
+		: LOCTEXT("RemoveTagHandleInContainer", "Remove Tag Handle in Container"));
+
+	InContainerProperty->NotifyPreChange();
+
+	InContainerProperty->EnumerateRawData(
+		[&InTagHandle, bInAdd](void* InStructRawData, const int32, const int32)->bool
+		{
+			FAvaTagSoftHandleContainer* Container = static_cast<FAvaTagSoftHandleContainer*>(InStructRawData);
+
+			ensureMsgf(Container->Source == InTagHandle.Source
+				, TEXT("Unexpected result setting tag handle in container: Container Source (%s) doesn't match Tag Handle Source (%s)")
+				, *Container->Source->GetName()
+				, *GetNameSafe(InTagHandle.Source));
+
+			if (bInAdd)
+			{
+				Container->AddTagHandle(InTagHandle);
+			}
+			else
+			{
+				Container->RemoveTagHandle(InTagHandle);
+			}
+			return true;
+		});
+
+	InContainerProperty->NotifyPostChange(bInAdd ? EPropertyChangeType::ArrayAdd : EPropertyChangeType::ArrayRemove);
+	InContainerProperty->NotifyFinishedChangingProperties();
+}
+
+bool FAvaTagSoftHandleContainerCustomizer::ContainsTagHandle(const void* InStructRawData, const FAvaTagHandle& InTagHandle) const
+{
+	return static_cast<const FAvaTagSoftHandleContainer*>(InStructRawData)->ContainsTagHandle(InTagHandle);
+}
+
+FName FAvaTagSoftHandleContainerCustomizer::GetDisplayValueName(const void* InStructRawData) const
+{
+	return *static_cast<const FAvaTagSoftHandleContainer*>(InStructRawData)->ToString();
+}
+
+#undef LOCTEXT_NAMESPACE
diff --git a/Engine/Plugins/VirtualProduction/Avalanche/Source/AvalancheTagEditor/Private/Customization/TagCustomizers/AvaTagSoftHandleContainerCustomizer.h b/Engine/Plugins/VirtualProduction/Avalanche/Source/AvalancheTagEditor/Private/Customization/TagCustomizers/AvaTagSoftHandleContainerCustomizer.h
new file mode 100644
--- /dev/null
+++ b/Engine/Plugins/VirtualProduction/Avalanche/Source/AvalancheTagEditor/Private/Customization/TagCustomizers/AvaTagSoftHandleContainerCustomizer.h
@@ -0,0 +1,22 @@
+// Copyright Epic Games, Inc. All Rights Reserved.
+
+#pragma once
+
+#include "IAvaTagHandleCustomizer.h"
+
+/** Customizer for FAvaTagSoftHandleContainer */
+class FAvaTagSoftHandleContainerCustomizer : public IAvaTagHandleCustomizer
+{
+public:
+	//~ Begin IAvaTagHandleCustomizer
+	virtual TSharedPtr<IPropertyHandle> GetTagCollectionHandle(const TSharedRef<IPropertyHandle>& InStructHandle) const override;
+	virtual const UAvaTagCollection* GetOrLoadTagCollection(const void* InStructRawData) const override;
+	virtual void SetTagHandleAdded(const TSharedRef<IPropertyHandle>& InContainerProperty, const FAvaTagHandle& InTagHandle, bool bInAdd) const override;
+	virtual bool ContainsTagHandle(const void* InStructRawData, const FAvaTagHandle& InTagHandle) const override;
+	virtual FName GetDisplayValueName(const void* InStructRawData) const override;
+	virtual bool AllowMultipleTags() const override
+	{
+		return true;
+	}
+	//~ End IAvaTagHandleCustomizer
+};
