Motion Design: Fix preview in command doing preview frame instead of preview in. Added preview frame command to also support preview frame.

diff --git a/Engine/Plugins/VirtualProduction/Avalanche/Source/AvalancheMediaEditor/Private/Rundown/AvaRundownEditor.cpp b/Engine/Plugins/VirtualProduction/Avalanche/Source/AvalancheMediaEditor/Private/Rundown/AvaRundownEditor.cpp
--- a/Engine/Plugins/VirtualProduction/Avalanche/Source/AvalancheMediaEditor/Private/Rundown/AvaRundownEditor.cpp
+++ b/Engine/Plugins/VirtualProduction/Avalanche/Source/AvalancheMediaEditor/Private/Rundown/AvaRundownEditor.cpp
@@ -1313,10 +1313,12 @@ void FAvaRundownEditor::UnloadPageCommand(const TArray<FString>& InArgs)
 	}
 }
 
-void FAvaRundownEditor::PlayPageCommand(const TArray<FString>& InArgs, bool bInPreview)
+void FAvaRundownEditor::PlayPageCommand(const TArray<FString>& InArgs, EAvaRundownPagePlayType InPlayType)
 {
-	FString Errors;	
-	const TArray<int32> PageIds = ArgumentsToPageIds(InArgs, bInPreview, Errors);
+	const bool bPreview = UE::AvaRundown::IsPreviewPlayType(InPlayType);
+
+	FString Errors;
+	const TArray<int32> PageIds = ArgumentsToPageIds(InArgs, bPreview, Errors);
 
 	if (!Errors.IsEmpty())
 	{
@@ -1327,7 +1329,7 @@ void FAvaRundownEditor::PlayPageCommand(const TArray<FString>& InArgs, bool bInP
 
 	if (Rundown && !PageIds.IsEmpty())
 	{
-		Rundown->PlayPages(PageIds, bInPreview ? EAvaRundownPagePlayType::PreviewFromFrame : EAvaRundownPagePlayType::PlayFromStart);
+		Rundown->PlayPages(PageIds, InPlayType);
 	}
 }
 
@@ -1638,12 +1640,12 @@ const FAvaRundownEditor::FBindableMacroCommands& FAvaRundownEditor::GetBindableM
 		using namespace UE::AvaRundownEditor::Private;
 		BindableMacroCommands.Add(GetName(EAvaRundownEditorMacroCommand::LoadPage), [this](const TArray<FString>& InArgs){LoadPageCommand(InArgs);});
 		BindableMacroCommands.Add(GetName(EAvaRundownEditorMacroCommand::UnloadPage), [this](const TArray<FString>& InArgs){UnloadPageCommand(InArgs);});
-		BindableMacroCommands.Add(GetName(EAvaRundownEditorMacroCommand::TakeIn), [this](const TArray<FString>& InArgs){PlayPageCommand(InArgs, /*bInPreview*/ false);});
+		BindableMacroCommands.Add(GetName(EAvaRundownEditorMacroCommand::TakeIn), [this](const TArray<FString>& InArgs){PlayPageCommand(InArgs, EAvaRundownPagePlayType::PlayFromStart);});
 		BindableMacroCommands.Add(GetName(EAvaRundownEditorMacroCommand::TakeOut), [this](const TArray<FString>& InArgs){StopPageCommand(InArgs, /*bInPreview*/ false, /*bInForce*/ false);});
 		BindableMacroCommands.Add(GetName(EAvaRundownEditorMacroCommand::ForceTakeOut), [this](const TArray<FString>& InArgs){StopPageCommand(InArgs, /*bInPreview*/ false, /*bInForce*/ true);});
 		BindableMacroCommands.Add(GetName(EAvaRundownEditorMacroCommand::TakeNext), [this](const TArray<FString>& InArgs){PlayNextPageCommand(InArgs, /*bInPreview*/ false);});
 		BindableMacroCommands.Add(GetName(EAvaRundownEditorMacroCommand::Continue), [this](const TArray<FString>& InArgs){ContinuePageCommand(InArgs, /*bInPreview*/ false);});
-		BindableMacroCommands.Add(GetName(EAvaRundownEditorMacroCommand::PreviewIn), [this](const TArray<FString>& InArgs){PlayPageCommand(InArgs, /*bInPreview*/ true);});
+		BindableMacroCommands.Add(GetName(EAvaRundownEditorMacroCommand::PreviewIn), [this](const TArray<FString>& InArgs){PlayPageCommand(InArgs, EAvaRundownPagePlayType::PreviewFromStart);});
 		BindableMacroCommands.Add(GetName(EAvaRundownEditorMacroCommand::PreviewOut), [this](const TArray<FString>& InArgs){StopPageCommand(InArgs, /*bInPreview*/ true, /*bInForce*/ false);});
 		BindableMacroCommands.Add(GetName(EAvaRundownEditorMacroCommand::ForcePreviewOut), [this](const TArray<FString>& InArgs){StopPageCommand(InArgs, /*bInPreview*/ true, /*bInForce*/ true);});
 		BindableMacroCommands.Add(GetName(EAvaRundownEditorMacroCommand::PreviewNext), [this](const TArray<FString>& InArgs){PlayNextPageCommand(InArgs, /*bInPreview*/ true);});
@@ -1652,6 +1654,7 @@ const FAvaRundownEditor::FBindableMacroCommands& FAvaRundownEditor::GetBindableM
 		BindableMacroCommands.Add(GetName(EAvaRundownEditorMacroCommand::StopPageTransition), [this](const TArray<FString>& InArgs){StopPageTransitionCommand(InArgs);});
 		BindableMacroCommands.Add(GetName(EAvaRundownEditorMacroCommand::StartChannel), [this](const TArray<FString>& InArgs){StartChannelCommand(InArgs);});
 		BindableMacroCommands.Add(GetName(EAvaRundownEditorMacroCommand::StopChannel), [this](const TArray<FString>& InArgs){StopChannelCommand(InArgs);});
+		BindableMacroCommands.Add(GetName(EAvaRundownEditorMacroCommand::PreviewFrame), [this](const TArray<FString>& InArgs){PlayPageCommand(InArgs, EAvaRundownPagePlayType::PreviewFromFrame);});
 	}
 	return BindableMacroCommands;
 }
diff --git a/Engine/Plugins/VirtualProduction/Avalanche/Source/AvalancheMediaEditor/Private/Rundown/AvaRundownEditorMacroCommands.h b/Engine/Plugins/VirtualProduction/Avalanche/Source/AvalancheMediaEditor/Private/Rundown/AvaRundownEditorMacroCommands.h
--- a/Engine/Plugins/VirtualProduction/Avalanche/Source/AvalancheMediaEditor/Private/Rundown/AvaRundownEditorMacroCommands.h
+++ b/Engine/Plugins/VirtualProduction/Avalanche/Source/AvalancheMediaEditor/Private/Rundown/AvaRundownEditorMacroCommands.h
@@ -25,7 +25,8 @@ enum class EAvaRundownEditorMacroCommand : uint8
 	TakeToProgram,
 	StopPageTransition,
 	StartChannel,
-	StopChannel
+	StopChannel,
+	PreviewFrame
 };
 
 struct FAvaRundownEditorMacroCommands
diff --git a/Engine/Plugins/VirtualProduction/Avalanche/Source/AvalancheMediaEditor/Public/Rundown/AvaRundownEditor.h b/Engine/Plugins/VirtualProduction/Avalanche/Source/AvalancheMediaEditor/Public/Rundown/AvaRundownEditor.h
--- a/Engine/Plugins/VirtualProduction/Avalanche/Source/AvalancheMediaEditor/Public/Rundown/AvaRundownEditor.h
+++ b/Engine/Plugins/VirtualProduction/Avalanche/Source/AvalancheMediaEditor/Public/Rundown/AvaRundownEditor.h
@@ -129,7 +129,7 @@ protected:
 	
 	void LoadPageCommand(const TArray<FString>& InArgs);
 	void UnloadPageCommand(const TArray<FString>& InArgs);
-	void PlayPageCommand(const TArray<FString>& InArgs, bool bInPreview);
+	void PlayPageCommand(const TArray<FString>& InArgs, EAvaRundownPagePlayType InPlayType);
 	void ContinuePageCommand(const TArray<FString>& InArgs, bool bInPreview);
 	void PlayNextPageCommand(const TArray<FString>& InArgs, bool bInPreview);
 	void StopPageCommand(const TArray<FString>& InArgs, bool bInPreview, bool bInForce);
