Motion Design: added option to DetachLoader instead of ResetLoaders to avoid reloading package synchronously on rundown playable. Even though the level has been streamed in (loaded), the ResetLoader call uses the base package path causing a synchronous package loading (AsyncLoad + FlushAsync) to occur -- which is one of the main visible hitches in MD playback.

diff --git a/Engine/Plugins/VirtualProduction/Avalanche/Source/AvalancheMedia/Internal/Playback/AvaPlaybackUtils.h b/Engine/Plugins/VirtualProduction/Avalanche/Source/AvalancheMedia/Internal/Playback/AvaPlaybackUtils.h
--- a/Engine/Plugins/VirtualProduction/Avalanche/Source/AvalancheMedia/Internal/Playback/AvaPlaybackUtils.h
+++ b/Engine/Plugins/VirtualProduction/Avalanche/Source/AvalancheMedia/Internal/Playback/AvaPlaybackUtils.h
@@ -17,7 +17,7 @@ public:
 	/**
 	 * Reset the loaders of any loaded packages so the client can overwrite the files.
 	*/
-	static void FlushPackageLoading(UPackage* InPackage);
+	static void FlushPackageLoading(UPackage* InPackage, bool bInResetLoaders=false);
 
 	/**
 	 * Checks if the given packages has been deleted on disk.
diff --git a/Engine/Plugins/VirtualProduction/Avalanche/Source/AvalancheMedia/Private/Playback/AvaPlaybackUtils.cpp b/Engine/Plugins/VirtualProduction/Avalanche/Source/AvalancheMedia/Private/Playback/AvaPlaybackUtils.cpp
--- a/Engine/Plugins/VirtualProduction/Avalanche/Source/AvalancheMedia/Private/Playback/AvaPlaybackUtils.cpp
+++ b/Engine/Plugins/VirtualProduction/Avalanche/Source/AvalancheMedia/Private/Playback/AvaPlaybackUtils.cpp
@@ -21,14 +21,37 @@
 #include "Subsystems/AssetEditorSubsystem.h"
 #endif
 
-void FAvaPlaybackUtils::FlushPackageLoading(UPackage* InPackage)
+namespace UE::AvaPlayback::Utils::Private
+{
+	static TAutoConsoleVariable<bool> CVarForceResetLoadersInPackageFlush(
+		TEXT("MotionDesignPlayback.ForceResetLoadersInPackageFlush"),
+		false,
+		TEXT("Forces to ResetLoaders instead of DetachLoader during FlushPackageLoading")
+	);
+}
+
+void FAvaPlaybackUtils::FlushPackageLoading(UPackage* InPackage, bool bInResetLoaders)
 {
 	if (!InPackage->IsFullyLoaded())
 	{
 		FlushAsyncLoading();
 		InPackage->FullyLoad();
 	}
-	ResetLoaders(InPackage);
+
+	if (bInResetLoaders || UE::AvaPlayback::Utils::Private::CVarForceResetLoadersInPackageFlush.GetValueOnGameThread())
+	{
+		ResetLoaders(InPackage);
+	}
+	else if (FLinkerLoad* LinkerLoad = FLinkerLoad::FindExistingLinkerForPackage(InPackage))
+	{
+		if (!LinkerLoad->IsDestroyingLoader())
+		{
+			// Ideally there should be an option for the Linker to not lock the package files preventing saves.
+			// Resetting the loader will cause the level to be reloaded via ConditionalFlushAsyncLoadingForLinkers, hitching the rundown playback.
+			// DetachLoader allows releasing the lock on the file handle so the package can be saved.
+			LinkerLoad->DetachLoader();
+		}
+	}
 }
 
 bool FAvaPlaybackUtils::IsPackageDeleted(const UPackage* InExistingPackage)
