Motion Design: updated comments to FlushPackageLoading

diff --git a/Engine/Plugins/VirtualProduction/Avalanche/Source/AvalancheMedia/Internal/Playback/AvaPlaybackUtils.h b/Engine/Plugins/VirtualProduction/Avalanche/Source/AvalancheMedia/Internal/Playback/AvaPlaybackUtils.h
--- a/Engine/Plugins/VirtualProduction/Avalanche/Source/AvalancheMedia/Internal/Playback/AvaPlaybackUtils.h
+++ b/Engine/Plugins/VirtualProduction/Avalanche/Source/AvalancheMedia/Internal/Playback/AvaPlaybackUtils.h
@@ -15,7 +15,7 @@ class AVALANCHEMEDIA_API FAvaPlaybackUtils
 {
 public:
 	/**
-	 * Reset the loaders of any loaded packages so the client can overwrite the files.
+	 * Flushes loaded packages so the client can write to the files.
 	*/
 	static void FlushPackageLoading(UPackage* InPackage, bool bInResetLoaders=false);
 
