Avalanche - Disabling outliner faulty tests

diff --git a/Engine/Plugins/VirtualProduction/Avalanche/Source/AvalancheOutliner/Private/Tests/AvaOutliner.spec.cpp b/Engine/Plugins/VirtualProduction/Avalanche/Source/AvalancheOutliner/Private/Tests/AvaOutliner.spec.cpp
--- a/Engine/Plugins/VirtualProduction/Avalanche/Source/AvalancheOutliner/Private/Tests/AvaOutliner.spec.cpp
+++ b/Engine/Plugins/VirtualProduction/Avalanche/Source/AvalancheOutliner/Private/Tests/AvaOutliner.spec.cpp
@@ -93,6 +93,7 @@ void FAvaOutlinerSpec::Define()
 			TestItems = OutlinerProvider->GetOutlinerItems(TestObjects);
 		});
 
+#if 0 // DISABLED: Flaky in CI — temporarily disable
 	Describe("Spawning an Actor", [this]()
 		{
 			BeforeEach([this]()
@@ -113,7 +114,9 @@ void FAvaOutlinerSpec::Define()
 					TestTrue(TEXT("Actor item matches the spawned actor"), Actor == OutlinerProvider->GetTestSpawnActor());
 				});
 		});
+#endif
 
+#if 0 // DISABLED: Flaky in CI — temporarily disable
 	Describe("Duplicating Selected Items", [this]()
 		{
 			When(EWhenType::NoItemsSelected, [this]()
@@ -209,7 +212,9 @@ void FAvaOutlinerSpec::Define()
 						});
 				});
 		});
+#endif
 
+#if 0 // DISABLED: Flaky in CI — temporarily disable
 	Describe("Grouping Selected Items", [this]()
 		{
 			When("grouping with null", [this]()
@@ -304,6 +309,7 @@ void FAvaOutlinerSpec::Define()
 						});
 				});
 		});
+#endif
 
 	AfterEach([this]()
 		{
