Motion Design: added material bridge to abstract how materials are gathered per component. This is used for Rundown to gather all used materials in a level without needing to know about specific implementations like Text3D/Shapes/etc that handle materials differently than primitives.

diff --git a/Engine/Plugins/VirtualProduction/Avalanche/Avalanche.uplugin b/Engine/Plugins/VirtualProduction/Avalanche/Avalanche.uplugin
--- a/Engine/Plugins/VirtualProduction/Avalanche/Avalanche.uplugin
+++ b/Engine/Plugins/VirtualProduction/Avalanche/Avalanche.uplugin
@@ -242,6 +242,11 @@
 			"Name": "AvalancheSceneRigEditor",
 			"Type": "Editor",
 			"LoadingPhase": "Default"
+		},
+		{
+			"Name": "AvalancheMaterial",
+			"Type": "Runtime",
+			"LoadingPhase": "Default"
 		}
 	],
 	"Plugins": [
diff --git a/Engine/Plugins/VirtualProduction/Avalanche/Source/AvalancheCore/Private/AvaDataView.cpp b/Engine/Plugins/VirtualProduction/Avalanche/Source/AvalancheCore/Private/AvaDataView.cpp
new file mode 100644
--- /dev/null
+++ b/Engine/Plugins/VirtualProduction/Avalanche/Source/AvalancheCore/Private/AvaDataView.cpp
@@ -0,0 +1,54 @@
+// Copyright Epic Games, Inc. All Rights Reserved.
+
+#include "AvaDataView.h"
+
+namespace UE::Ava
+{
+
+FDataView::FDataView(const UStruct* InStruct, void* InMemory)
+	: Struct(InStruct)
+	, Memory(InMemory)
+{
+	// Avoid cases where a non-null memory is provided with a null struct.
+	check(!Memory || (Memory && Struct));
+}
+
+FDataView::FDataView(UObject* Object)
+	: FDataView(Object ? Object->GetClass() : nullptr, Object)
+{
+}
+
+bool FDataView::IsValid() const
+{
+	return Memory && Struct;
+}
+
+bool FDataView::IsValidFor(const UStruct* InStruct) const
+{
+	return IsValid() && Struct->IsChildOf(InStruct);
+}
+
+FConstDataView::FConstDataView(const UStruct* InStruct, const void* InMemory)
+	: Struct(InStruct)
+	, Memory(InMemory)
+{
+	// Avoid cases where a non-null memory is provided with a null struct.
+	check(!Memory || (Memory && Struct));
+}
+
+FConstDataView::FConstDataView(const UObject* Object)
+	: FConstDataView(Object ? Object->GetClass() : nullptr, Object)
+{
+}
+
+bool FConstDataView::IsValid() const
+{
+	return Memory && Struct;
+}
+
+bool FConstDataView::IsValidFor(const UStruct* InStruct) const
+{
+	return IsValid() && Struct->IsChildOf(InStruct);
+}
+
+} // UE::Ava
diff --git a/Engine/Plugins/VirtualProduction/Avalanche/Source/AvalancheCore/Public/AvaDataView.h b/Engine/Plugins/VirtualProduction/Avalanche/Source/AvalancheCore/Public/AvaDataView.h
new file mode 100644
--- /dev/null
+++ b/Engine/Plugins/VirtualProduction/Avalanche/Source/AvalancheCore/Public/AvaDataView.h
@@ -0,0 +1,161 @@
+// Copyright Epic Games, Inc. All Rights Reserved.
+
+#pragma once
+
+#include "StructUtils/StructUtils.h"
+#include "StructUtils/StructUtilsTypes.h"
+#include "UObject/Class.h"
+
+#define UE_API AVALANCHECORE_API
+
+namespace UE::Ava
+{
+
+/** Short-lived pointer to an object or struct. */
+struct FDataView
+{
+	FDataView() = default;
+
+	/**
+	 * Generic struct constructor.
+	 * Valid UStruct is required when passing in valid memory pointer.
+	 * Both can be null and the constructed data view will be considered as invalid.
+	 * @see IsValid
+	 */
+	UE_API explicit FDataView(const UStruct* InStruct, void* InMemory);
+
+	/** UObject constructor. */
+	UE_API explicit FDataView(UObject* Object);
+
+	/**
+	 * Check is the view is valid (both pointer and type are set).
+	 * @return True if the view is valid.
+	 */
+	UE_API bool IsValid() const;
+
+	/** Returns true if both the type and memory are set for the given struct. */
+	UE_API bool IsValidFor(const UStruct* InStruct) const;
+
+	/** Returns true if both the type and memory are set for the given struct. */
+	template<typename T>
+	bool IsValidFor() const
+	{
+		return this->IsValidFor(StructUtils::GetAsUStruct<T>());
+	}
+
+	template<typename T>
+	const T* GetPtr() const
+	{
+		return GetMutablePtr<T>();
+	}
+
+	template<typename T>
+	const T& Get() const
+	{
+		return GetMutable<T>();
+	}
+
+	template<typename T>
+	T* GetMutablePtr() const
+	{
+		if (IsValidFor<T>())
+		{
+			return static_cast<T*>(Memory);
+		}
+		return nullptr;
+	}
+
+	template<typename T>
+	T& GetMutable() const
+	{
+		check(IsValidFor<T>());
+		return *static_cast<T*>(Memory);
+	}
+
+	const UStruct* GetStruct() const
+	{
+		return Struct;
+	}
+
+	const void* GetMemory() const
+	{
+		return GetMutableMemory();
+	}
+
+	void* GetMutableMemory() const
+	{
+		return Memory;
+	}
+
+protected:
+	/** UClass or UScriptStruct of the data. */
+	const UStruct* Struct = nullptr;
+	/** Memory pointing at the class or struct */
+	void* Memory = nullptr;
+};
+
+/** Short-lived const pointer to an object or struct. */
+struct FConstDataView
+{
+	FConstDataView() = default;
+
+	/**
+	 * Generic struct constructor.
+	 * Valid UStruct is required when passing in valid memory pointer.
+	 * Both can be null and the constructed data view will be considered as invalid.
+	 * @see IsValid
+	 */
+	UE_API explicit FConstDataView(const UStruct* InStruct, const void* InMemory);
+
+	/** UObject constructor. */
+	UE_API explicit FConstDataView(const UObject* Object);
+
+	/** Returns true if both the type and memory are set. */
+	UE_API bool IsValid() const;
+
+	/** Returns true if both the type and memory are set for the given struct. */
+	UE_API bool IsValidFor(const UStruct* InStruct) const;
+
+	template<typename T>
+	bool IsValidFor() const
+	{
+		return this->IsValidFor(StructUtils::GetAsUStruct<T>());
+	}
+
+	template<typename T> requires (std::is_const_v<T>)
+	T* GetPtr() const
+	{
+		if (IsValidFor<T>())
+		{
+			return static_cast<T*>(Memory);
+		}
+		return nullptr;
+	}
+
+	template<typename T> requires (std::is_const_v<T>)
+	T& Get() const
+	{
+		check(IsValidFor<T>());
+		return *static_cast<T*>(Memory);
+	}
+
+	const UStruct* GetStruct() const
+	{
+		return Struct;
+	}
+
+	const void* GetMemory() const
+	{
+		return Memory;
+	}
+
+protected:
+	/** UClass or UScriptStruct of the data. */
+	const UStruct* Struct = nullptr;
+	/** Const memory pointing at the class or struct */
+	const void* Memory = nullptr;
+};
+	
+} // UE::Ava
+
+#undef UE_API
diff --git a/Engine/Plugins/VirtualProduction/Avalanche/Source/AvalancheMaterial/AvalancheMaterial.Build.cs b/Engine/Plugins/VirtualProduction/Avalanche/Source/AvalancheMaterial/AvalancheMaterial.Build.cs
new file mode 100644
--- /dev/null
+++ b/Engine/Plugins/VirtualProduction/Avalanche/Source/AvalancheMaterial/AvalancheMaterial.Build.cs
@@ -0,0 +1,21 @@
+// Copyright Epic Games, Inc. All Rights Reserved.
+
+using UnrealBuildTool;
+
+public class AvalancheMaterial : ModuleRules
+{
+    public AvalancheMaterial(ReadOnlyTargetRules Target) : base(Target)
+    {
+        PCHUsage = ModuleRules.PCHUsageMode.UseExplicitOrSharedPCHs;
+
+        PublicDependencyModuleNames.AddRange(
+            new string[]
+            {
+                "AvalancheCore",
+                "Core",
+                "CoreUObject",
+                "Engine",
+            }
+        );
+    }
+}
diff --git a/Engine/Plugins/VirtualProduction/Avalanche/Source/AvalancheMaterial/Private/AvaActorMaterialBridge.cpp b/Engine/Plugins/VirtualProduction/Avalanche/Source/AvalancheMaterial/Private/AvaActorMaterialBridge.cpp
new file mode 100644
--- /dev/null
+++ b/Engine/Plugins/VirtualProduction/Avalanche/Source/AvalancheMaterial/Private/AvaActorMaterialBridge.cpp
@@ -0,0 +1,31 @@
+// Copyright Epic Games, Inc. All Rights Reserved.
+
+#include "AvaActorMaterialBridge.h"
+#include "AvaMaterialBridgeRegistry.h"
+#include "GameFramework/Actor.h"
+
+namespace UE::Ava
+{
+
+const UStruct* FActorMaterialBridge::OnGetBridgedType() const
+{
+	return AActor::StaticClass();
+}
+
+void FActorMaterialBridge::OnForEachUsedMaterial(FConstDataView InDataView, TFunctionRef<void(UMaterialInterface*)> InFunc, const FGetUsedMaterialOptions& InOptions) const
+{
+	const FMaterialBridgeRegistry& MaterialBridgeRegistry = FMaterialBridgeRegistry::Get();
+
+	const AActor& Actor = InDataView.Get<const AActor>();
+
+	Actor.ForEachComponent(InOptions.bIncludeAttachedMaterials,
+		[&MaterialBridgeRegistry, &InFunc, &InOptions](UActorComponent* InComponent)
+		{
+			if (const FMaterialBridge* MaterialBridge = MaterialBridgeRegistry.GetMaterialBridgeForObject(InComponent))
+			{
+				MaterialBridge->ForEachUsedMaterial(FConstDataView(InComponent), InFunc, InOptions);
+			}
+		});
+}
+
+} // UE::Ava
diff --git a/Engine/Plugins/VirtualProduction/Avalanche/Source/AvalancheMaterial/Private/AvaDecalComponentMaterialBridge.cpp b/Engine/Plugins/VirtualProduction/Avalanche/Source/AvalancheMaterial/Private/AvaDecalComponentMaterialBridge.cpp
new file mode 100644
--- /dev/null
+++ b/Engine/Plugins/VirtualProduction/Avalanche/Source/AvalancheMaterial/Private/AvaDecalComponentMaterialBridge.cpp
@@ -0,0 +1,21 @@
+// Copyright Epic Games, Inc. All Rights Reserved.
+
+#include "AvaDecalComponentMaterialBridge.h"
+#include "AvaDataView.h"
+#include "Components/DecalComponent.h"
+
+namespace UE::Ava
+{
+
+const UStruct* FDecalComponentMaterialBridge::OnGetBridgedType() const
+{
+	return UDecalComponent::StaticClass();
+}
+
+void FDecalComponentMaterialBridge::OnForEachUsedMaterial(FConstDataView InDataView, TFunctionRef<void(UMaterialInterface*)> InFunc, const FGetUsedMaterialOptions& InOptions) const
+{
+	const UDecalComponent& Component = InDataView.Get<const UDecalComponent>();
+	InFunc(Component.GetDecalMaterial());
+}
+
+} // UE::Ava
diff --git a/Engine/Plugins/VirtualProduction/Avalanche/Source/AvalancheMaterial/Private/AvaMaterialBridge.cpp b/Engine/Plugins/VirtualProduction/Avalanche/Source/AvalancheMaterial/Private/AvaMaterialBridge.cpp
new file mode 100644
--- /dev/null
+++ b/Engine/Plugins/VirtualProduction/Avalanche/Source/AvalancheMaterial/Private/AvaMaterialBridge.cpp
@@ -0,0 +1,26 @@
+// Copyright Epic Games, Inc. All Rights Reserved.
+
+#include "AvaMaterialBridge.h"
+
+namespace UE::Ava
+{
+
+void FMaterialBridge::Initialize()
+{
+	BridgedType = OnGetBridgedType();
+}
+
+const UStruct* FMaterialBridge::GetBridgedType() const
+{
+	return BridgedType;
+}
+
+void FMaterialBridge::ForEachUsedMaterial(FConstDataView InDataView, TFunctionRef<void(UMaterialInterface*)> InFunc, const FGetUsedMaterialOptions& InOptions) const
+{
+	if (InDataView.IsValidFor(BridgedType))
+	{
+		OnForEachUsedMaterial(InDataView, InFunc, InOptions);
+	}
+}
+
+} // UE::Ava
diff --git a/Engine/Plugins/VirtualProduction/Avalanche/Source/AvalancheMaterial/Private/AvaMaterialBridgeRegistry.cpp b/Engine/Plugins/VirtualProduction/Avalanche/Source/AvalancheMaterial/Private/AvaMaterialBridgeRegistry.cpp
new file mode 100644
--- /dev/null
+++ b/Engine/Plugins/VirtualProduction/Avalanche/Source/AvalancheMaterial/Private/AvaMaterialBridgeRegistry.cpp
@@ -0,0 +1,62 @@
+// Copyright Epic Games, Inc. All Rights Reserved.
+
+#include "AvaMaterialBridgeRegistry.h"
+#include "AvaMaterialBridge.h"
+
+namespace UE::Ava
+{
+
+const FMaterialBridgeRegistry& FMaterialBridgeRegistry::Get()
+{
+	return GetMutable();
+}
+
+FMaterialBridgeRegistry& FMaterialBridgeRegistry::GetMutable()
+{
+	static FMaterialBridgeRegistry MaterialBridgeRegistry;
+	return MaterialBridgeRegistry;
+}
+
+const FMaterialBridge* FMaterialBridgeRegistry::GetMaterialBridgeForObject(const UObject* InObject) const
+{
+	if (InObject)
+	{
+		return GetMaterialBridgeForType(InObject->GetClass());
+	}
+	return nullptr;
+}
+
+const FMaterialBridge* FMaterialBridgeRegistry::GetMaterialBridgeForType(const UStruct* InStruct) const
+{
+	if (!InStruct)
+	{
+		return nullptr;
+	}
+	for (const UStruct* Struct : InStruct->GetSuperStructIterator())
+	{
+		if (const TUniquePtr<FMaterialBridge>* MaterialBridge = RegisteredMaterialBridges.Find(Struct))
+		{
+			return MaterialBridge->Get();
+		}
+	}
+	return nullptr;
+}
+
+
+void FMaterialBridgeRegistry::RegisterInternal(TUniquePtr<UE::Ava::FMaterialBridge>&& InMaterialBridge)
+{
+	if (!ensureMsgf(UObjectInitialized(), TEXT("Cannot register Material Bridges. UObject system not initialized!")))
+	{
+		return;
+	}
+
+	InMaterialBridge->Initialize();
+	const UStruct* const BridgedType = InMaterialBridge->GetBridgedType();
+
+	if (ensureMsgf(BridgedType, TEXT("Cannot register Material Bridge with a null bridged type!")))
+	{
+		RegisteredMaterialBridges.Add(BridgedType, MoveTemp(InMaterialBridge));
+	}
+}
+
+} // UE::Ava
diff --git a/Engine/Plugins/VirtualProduction/Avalanche/Source/AvalancheMaterial/Private/AvaMaterialModule.cpp b/Engine/Plugins/VirtualProduction/Avalanche/Source/AvalancheMaterial/Private/AvaMaterialModule.cpp
new file mode 100644
--- /dev/null
+++ b/Engine/Plugins/VirtualProduction/Avalanche/Source/AvalancheMaterial/Private/AvaMaterialModule.cpp
@@ -0,0 +1,21 @@
+// Copyright Epic Games, Inc. All Rights Reserved.
+
+#include "AvaMaterialModule.h"
+#include "AvaActorMaterialBridge.h"
+#include "AvaDecalComponentMaterialBridge.h"
+#include "AvaMaterialBridge.h"
+#include "AvaMaterialBridgeRegistry.h"
+#include "AvaPrimitiveComponentMaterialBridge.h"
+#include "Modules/ModuleManager.h"
+
+IMPLEMENT_MODULE(FAvaMaterialModule, AvalancheMaterial)
+
+void FAvaMaterialModule::StartupModule()
+{
+	UE::Ava::FMaterialBridgeRegistry& MaterialBridgeRegistry = UE::Ava::FMaterialBridgeRegistry::GetMutable();
+
+	// Register built-in material bridges
+	MaterialBridgeRegistry.Register<UE::Ava::FActorMaterialBridge>();
+	MaterialBridgeRegistry.Register<UE::Ava::FPrimitiveComponentMaterialBridge>();
+	MaterialBridgeRegistry.Register<UE::Ava::FDecalComponentMaterialBridge>();
+}
diff --git a/Engine/Plugins/VirtualProduction/Avalanche/Source/AvalancheMaterial/Private/AvaMaterialModule.h b/Engine/Plugins/VirtualProduction/Avalanche/Source/AvalancheMaterial/Private/AvaMaterialModule.h
new file mode 100644
--- /dev/null
+++ b/Engine/Plugins/VirtualProduction/Avalanche/Source/AvalancheMaterial/Private/AvaMaterialModule.h
@@ -0,0 +1,13 @@
+// Copyright Epic Games, Inc. All Rights Reserved.
+
+#pragma once
+
+#include "Modules/ModuleInterface.h"
+
+class FAvaMaterialModule : public IModuleInterface
+{
+public:
+	//~ Begin IModuleInterface
+	virtual void StartupModule() override;
+	//~ End IModuleInterface
+};
diff --git a/Engine/Plugins/VirtualProduction/Avalanche/Source/AvalancheMaterial/Private/AvaPrimitiveComponentMaterialBridge.cpp b/Engine/Plugins/VirtualProduction/Avalanche/Source/AvalancheMaterial/Private/AvaPrimitiveComponentMaterialBridge.cpp
new file mode 100644
--- /dev/null
+++ b/Engine/Plugins/VirtualProduction/Avalanche/Source/AvalancheMaterial/Private/AvaPrimitiveComponentMaterialBridge.cpp
@@ -0,0 +1,28 @@
+// Copyright Epic Games, Inc. All Rights Reserved.
+
+#include "AvaPrimitiveComponentMaterialBridge.h"
+#include "AvaDataView.h"
+#include "Components/PrimitiveComponent.h"
+
+namespace UE::Ava
+{
+
+const UStruct* FPrimitiveComponentMaterialBridge::OnGetBridgedType() const
+{
+	return UPrimitiveComponent::StaticClass();
+}
+
+void FPrimitiveComponentMaterialBridge::OnForEachUsedMaterial(FConstDataView InDataView, TFunctionRef<void(UMaterialInterface*)> InFunc, const FGetUsedMaterialOptions& InOptions) const
+{
+	const UPrimitiveComponent& Component = InDataView.Get<const UPrimitiveComponent>();
+
+	TArray<UMaterialInterface*> Materials;
+	Component.GetUsedMaterials(Materials, InOptions.bIncludeDebugMaterials);
+
+	for (UMaterialInterface* Material : Materials)
+	{
+		InFunc(Material);
+	}
+}
+
+} // UE::Ava
diff --git a/Engine/Plugins/VirtualProduction/Avalanche/Source/AvalancheMaterial/Public/AvaActorMaterialBridge.h b/Engine/Plugins/VirtualProduction/Avalanche/Source/AvalancheMaterial/Public/AvaActorMaterialBridge.h
new file mode 100644
--- /dev/null
+++ b/Engine/Plugins/VirtualProduction/Avalanche/Source/AvalancheMaterial/Public/AvaActorMaterialBridge.h
@@ -0,0 +1,24 @@
+// Copyright Epic Games, Inc. All Rights Reserved.
+
+#pragma once
+
+#include "AvaMaterialBridge.h"
+
+#define UE_API AVALANCHEMATERIAL_API
+
+namespace UE::Ava
+{
+
+/** Material Bridge for Actors */
+class FActorMaterialBridge : public FMaterialBridge
+{
+protected:
+	//~ Begin FMaterialBridge
+	UE_API virtual const UStruct* OnGetBridgedType() const override;
+	UE_API virtual void OnForEachUsedMaterial(FConstDataView InDataView, TFunctionRef<void(UMaterialInterface*)> InFunc, const FGetUsedMaterialOptions& InOptions) const override;
+	//~ End FMaterialBridge
+};
+
+} // UE::Ava
+
+#undef UE_API
diff --git a/Engine/Plugins/VirtualProduction/Avalanche/Source/AvalancheMaterial/Public/AvaDecalComponentMaterialBridge.h b/Engine/Plugins/VirtualProduction/Avalanche/Source/AvalancheMaterial/Public/AvaDecalComponentMaterialBridge.h
new file mode 100644
--- /dev/null
+++ b/Engine/Plugins/VirtualProduction/Avalanche/Source/AvalancheMaterial/Public/AvaDecalComponentMaterialBridge.h
@@ -0,0 +1,24 @@
+// Copyright Epic Games, Inc. All Rights Reserved.
+
+#pragma once
+
+#include "AvaMaterialBridge.h"
+
+#define UE_API AVALANCHEMATERIAL_API
+
+namespace UE::Ava
+{
+
+/** Material Bridge for Decal Components */
+class FDecalComponentMaterialBridge : public FMaterialBridge
+{
+protected:
+	//~ Begin FMaterialBridge
+	UE_API virtual const UStruct* OnGetBridgedType() const override;
+	UE_API virtual void OnForEachUsedMaterial(FConstDataView InDataView, TFunctionRef<void(UMaterialInterface*)> InFunc, const FGetUsedMaterialOptions& InOptions) const override;
+	//~ End FMaterialBridge
+};
+
+} // UE::Ava
+
+#undef UE_API
diff --git a/Engine/Plugins/VirtualProduction/Avalanche/Source/AvalancheMaterial/Public/AvaMaterialBridge.h b/Engine/Plugins/VirtualProduction/Avalanche/Source/AvalancheMaterial/Public/AvaMaterialBridge.h
new file mode 100644
--- /dev/null
+++ b/Engine/Plugins/VirtualProduction/Avalanche/Source/AvalancheMaterial/Public/AvaMaterialBridge.h
@@ -0,0 +1,62 @@
+// Copyright Epic Games, Inc. All Rights Reserved.
+
+#pragma once
+
+#include "AvaDataView.h"
+#include "Containers/ContainersFwd.h"
+#include "StructUtils/StructUtilsTypes.h"
+#include "UObject/ObjectPtr.h"
+
+#define UE_API AVALANCHEMATERIAL_API
+
+class UMaterialInterface;
+
+namespace UE::Ava
+{
+
+/**
+ * Base class that allows interacting with materials from any type of data view (Actors, ActorComponents, etc).
+ */
+class FMaterialBridge
+{
+public:
+	virtual ~FMaterialBridge() = default;
+
+	/** Initializes the bridge for use */
+	void Initialize();
+
+	/** Retrieves the type that this bridge supports */
+	UE_API const UStruct* GetBridgedType() const;
+
+	struct FGetUsedMaterialOptions
+	{
+		/** Declare default ctor explicitly due to https://gcc.gnu.org/bugzilla/show_bug.cgi?id=88165 */
+		FGetUsedMaterialOptions() {};
+		/** Whether to include materials from attachments (e.g. whether to include from attached actors when the object is an actor) */
+		bool bIncludeAttachedMaterials = false;
+		/** Whether to include debug materials */
+		bool bIncludeDebugMaterials = false;
+	};
+	/**
+	 * Iterates all the used materials in the given data view.
+	 * @param InDataView the data view to find materials from.
+	 * @param InFunc function to call for each material found.
+	 * @param InOptions additional options to get materials.
+	 */
+	UE_API void ForEachUsedMaterial(FConstDataView InDataView, TFunctionRef<void(UMaterialInterface*)> InFunc, const FGetUsedMaterialOptions& InOptions = {}) const;
+
+protected:
+	/** Retrieves the type that this bridge supports. Only called once in Initialize() */
+	virtual const UStruct* OnGetBridgedType() const = 0;
+
+	/** Get all the used materials in the given data view. */
+	virtual void OnForEachUsedMaterial(FConstDataView InDataView, TFunctionRef<void(UMaterialInterface*)> InFunc, const FGetUsedMaterialOptions& InOptions) const = 0;
+
+private:
+	/** The type that this Material Bridge works on */
+	TObjectPtr<const UStruct> BridgedType;
+};
+
+} // UE::Ava
+
+#undef UE_API
diff --git a/Engine/Plugins/VirtualProduction/Avalanche/Source/AvalancheMaterial/Public/AvaMaterialBridgeRegistry.h b/Engine/Plugins/VirtualProduction/Avalanche/Source/AvalancheMaterial/Public/AvaMaterialBridgeRegistry.h
new file mode 100644
--- /dev/null
+++ b/Engine/Plugins/VirtualProduction/Avalanche/Source/AvalancheMaterial/Public/AvaMaterialBridgeRegistry.h
@@ -0,0 +1,49 @@
+// Copyright Epic Games, Inc. All Rights Reserved.
+
+#pragma once
+
+#include "Concepts/DerivedFrom.h"
+#include "Containers/Map.h"
+#include "Templates/SharedPointerFwd.h"
+#include "UObject/ObjectKey.h"
+
+#define UE_API AVALANCHEMATERIAL_API
+
+namespace UE::Ava
+{
+	class FMaterialBridge;
+}
+
+namespace UE::Ava
+{
+
+/** Holds all the Material Bridges to use */
+class FMaterialBridgeRegistry
+{
+public:
+	UE_API static const FMaterialBridgeRegistry& Get();
+	UE_API static FMaterialBridgeRegistry& GetMutable();
+
+	template<typename InMaterialBridgeType, typename... InArgTypes> requires (UE::CDerivedFrom<InMaterialBridgeType, UE::Ava::FMaterialBridge>)
+	void Register(InArgTypes&&... InArgs)
+	{
+		this->RegisterInternal(MakeUnique<InMaterialBridgeType>(InArgs...));
+	}
+
+	/** Tries to find the most relevant material bridge for a given object, returning a Null Material Bridge if object is invalid, or if a relevant material bridge was not found. */
+	UE_API const FMaterialBridge* GetMaterialBridgeForObject(const UObject* InObject) const;
+
+	/** Tries to find the most relevant material bridge for a given type, returning null if a relevant material bridge was not found. */
+	UE_API const FMaterialBridge* GetMaterialBridgeForType(const UStruct* InType) const;
+
+protected:
+	UE_API void RegisterInternal(TUniquePtr<UE::Ava::FMaterialBridge>&& InMaterialBridge);
+
+private:
+	/** Registered material bridges by their supported type */
+	TMap<TObjectKey<UStruct>, TUniquePtr<FMaterialBridge>> RegisteredMaterialBridges;
+};
+
+} // UE::Ava
+
+#undef UE_API
diff --git a/Engine/Plugins/VirtualProduction/Avalanche/Source/AvalancheMaterial/Public/AvaPrimitiveComponentMaterialBridge.h b/Engine/Plugins/VirtualProduction/Avalanche/Source/AvalancheMaterial/Public/AvaPrimitiveComponentMaterialBridge.h
new file mode 100644
--- /dev/null
+++ b/Engine/Plugins/VirtualProduction/Avalanche/Source/AvalancheMaterial/Public/AvaPrimitiveComponentMaterialBridge.h
@@ -0,0 +1,24 @@
+// Copyright Epic Games, Inc. All Rights Reserved.
+
+#pragma once
+
+#include "AvaMaterialBridge.h"
+
+#define UE_API AVALANCHEMATERIAL_API
+
+namespace UE::Ava
+{
+
+/** Material Bridge for Primitive Components */
+class FPrimitiveComponentMaterialBridge : public FMaterialBridge
+{
+protected:
+	//~ Begin FMaterialBridge
+	UE_API virtual const UStruct* OnGetBridgedType() const override;
+	UE_API virtual void OnForEachUsedMaterial(FConstDataView InDataView, TFunctionRef<void(UMaterialInterface*)> InFunc, const FGetUsedMaterialOptions& InOptions) const override;
+	//~ End FMaterialBridge
+};
+
+} // UE::Ava
+
+#undef UE_API
