Motion Design: added material bridges for Text3D and Shapes.

diff --git a/Engine/Plugins/VirtualProduction/Avalanche/Source/AvalancheShapes/AvalancheShapes.Build.cs b/Engine/Plugins/VirtualProduction/Avalanche/Source/AvalancheShapes/AvalancheShapes.Build.cs
--- a/Engine/Plugins/VirtualProduction/Avalanche/Source/AvalancheShapes/AvalancheShapes.Build.cs
+++ b/Engine/Plugins/VirtualProduction/Avalanche/Source/AvalancheShapes/AvalancheShapes.Build.cs
@@ -11,8 +11,10 @@ public class AvalancheShapes : ModuleRules
 		PublicDependencyModuleNames.AddRange(
 			new string[]
 			{
+				"AvalancheCore",
 				"Avalanche",
 				"AvalancheInteractiveToolsRuntime",
+				"AvalancheMaterial",
 				"Core",
 				"CoreUObject",
 				"Engine",
diff --git a/Engine/Plugins/VirtualProduction/Avalanche/Source/AvalancheShapes/Private/AvaShapeComponentMaterialBridge.cpp b/Engine/Plugins/VirtualProduction/Avalanche/Source/AvalancheShapes/Private/AvaShapeComponentMaterialBridge.cpp
new file mode 100644
--- /dev/null
+++ b/Engine/Plugins/VirtualProduction/Avalanche/Source/AvalancheShapes/Private/AvaShapeComponentMaterialBridge.cpp
@@ -0,0 +1,24 @@
+// Copyright Epic Games, Inc. All Rights Reserved.
+
+#include "AvaShapeComponentMaterialBridge.h"
+#include "DynamicMeshes/AvaShapeDynMeshBase.h"
+
+namespace UE::Ava
+{
+
+const UStruct* FShapeComponentMaterialBridge::OnGetBridgedType() const
+{
+	return UAvaShapeDynamicMeshBase::StaticClass();
+}
+
+void FShapeComponentMaterialBridge::OnForEachUsedMaterial(FConstDataView InDataView, TFunctionRef<void(UMaterialInterface*)> InFunc, const FGetUsedMaterialOptions& InOptions) const
+{
+	const UAvaShapeDynamicMeshBase& Component = InDataView.Get<const UAvaShapeDynamicMeshBase>();
+
+	for (const TPair<int32, FAvaShapeMeshData>& Pair : Component.GetMeshDatas())
+	{
+		InFunc(Pair.Value.GetMaterial());
+	}
+}
+
+} // UE::Ava
diff --git a/Engine/Plugins/VirtualProduction/Avalanche/Source/AvalancheShapes/Private/AvaShapesModule.cpp b/Engine/Plugins/VirtualProduction/Avalanche/Source/AvalancheShapes/Private/AvaShapesModule.cpp
--- a/Engine/Plugins/VirtualProduction/Avalanche/Source/AvalancheShapes/Private/AvaShapesModule.cpp
+++ b/Engine/Plugins/VirtualProduction/Avalanche/Source/AvalancheShapes/Private/AvaShapesModule.cpp
@@ -1,5 +1,14 @@
 // Copyright Epic Games, Inc. All Rights Reserved.
 
+#include "AvaShapesModule.h"
+#include "AvaMaterialBridgeRegistry.h"
+#include "AvaShapeComponentMaterialBridge.h"
 #include "Modules/ModuleManager.h"
 
-IMPLEMENT_MODULE(FDefaultModuleImpl, AvalancheShapes)
+IMPLEMENT_MODULE(FAvaShapesModule, AvalancheShapes)
+
+void FAvaShapesModule::StartupModule()
+{
+	UE::Ava::FMaterialBridgeRegistry& MaterialBridgeRegistry = UE::Ava::FMaterialBridgeRegistry::GetMutable();
+	MaterialBridgeRegistry.Register<UE::Ava::FShapeComponentMaterialBridge>();
+}
diff --git a/Engine/Plugins/VirtualProduction/Avalanche/Source/AvalancheShapes/Private/AvaShapesModule.h b/Engine/Plugins/VirtualProduction/Avalanche/Source/AvalancheShapes/Private/AvaShapesModule.h
new file mode 100644
--- /dev/null
+++ b/Engine/Plugins/VirtualProduction/Avalanche/Source/AvalancheShapes/Private/AvaShapesModule.h
@@ -0,0 +1,13 @@
+// Copyright Epic Games, Inc. All Rights Reserved.
+
+#pragma once
+
+#include "Modules/ModuleInterface.h"
+
+class FAvaShapesModule : public IModuleInterface
+{
+public:
+	//~ Begin IModuleInterface
+	virtual void StartupModule() override;
+	//~ End IModuleInterface
+};
diff --git a/Engine/Plugins/VirtualProduction/Avalanche/Source/AvalancheShapes/Public/AvaShapeComponentMaterialBridge.h b/Engine/Plugins/VirtualProduction/Avalanche/Source/AvalancheShapes/Public/AvaShapeComponentMaterialBridge.h
new file mode 100644
--- /dev/null
+++ b/Engine/Plugins/VirtualProduction/Avalanche/Source/AvalancheShapes/Public/AvaShapeComponentMaterialBridge.h
@@ -0,0 +1,24 @@
+// Copyright Epic Games, Inc. All Rights Reserved.
+
+#pragma once
+
+#include "AvaMaterialBridge.h"
+
+#define UE_API AVALANCHESHAPES_API
+
+namespace UE::Ava
+{
+
+/** Material Bridge for Motion Design Shape Components */
+class FShapeComponentMaterialBridge : public FMaterialBridge
+{
+protected:
+	//~ Begin FMaterialBridge
+	UE_API virtual const UStruct* OnGetBridgedType() const override;
+	UE_API virtual void OnForEachUsedMaterial(FConstDataView InDataView, TFunctionRef<void(UMaterialInterface*)> InFunc, const FGetUsedMaterialOptions& InOptions) const override;
+	//~ End FMaterialBridge
+};
+
+} // UE::Ava
+
+#undef UE_API
diff --git a/Engine/Plugins/VirtualProduction/Avalanche/Source/AvalancheShapes/Public/AvaShapeMesh.h b/Engine/Plugins/VirtualProduction/Avalanche/Source/AvalancheShapes/Public/AvaShapeMesh.h
--- a/Engine/Plugins/VirtualProduction/Avalanche/Source/AvalancheShapes/Public/AvaShapeMesh.h
+++ b/Engine/Plugins/VirtualProduction/Avalanche/Source/AvalancheShapes/Public/AvaShapeMesh.h
@@ -116,6 +116,11 @@ struct FAvaShapeMeshData
 		return MeshName;
 	}
 
+	UMaterialInterface* GetMaterial() const
+	{
+		return Material;
+	}
+
 protected:
 	// whether the mesh is currently visible or not and should be editable
 	UPROPERTY(EditDefaultsOnly, Category="Shape")
diff --git a/Engine/Plugins/VirtualProduction/Avalanche/Source/AvalancheShapes/Public/DynamicMeshes/AvaShapeDynMeshBase.h b/Engine/Plugins/VirtualProduction/Avalanche/Source/AvalancheShapes/Public/DynamicMeshes/AvaShapeDynMeshBase.h
--- a/Engine/Plugins/VirtualProduction/Avalanche/Source/AvalancheShapes/Public/DynamicMeshes/AvaShapeDynMeshBase.h
+++ b/Engine/Plugins/VirtualProduction/Avalanche/Source/AvalancheShapes/Public/DynamicMeshes/AvaShapeDynMeshBase.h
@@ -223,6 +223,11 @@ public:
 	/** find a registered mesh and gets a pointer to it */
 	FAvaShapeMeshData* GetMeshData(int32 MeshIndex);
 
+	const TMap<int32, FAvaShapeMeshData>& GetMeshDatas() const
+	{
+		return MeshDatas;
+	}
+
 	// apply scale, offset, rotation on dynamic mesh section
 	bool ApplyUVsTransform(FAvaShapeMesh& InMesh, FAvaShapeMaterialUVParameters& InParams, FVector2D ShapeSize, FVector2D UVFixOffset, float UVFixRotation = 0.f);
 
diff --git a/Engine/Plugins/VirtualProduction/Avalanche/Source/AvalancheText/AvalancheText.Build.cs b/Engine/Plugins/VirtualProduction/Avalanche/Source/AvalancheText/AvalancheText.Build.cs
--- a/Engine/Plugins/VirtualProduction/Avalanche/Source/AvalancheText/AvalancheText.Build.cs
+++ b/Engine/Plugins/VirtualProduction/Avalanche/Source/AvalancheText/AvalancheText.Build.cs
@@ -11,11 +11,13 @@ public class AvalancheText : ModuleRules
 		PublicDependencyModuleNames.AddRange(
 			new string[]
 			{
+				"AvalancheCore",
+				"AvalancheMaterial",
 				"Core",
 				"CoreUObject",
 				"Engine",
 				"SlateCore",
-				"Text3D"
+				"Text3D",
 			}
 		);
 
diff --git a/Engine/Plugins/VirtualProduction/Avalanche/Source/AvalancheText/Private/AvaText3DComponentMaterialBridge.cpp b/Engine/Plugins/VirtualProduction/Avalanche/Source/AvalancheText/Private/AvaText3DComponentMaterialBridge.cpp
new file mode 100644
--- /dev/null
+++ b/Engine/Plugins/VirtualProduction/Avalanche/Source/AvalancheText/Private/AvaText3DComponentMaterialBridge.cpp
@@ -0,0 +1,31 @@
+// Copyright Epic Games, Inc. All Rights Reserved.
+
+#include "AvaText3DComponentMaterialBridge.h"
+#include "AvaDataView.h"
+#include "Extensions/Text3DMaterialExtensionBase.h"
+#include "Text3DComponent.h"
+
+namespace UE::Ava
+{
+
+const UStruct* FText3DComponentMaterialBridge::OnGetBridgedType() const
+{
+	return UText3DComponent::StaticClass();
+}
+
+void FText3DComponentMaterialBridge::OnForEachUsedMaterial(FConstDataView InDataView, TFunctionRef<void(UMaterialInterface*)> InFunc, const FGetUsedMaterialOptions& InOptions) const
+{
+	const UText3DComponent& Component = InDataView.Get<const UText3DComponent>();
+
+	if (const UText3DMaterialExtensionBase* MaterialExtension = Component.GetMaterialExtension())
+	{
+		MaterialExtension->ForEachMaterial(
+			[&InFunc](const UE::Text3D::Material::FMaterialParameters&, UMaterialInterface* InMaterial)->bool
+			{
+				InFunc(InMaterial);
+				return true; // continue				
+			});
+	}
+}
+
+} // UE::Ava
diff --git a/Engine/Plugins/VirtualProduction/Avalanche/Source/AvalancheText/Private/AvaTextModule.cpp b/Engine/Plugins/VirtualProduction/Avalanche/Source/AvalancheText/Private/AvaTextModule.cpp
--- a/Engine/Plugins/VirtualProduction/Avalanche/Source/AvalancheText/Private/AvaTextModule.cpp
+++ b/Engine/Plugins/VirtualProduction/Avalanche/Source/AvalancheText/Private/AvaTextModule.cpp
@@ -1,5 +1,14 @@
 // Copyright Epic Games, Inc. All Rights Reserved.
 
+#include "AvaTextModule.h"
+#include "AvaMaterialBridgeRegistry.h"
+#include "AvaText3DComponentMaterialBridge.h"
 #include "Modules/ModuleManager.h"
 
-IMPLEMENT_MODULE(FDefaultModuleImpl, AvalancheText)
+IMPLEMENT_MODULE(FAvaTextModule, AvalancheText)
+
+void FAvaTextModule::StartupModule()
+{
+	UE::Ava::FMaterialBridgeRegistry& MaterialBridgeRegistry = UE::Ava::FMaterialBridgeRegistry::GetMutable();
+	MaterialBridgeRegistry.Register<UE::Ava::FText3DComponentMaterialBridge>();
+}
diff --git a/Engine/Plugins/VirtualProduction/Avalanche/Source/AvalancheText/Private/AvaTextModule.h b/Engine/Plugins/VirtualProduction/Avalanche/Source/AvalancheText/Private/AvaTextModule.h
new file mode 100644
--- /dev/null
+++ b/Engine/Plugins/VirtualProduction/Avalanche/Source/AvalancheText/Private/AvaTextModule.h
@@ -0,0 +1,13 @@
+// Copyright Epic Games, Inc. All Rights Reserved.
+
+#pragma once
+
+#include "Modules/ModuleInterface.h"
+
+class FAvaTextModule : public IModuleInterface
+{
+public:
+	//~ Begin IModuleInterface
+	virtual void StartupModule() override;
+	//~ End IModuleInterface
+};
diff --git a/Engine/Plugins/VirtualProduction/Avalanche/Source/AvalancheText/Public/AvaText3DComponentMaterialBridge.h b/Engine/Plugins/VirtualProduction/Avalanche/Source/AvalancheText/Public/AvaText3DComponentMaterialBridge.h
new file mode 100644
--- /dev/null
+++ b/Engine/Plugins/VirtualProduction/Avalanche/Source/AvalancheText/Public/AvaText3DComponentMaterialBridge.h
@@ -0,0 +1,24 @@
+// Copyright Epic Games, Inc. All Rights Reserved.
+
+#pragma once
+
+#include "AvaMaterialBridge.h"
+
+#define UE_API AVALANCHETEXT_API
+
+namespace UE::Ava
+{
+
+/** Material Bridge for Text3D Components */
+class FText3DComponentMaterialBridge : public FMaterialBridge
+{
+protected:
+	//~ Begin FMaterialBridge
+	UE_API virtual const UStruct* OnGetBridgedType() const override;
+	UE_API virtual void OnForEachUsedMaterial(FConstDataView InDataView, TFunctionRef<void(UMaterialInterface*)> InFunc, const FGetUsedMaterialOptions& InOptions) const override;
+	//~ End FMaterialBridge
+};
+
+} // UE::Ava
+
+#undef UE_API
