Motion Design: added mechanism to cache shaders before starting the graphics. Additionally, Remote Control referenced assets also attempt to get streamed in and completed before play.

diff --git a/Engine/Plugins/VirtualProduction/Avalanche/Source/AvalancheMedia/AvalancheMedia.Build.cs b/Engine/Plugins/VirtualProduction/Avalanche/Source/AvalancheMedia/AvalancheMedia.Build.cs
--- a/Engine/Plugins/VirtualProduction/Avalanche/Source/AvalancheMedia/AvalancheMedia.Build.cs
+++ b/Engine/Plugins/VirtualProduction/Avalanche/Source/AvalancheMedia/AvalancheMedia.Build.cs
@@ -36,6 +36,8 @@ public class AvalancheMedia : ModuleRules
 			{
 				"ApplicationCore",
 				"AssetRegistry",
+				"AvalancheCore",
+				"AvalancheMaterial",
 				"AvalancheRemoteControl",
 				"HeadMountedDisplay",
 				"Json",
diff --git a/Engine/Plugins/VirtualProduction/Avalanche/Source/AvalancheMedia/Internal/Playback/AvaPlaybackUtils.h b/Engine/Plugins/VirtualProduction/Avalanche/Source/AvalancheMedia/Internal/Playback/AvaPlaybackUtils.h
--- a/Engine/Plugins/VirtualProduction/Avalanche/Source/AvalancheMedia/Internal/Playback/AvaPlaybackUtils.h
+++ b/Engine/Plugins/VirtualProduction/Avalanche/Source/AvalancheMedia/Internal/Playback/AvaPlaybackUtils.h
@@ -9,6 +9,7 @@
 #include "UObject/SoftObjectPath.h"
 
 class UPackage;
+class UStreamableRenderAsset;
 struct FAssetData;
 
 class AVALANCHEMEDIA_API FAvaPlaybackUtils
@@ -71,22 +72,24 @@ namespace UE::AvaPlayback::Utils
 	/**
 	 * Utility class to perform async load of assets. 
 	 */
-	class AVALANCHEMEDIA_API FAsyncAssetLoader : public TSharedFromThis<FAsyncAssetLoader>
+	class FAsyncAssetLoader : public TSharedFromThis<FAsyncAssetLoader>
 	{
 	public:
+		void HandleStreamableObject(UObject* InObject);
+
 		/**
 		 * @brief Issue the async load command. Returns immediately.
 		 * @param InAssetsToLoad Assets to load.
 		 *
 		 * The specified asset to load will added to the current set of pending assets.
 		 */
-		void BeginLoadingAssets(const TArray<FSoftObjectPath>& InAssetsToLoad);
+		AVALANCHEMEDIA_API void BeginLoadingAssets(const TArray<FSoftObjectPath>& InAssetsToLoad);
 
 		/** Returns true when all the requested assets have finished loading. */
-		bool IsLoadingCompleted() const
-		{
-			return PendingAssets.IsEmpty();
-		}
+		AVALANCHEMEDIA_API bool IsLoadingCompleted() const;
+
+		/** Removes streamable render assets that have fully been streamed in */
+		AVALANCHEMEDIA_API void CleanupStreamingAssets() const;
 
 		/** Delegate called when all assets have finished loading. */
 		DECLARE_MULTICAST_DELEGATE(FOnLoadingCompleted);
@@ -94,5 +97,6 @@ namespace UE::AvaPlayback::Utils
 
 	private:
 		TSet<FSoftObjectPath> PendingAssets;
+		mutable TArray<TObjectPtr<UStreamableRenderAsset>> PendingStreamingAssets;
 	};
-}
\ No newline at end of file
+}
diff --git a/Engine/Plugins/VirtualProduction/Avalanche/Source/AvalancheMedia/Private/Playable/Playables/AvaPlayableLevelStreaming.cpp b/Engine/Plugins/VirtualProduction/Avalanche/Source/AvalancheMedia/Private/Playable/Playables/AvaPlayableLevelStreaming.cpp
--- a/Engine/Plugins/VirtualProduction/Avalanche/Source/AvalancheMedia/Private/Playable/Playables/AvaPlayableLevelStreaming.cpp
+++ b/Engine/Plugins/VirtualProduction/Avalanche/Source/AvalancheMedia/Private/Playable/Playables/AvaPlayableLevelStreaming.cpp
@@ -1,11 +1,11 @@
 // Copyright Epic Games, Inc. All Rights Reserved.
 
 #include "Playable/Playables/AvaPlayableLevelStreaming.h"
-
+#include "AvaMaterialBridge.h"
+#include "AvaMaterialBridgeRegistry.h"
 #include "AvaRemoteControlRebind.h"
 #include "AvaRemoteControlUtils.h"
 #include "AvaScene.h"
-#include "Components/PrimitiveComponent.h"
 #include "Engine/Level.h"
 #include "Engine/LevelStreamingDynamic.h"
 #include "Engine/LocalPlayer.h"
@@ -15,12 +15,17 @@
 #include "Framework/AvaSoftAssetPtr.h"
 #include "IAvaMediaModule.h"
 #include "IAvaRemoteControlInterface.h"
+#include "MaterialShared.h"
+#include "Materials/Material.h"
+#include "Materials/MaterialInterface.h"
 #include "Playable/AvaPlayableAssetUserData.h"
 #include "Playable/AvaPlayableGroup.h"
 #include "Playable/AvaPlayableGroupManager.h"
 #include "Playable/AvaPlayableUtils.h"
 #include "Playback/AvaPlaybackUtils.h"
 #include "SceneView.h"
+#include "ShaderCompiler.h"
+#include "UnrealEngine.h"
 #include "Streaming/LevelStreamingDelegates.h"
 #include "UObject/Package.h"
 
@@ -28,6 +33,12 @@
 
 namespace UE::AvaMedia::LevelStreamingPlayable::Private
 {
+	static TAutoConsoleVariable<bool> CVarRuntimePrecacheShaders(
+		TEXT("MotionDesignPlayable.RuntimePrecacheShaders"),
+		true,
+		TEXT("Caches any remaining shaders from materials pending completion found in the playable while preparing to start play.")
+	);
+
 	UAvaPlayableAssetUserData* FindOrAddPlayableAssetUserData(ULevel* InLevel)
 	{
 		if (!InLevel)
@@ -268,6 +279,30 @@ namespace UE::AvaMedia::LevelStreamingPlayable::Private
 		
 		return GetPlayableStatusRelevanceShouldBeUnloaded(InStatus);
 	}
+
+	/** Calls cache shaders for each material in the set, until the time budget is exceeded */
+	void CacheMaterials(const TSet<UMaterialInterface*>& InMaterialsToCache, double InTimeBudget)
+	{
+		TRACE_CPUPROFILER_EVENT_SCOPE(MotionDesignPlayable_CacheMaterials)
+
+		const double StartTime = FPlatformTime::Seconds();
+
+		FMaterialUpdateContext UpdateContext(0);
+		for (UMaterialInterface* MaterialInterface : InMaterialsToCache)
+		{
+			constexpr EMaterialShaderPrecompileMode PrecompileMode = EMaterialShaderPrecompileMode::Background;
+
+			// This is needed because CacheShaders blindly recreates uniform buffers
+			// which can only be done if the draw command is going to be re-cached.
+			UpdateContext.AddMaterialInterface(MaterialInterface);
+			MaterialInterface->CacheShaders(PrecompileMode);
+
+			if (FPlatformTime::Seconds() - StartTime > InTimeBudget)
+			{
+				return;
+			}
+		}
+	}
 }
 
 bool UAvaPlayableLevelStreaming::LoadAsset(const FAvaSoftAssetPtr& InSourceAsset, bool bInInitiallyVisible, const FString& InLoadOptions)
@@ -371,6 +406,90 @@ void UAvaPlayableLevelStreaming::UpdatePlayableStatus(ELevelStreamingState InNew
 	}
 }
 
+bool UAvaPlayableLevelStreaming::IsReadyToPlay() const
+{
+	TRACE_CPUPROFILER_EVENT_SCOPE(UAvaPlayableLevelStreaming::IsReadyToPlay)
+
+	using namespace UE::AvaMedia::LevelStreamingPlayable;
+
+	for (const UAvaPlayableLevelStreaming* SubPlayable : SubPlayables)
+	{
+		if (SubPlayable && !SubPlayable->IsReadyToPlay())
+		{
+			return false;
+		}
+	}
+
+	if (!Private::CVarRuntimePrecacheShaders.GetValueOnGameThread())
+	{
+		// Not allowed to cache. Ready to play
+		return true;
+	}
+
+	if (!LevelStreaming)
+	{
+		return false;
+	}
+
+	const ULevel* const Level = LevelStreaming->GetLoadedLevel();
+	if (!Level)
+	{
+		return false;
+	}
+
+	TSet<UMaterialInterface*> MaterialsToCache;
+
+	const UE::Ava::FMaterialBridgeRegistry& MaterialBridgeRegistry = UE::Ava::FMaterialBridgeRegistry::Get();
+
+	for (AActor* Actor : Level->Actors)
+	{
+		if (const UE::Ava::FMaterialBridge* ActorMaterialBridge = MaterialBridgeRegistry.GetMaterialBridgeForObject(Actor))
+		{
+			ActorMaterialBridge->ForEachUsedMaterial(UE::Ava::FConstDataView(Actor),
+				[&MaterialsToCache](UMaterialInterface* InMaterialInterface)
+				{
+					if (!InMaterialInterface)
+					{
+						return;
+					}
+
+					UMaterial* const Material = InMaterialInterface->GetMaterial();
+					if (Material && Material != InMaterialInterface && !Material->IsComplete())
+					{
+						MaterialsToCache.Add(Material);
+					}
+					if (!InMaterialInterface->IsComplete())
+					{
+						MaterialsToCache.Add(InMaterialInterface);
+					}
+				});
+		}
+	}
+
+	if (MaterialsToCache.IsEmpty())
+	{
+		bool bFinishedCompilation = true;
+
+		if (GShaderCompilingManager)
+		{
+			if (GShaderCompilingManager->IsCompiling())
+			{
+				constexpr bool bLimitExecutionTime = true;
+				constexpr bool bBlockOnGlobalShaderCompletion = false;
+				GShaderCompilingManager->ProcessAsyncResults(bLimitExecutionTime, bBlockOnGlobalShaderCompletion);
+
+				bFinishedCompilation = false;
+			}
+		}
+
+		return bFinishedCompilation;
+	}
+
+	constexpr double TimeBudget = 0.001;
+	Private::CacheMaterials(MaterialsToCache, TimeBudget);
+	return false;
+}
+
 IAvaSceneInterface* UAvaPlayableLevelStreaming::GetSceneInterface() const
 {
 	if (Scene)
diff --git a/Engine/Plugins/VirtualProduction/Avalanche/Source/AvalancheMedia/Private/Playable/Playables/AvaPlayableLevelStreaming.h b/Engine/Plugins/VirtualProduction/Avalanche/Source/AvalancheMedia/Private/Playable/Playables/AvaPlayableLevelStreaming.h
--- a/Engine/Plugins/VirtualProduction/Avalanche/Source/AvalancheMedia/Private/Playable/Playables/AvaPlayableLevelStreaming.h
+++ b/Engine/Plugins/VirtualProduction/Avalanche/Source/AvalancheMedia/Private/Playable/Playables/AvaPlayableLevelStreaming.h
@@ -27,6 +27,7 @@ public:
 	virtual bool GetShouldBeVisible() const override;
 	virtual void SetShouldBeVisible(bool bInShouldBeVisible) override;
 	virtual void SetupView(FSceneViewFamily& InViewFamily, FSceneView& InView) override;
+	virtual bool IsReadyToPlay() const override;
 
 protected:
 	virtual bool InitPlayable(const FPlayableCreationInfo& InPlayableInfo) override;
diff --git a/Engine/Plugins/VirtualProduction/Avalanche/Source/AvalancheMedia/Private/Playable/Playables/AvaPlayableRemoteProxy.cpp b/Engine/Plugins/VirtualProduction/Avalanche/Source/AvalancheMedia/Private/Playable/Playables/AvaPlayableRemoteProxy.cpp
--- a/Engine/Plugins/VirtualProduction/Avalanche/Source/AvalancheMedia/Private/Playable/Playables/AvaPlayableRemoteProxy.cpp
+++ b/Engine/Plugins/VirtualProduction/Avalanche/Source/AvalancheMedia/Private/Playable/Playables/AvaPlayableRemoteProxy.cpp
@@ -182,6 +182,12 @@ EAvaPlayableStatus UAvaPlayableRemoteProxy::GetPlayableStatus() const
 	return PlayableStatus;
 }
 
+bool UAvaPlayableRemoteProxy::IsReadyToPlay() const
+{
+	// todo:
+	return true;
+}
+
 IAvaSceneInterface* UAvaPlayableRemoteProxy::GetSceneInterface() const
 {
 	return nullptr;
diff --git a/Engine/Plugins/VirtualProduction/Avalanche/Source/AvalancheMedia/Private/Playable/Playables/AvaPlayableRemoteProxy.h b/Engine/Plugins/VirtualProduction/Avalanche/Source/AvalancheMedia/Private/Playable/Playables/AvaPlayableRemoteProxy.h
--- a/Engine/Plugins/VirtualProduction/Avalanche/Source/AvalancheMedia/Private/Playable/Playables/AvaPlayableRemoteProxy.h
+++ b/Engine/Plugins/VirtualProduction/Avalanche/Source/AvalancheMedia/Private/Playable/Playables/AvaPlayableRemoteProxy.h
@@ -26,6 +26,7 @@ public:
 	virtual bool UnloadAsset() override;
 	virtual const FSoftObjectPath& GetSourceAssetPath() const override { return SourceAssetPath; }
 	virtual EAvaPlayableStatus GetPlayableStatus() const override;
+	virtual bool IsReadyToPlay() const override;
 	virtual IAvaSceneInterface* GetSceneInterface() const override;
 	virtual EAvaPlayableCommandResult ExecuteAnimationCommand(EAvaPlaybackAnimAction InAnimAction, const FAvaPlaybackAnimPlaySettings& InAnimPlaySettings) override;
 	virtual EAvaPlayableCommandResult UpdateRemoteControlCommand(const TSharedRef<FAvaPlayableRemoteControlValues>& InRemoteControlValues, EAvaPlayableRCUpdateFlags InFlags) override;
diff --git a/Engine/Plugins/VirtualProduction/Avalanche/Source/AvalancheMedia/Private/Playback/AvaPlaybackUtils.cpp b/Engine/Plugins/VirtualProduction/Avalanche/Source/AvalancheMedia/Private/Playback/AvaPlaybackUtils.cpp
--- a/Engine/Plugins/VirtualProduction/Avalanche/Source/AvalancheMedia/Private/Playback/AvaPlaybackUtils.cpp
+++ b/Engine/Plugins/VirtualProduction/Avalanche/Source/AvalancheMedia/Private/Playback/AvaPlaybackUtils.cpp
@@ -5,11 +5,15 @@
 #include "AssetRegistry/AssetData.h"
 #include "AssetRegistry/IAssetRegistry.h"
 #include "AvaAssetTags.h"
+#include "ContentStreaming.h"
+#include "Engine/StreamableRenderAsset.h"
 #include "Framework/AvaSoftAssetPtr.h"
-#include "IAvaMediaModule.h"
 #include "HAL/PlatformFileManager.h"
+#include "IAvaMediaModule.h"
+#include "Materials/MaterialInterface.h"
 #include "Misc/Paths.h"
 #include "RenderingThread.h"
+#include "ShaderCompiler.h"
 #include "UObject/Linker.h"
 #include "UObject/LinkerLoad.h"
 #include "UObject/Package.h"
@@ -28,6 +32,12 @@ namespace UE::AvaPlayback::Utils::Private
 		false,
 		TEXT("Forces to ResetLoaders instead of DetachLoader during FlushPackageLoading")
 	);
+
+	static TAutoConsoleVariable<bool> CVarRuntimePrestreamRCReferencedAssets(
+		TEXT("MotionDesignPlayback.RuntimePrestreamRemoteControlReferencedAssets"),
+		true,
+		TEXT("Manually streams and waits for all Remote Control referenced assets to be ready while preparing to start play.")
+	);
 }
 
 void FAvaPlaybackUtils::FlushPackageLoading(UPackage* InPackage, bool bInResetLoaders)
@@ -230,16 +240,48 @@ namespace UE::AvaPlayback::Utils
 		return FString::Printf(TEXT("[%d]"), GFrameNumber);
 	}
 
+	void FAsyncAssetLoader::HandleStreamableObject(UObject* InObject)
+	{
+		if (!Private::CVarRuntimePrestreamRCReferencedAssets.GetValueOnGameThread())
+		{
+			// Not allowed to stream assets manually.
+			return;
+		}
+
+		UStreamableRenderAsset* const StreamableAsset = Cast<UStreamableRenderAsset>(InObject);
+		if (StreamableAsset && StreamableAsset->IsStreamable())
+		{
+			constexpr float DurationSeconds = 30.f;
+			constexpr uint32 CinematicLODGroupMask = 0xFFFFFFFF;
+			StreamableAsset->SetForceMipLevelsToBeResident(DurationSeconds, CinematicLODGroupMask);
+
+			constexpr bool bHighPriority = true;
+			StreamableAsset->StreamIn(FStreamableRenderResourceState::MAX_LOD_COUNT, bHighPriority);
+
+			PendingStreamingAssets.Add(StreamableAsset);
+		}
+	}
+
 	void FAsyncAssetLoader::BeginLoadingAssets(const TArray<FSoftObjectPath>& InAssetsToLoad)
 	{
 		for (const FSoftObjectPath& AssetToLoad : InAssetsToLoad)
 		{
-			if (AssetToLoad.ResolveObject() // Already loaded.
-				|| PendingAssets.Contains(AssetToLoad)) // Already pending	
+			if (PendingAssets.Contains(AssetToLoad))
+			{
+				// Already processed asset
+				continue;
+			}
+
+			UObject* ResolvedObject = AssetToLoad.ResolveObject();
+			if (ResolvedObject)
 			{
+				// Already loaded... but streamable objects could still require streaming in
+				HandleStreamableObject(ResolvedObject);
 				continue;
 			}
 
+			PendingAssets.Add(AssetToLoad);
+
 			AssetToLoad.LoadAsync(FLoadSoftObjectPathAsyncDelegate::CreateSPLambda(this, [this](const FSoftObjectPath& InObjectPath, UObject* InLoadedObject)
 			{
 				PendingAssets.Remove(InObjectPath);
@@ -247,7 +289,47 @@ namespace UE::AvaPlayback::Utils
 				{
 					OnLoadingCompleted.Broadcast();
 				}
+				HandleStreamableObject(InLoadedObject);
 			}));
 		}
 	}
-}
\ No newline at end of file
+
+	bool FAsyncAssetLoader::IsLoadingCompleted() const
+	{
+		if (!PendingAssets.IsEmpty())
+		{
+			return false;
+		}
+
+		CleanupStreamingAssets();
+
+		if (!PendingStreamingAssets.IsEmpty())
+		{
+			return false;
+		}
+
+		bool bFinishedCompilation = true;
+		if (GShaderCompilingManager)
+		{
+			if (GShaderCompilingManager->IsCompiling())
+			{
+				constexpr bool bLimitExecutionTime = true;
+				constexpr bool bBlockOnGlobalShaderCompletion = false;
+				GShaderCompilingManager->ProcessAsyncResults(bLimitExecutionTime, bBlockOnGlobalShaderCompletion);
+
+				bFinishedCompilation = false;
+			}
+		}
+		return bFinishedCompilation;
+	}
+
+	void FAsyncAssetLoader::CleanupStreamingAssets() const
+	{
+		PendingStreamingAssets.RemoveAll(
+			[](UStreamableRenderAsset* InStreamableRenderAsset)
+			{
+				constexpr bool bWaitForLODTransition = true;
+				return !InStreamableRenderAsset->HasPendingInitOrStreaming(bWaitForLODTransition) && InStreamableRenderAsset->IsFullyStreamedIn();
+			});
+	}
+}
diff --git a/Engine/Plugins/VirtualProduction/Avalanche/Source/AvalancheMedia/Private/Playback/Transition/AvaPlaybackServerTransition.cpp b/Engine/Plugins/VirtualProduction/Avalanche/Source/AvalancheMedia/Private/Playback/Transition/AvaPlaybackServerTransition.cpp
--- a/Engine/Plugins/VirtualProduction/Avalanche/Source/AvalancheMedia/Private/Playback/Transition/AvaPlaybackServerTransition.cpp
+++ b/Engine/Plugins/VirtualProduction/Avalanche/Source/AvalancheMedia/Private/Playback/Transition/AvaPlaybackServerTransition.cpp
@@ -117,6 +117,13 @@ namespace UE::AvaPlaybackServerTransition::Private
 				bOutShouldDiscard = false;
 				return false;
 			}
+
+			if (!Playable->IsReadyToPlay())
+			{
+				// Keep the command in the queue for next tick.
+				bOutShouldDiscard = false;
+				return false;
+			}
 		}
 
 		bOutShouldDiscard = true;
diff --git a/Engine/Plugins/VirtualProduction/Avalanche/Source/AvalancheMedia/Private/Rundown/Transition/AvaRundownPageTransition.cpp b/Engine/Plugins/VirtualProduction/Avalanche/Source/AvalancheMedia/Private/Rundown/Transition/AvaRundownPageTransition.cpp
--- a/Engine/Plugins/VirtualProduction/Avalanche/Source/AvalancheMedia/Private/Rundown/Transition/AvaRundownPageTransition.cpp
+++ b/Engine/Plugins/VirtualProduction/Avalanche/Source/AvalancheMedia/Private/Rundown/Transition/AvaRundownPageTransition.cpp
@@ -234,6 +234,13 @@ bool UAvaRundownPageTransition::CanStart(bool& bOutShouldDiscard)
 						bOutShouldDiscard = false;
 						return false;
 					}
+
+					if (!Playable->IsRemoteProxy() && !Playable->IsReadyToPlay())
+					{
+						// Keep the command in the queue for next tick.
+						bOutShouldDiscard = false;
+						return false;
+					}
 				}
 				else
 				{
diff --git a/Engine/Plugins/VirtualProduction/Avalanche/Source/AvalancheMedia/Public/Playable/AvaPlayable.h b/Engine/Plugins/VirtualProduction/Avalanche/Source/AvalancheMedia/Public/Playable/AvaPlayable.h
--- a/Engine/Plugins/VirtualProduction/Avalanche/Source/AvalancheMedia/Public/Playable/AvaPlayable.h
+++ b/Engine/Plugins/VirtualProduction/Avalanche/Source/AvalancheMedia/Public/Playable/AvaPlayable.h
@@ -117,7 +117,16 @@ public:
 	virtual bool GetShouldBeVisible() const { return true; }
 	virtual void SetShouldBeVisible(bool bInShouldBeVisible) {}
 	virtual void SetupView(FSceneViewFamily& InViewFamily, FSceneView& InView) {}
-	
+
+	/**
+	 * Called when the playable is visible, to determine if the playable is ready to transition and be shown.
+	 * Should return true when all the elements are loaded/ready after becoming visible.
+	 */
+	virtual bool IsReadyToPlay() const
+	{
+		return false;
+	}
+
 	/**
 	 * @brief Ensures the given asset is playing (visible) with the given parameters.
 	 * @param InWorldPlaySettings World settings (render target) to render this asset with.
