Motion Design: fixed issue with playback client not sending the start transition request.
+ Added some extra optimizations of not needing to check for IsReadyToPlay for Playing and Exit playables.

diff --git a/Engine/Plugins/VirtualProduction/Avalanche/Source/AvalancheMedia/Internal/Playback/AvaPlaybackUtils.h b/Engine/Plugins/VirtualProduction/Avalanche/Source/AvalancheMedia/Internal/Playback/AvaPlaybackUtils.h
--- a/Engine/Plugins/VirtualProduction/Avalanche/Source/AvalancheMedia/Internal/Playback/AvaPlaybackUtils.h
+++ b/Engine/Plugins/VirtualProduction/Avalanche/Source/AvalancheMedia/Internal/Playback/AvaPlaybackUtils.h
@@ -98,5 +98,6 @@ namespace UE::AvaPlayback::Utils
 	private:
 		TSet<FSoftObjectPath> PendingAssets;
 		mutable TArray<TObjectPtr<UStreamableRenderAsset>> PendingStreamingAssets;
+		bool bWaitForStreamingAssets = false;
 	};
 }
diff --git a/Engine/Plugins/VirtualProduction/Avalanche/Source/AvalancheMedia/Private/Playable/Playables/AvaPlayableRemoteProxy.cpp b/Engine/Plugins/VirtualProduction/Avalanche/Source/AvalancheMedia/Private/Playable/Playables/AvaPlayableRemoteProxy.cpp
--- a/Engine/Plugins/VirtualProduction/Avalanche/Source/AvalancheMedia/Private/Playable/Playables/AvaPlayableRemoteProxy.cpp
+++ b/Engine/Plugins/VirtualProduction/Avalanche/Source/AvalancheMedia/Private/Playable/Playables/AvaPlayableRemoteProxy.cpp
@@ -182,12 +182,6 @@ EAvaPlayableStatus UAvaPlayableRemoteProxy::GetPlayableStatus() const
 	return PlayableStatus;
 }
 
-bool UAvaPlayableRemoteProxy::IsReadyToPlay() const
-{
-	// todo:
-	return true;
-}
-
 IAvaSceneInterface* UAvaPlayableRemoteProxy::GetSceneInterface() const
 {
 	return nullptr;
diff --git a/Engine/Plugins/VirtualProduction/Avalanche/Source/AvalancheMedia/Private/Playable/Playables/AvaPlayableRemoteProxy.h b/Engine/Plugins/VirtualProduction/Avalanche/Source/AvalancheMedia/Private/Playable/Playables/AvaPlayableRemoteProxy.h
--- a/Engine/Plugins/VirtualProduction/Avalanche/Source/AvalancheMedia/Private/Playable/Playables/AvaPlayableRemoteProxy.h
+++ b/Engine/Plugins/VirtualProduction/Avalanche/Source/AvalancheMedia/Private/Playable/Playables/AvaPlayableRemoteProxy.h
@@ -26,7 +26,6 @@ public:
 	virtual bool UnloadAsset() override;
 	virtual const FSoftObjectPath& GetSourceAssetPath() const override { return SourceAssetPath; }
 	virtual EAvaPlayableStatus GetPlayableStatus() const override;
-	virtual bool IsReadyToPlay() const override;
 	virtual IAvaSceneInterface* GetSceneInterface() const override;
 	virtual EAvaPlayableCommandResult ExecuteAnimationCommand(EAvaPlaybackAnimAction InAnimAction, const FAvaPlaybackAnimPlaySettings& InAnimPlaySettings) override;
 	virtual EAvaPlayableCommandResult UpdateRemoteControlCommand(const TSharedRef<FAvaPlayableRemoteControlValues>& InRemoteControlValues, EAvaPlayableRCUpdateFlags InFlags) override;
diff --git a/Engine/Plugins/VirtualProduction/Avalanche/Source/AvalancheMedia/Private/Playback/AvaPlaybackUtils.cpp b/Engine/Plugins/VirtualProduction/Avalanche/Source/AvalancheMedia/Private/Playback/AvaPlaybackUtils.cpp
--- a/Engine/Plugins/VirtualProduction/Avalanche/Source/AvalancheMedia/Private/Playback/AvaPlaybackUtils.cpp
+++ b/Engine/Plugins/VirtualProduction/Avalanche/Source/AvalancheMedia/Private/Playback/AvaPlaybackUtils.cpp
@@ -259,6 +259,7 @@ namespace UE::AvaPlayback::Utils
 			StreamableAsset->StreamIn(FStreamableRenderResourceState::MAX_LOD_COUNT, bHighPriority);
 
 			PendingStreamingAssets.Add(StreamableAsset);
+			bWaitForStreamingAssets = true;
 		}
 	}
 
@@ -308,8 +309,10 @@ namespace UE::AvaPlayback::Utils
 			return false;
 		}
 
-		bool bFinishedCompilation = true;
-		if (GShaderCompilingManager)
+		bool bIsCompilationComplete = true;
+
+		// Only wait for compilation if there were streaming assets that were added
+		if (bWaitForStreamingAssets && GShaderCompilingManager)
 		{
 			if (GShaderCompilingManager->IsCompiling())
 			{
@@ -317,10 +320,11 @@ namespace UE::AvaPlayback::Utils
 				constexpr bool bBlockOnGlobalShaderCompletion = false;
 				GShaderCompilingManager->ProcessAsyncResults(bLimitExecutionTime, bBlockOnGlobalShaderCompletion);
 
-				bFinishedCompilation = false;
+				bIsCompilationComplete = false;
 			}
 		}
-		return bFinishedCompilation;
+
+		return bIsCompilationComplete;
 	}
 
 	void FAsyncAssetLoader::CleanupStreamingAssets() const
diff --git a/Engine/Plugins/VirtualProduction/Avalanche/Source/AvalancheMedia/Private/Playback/Transition/AvaPlaybackServerTransition.cpp b/Engine/Plugins/VirtualProduction/Avalanche/Source/AvalancheMedia/Private/Playback/Transition/AvaPlaybackServerTransition.cpp
--- a/Engine/Plugins/VirtualProduction/Avalanche/Source/AvalancheMedia/Private/Playback/Transition/AvaPlaybackServerTransition.cpp
+++ b/Engine/Plugins/VirtualProduction/Avalanche/Source/AvalancheMedia/Private/Playback/Transition/AvaPlaybackServerTransition.cpp
@@ -58,10 +58,9 @@ namespace UE::AvaPlaybackServerTransition::Private
 		return nullptr;
 	}
 
-	
 	// Check if the given set of playback instances are preventing the start of the transition.
 	bool CanStartTransition(const UAvaPlaybackServerTransition* InTransition,
-		const TArray<TWeakPtr<FAvaPlaybackInstance>>& InPlaybackInstancesWeak, bool& bOutShouldDiscard)
+		const TArray<TWeakPtr<FAvaPlaybackInstance>>& InPlaybackInstancesWeak, EAvaPlayableTransitionEntryRole InEntryRole, bool& bOutShouldDiscard)
 	{
 		using namespace UE::AvaPlayback::Utils;
 		for (const TWeakPtr<FAvaPlaybackInstance>& InstanceWeak : InPlaybackInstancesWeak)
@@ -118,7 +117,9 @@ namespace UE::AvaPlaybackServerTransition::Private
 				return false;
 			}
 
-			if (!Playable->IsReadyToPlay())
+			// Only check ready to play for enter playables,
+			// as these are the ones that could be introducing a new asset to prepare
+			if (InEntryRole == EAvaPlayableTransitionEntryRole::Enter && !Playable->IsReadyToPlay())
 			{
 				// Keep the command in the queue for next tick.
 				bOutShouldDiscard = false;
@@ -305,7 +306,7 @@ bool UAvaPlaybackServerTransition::CanStart(bool& bOutShouldDiscard)
 		AsyncAssetLoader->BeginLoadingAssets(Assets.Array());
 	}
 
-	if (!CanStartTransition(this, EnterPlaybackInstancesWeak, bOutShouldDiscard))
+	if (!CanStartTransition(this, EnterPlaybackInstancesWeak, EAvaPlayableTransitionEntryRole::Enter, bOutShouldDiscard))
 	{
 		return false;
 	}
@@ -313,12 +314,12 @@ bool UAvaPlaybackServerTransition::CanStart(bool& bOutShouldDiscard)
 	// Note: need to check the "non-entering" instances too in case the playback commands got delayed
 	// causing playables to also need loading/recovering.
 	
-	if (!CanStartTransition(this, PlayingPlaybackInstancesWeak, bOutShouldDiscard))
+	if (!CanStartTransition(this, PlayingPlaybackInstancesWeak, EAvaPlayableTransitionEntryRole::Playing, bOutShouldDiscard))
 	{
 		return false;
 	}
 
-	if (!CanStartTransition(this, ExitPlaybackInstancesWeak, bOutShouldDiscard))
+	if (!CanStartTransition(this, ExitPlaybackInstancesWeak, EAvaPlayableTransitionEntryRole::Exit, bOutShouldDiscard))
 	{
 		return false;
 	}
diff --git a/Engine/Plugins/VirtualProduction/Avalanche/Source/AvalancheMedia/Private/Rundown/Transition/AvaRundownPageTransition.cpp b/Engine/Plugins/VirtualProduction/Avalanche/Source/AvalancheMedia/Private/Rundown/Transition/AvaRundownPageTransition.cpp
--- a/Engine/Plugins/VirtualProduction/Avalanche/Source/AvalancheMedia/Private/Rundown/Transition/AvaRundownPageTransition.cpp
+++ b/Engine/Plugins/VirtualProduction/Avalanche/Source/AvalancheMedia/Private/Rundown/Transition/AvaRundownPageTransition.cpp
@@ -92,6 +92,23 @@ namespace UE::AvaMedia::Rundown::PageTransition::Private
 		}
 		return nullptr;
 	}
+
+	bool HasLocalPlayable(TNotNull<const UAvaRundownPagePlayer*> InPlayer)
+	{
+		// Check if at least one playable is local
+		const int32 NumInstancePlayers = InPlayer->GetNumInstancePlayers();
+		for (int32 InstancePlayerIndex = 0; InstancePlayerIndex < NumInstancePlayers; ++InstancePlayerIndex)
+		{
+			if (const UAvaPlayable* Playable = GetPlayable(InPlayer, InstancePlayerIndex))
+			{
+				if (!Playable->IsRemoteProxy())
+				{
+					return true;
+				}
+			}
+		}
+		return false;
+	}
 }
 
 UAvaRundownPageTransition* UAvaRundownPageTransition::MakeNew(UAvaRundown* InRundown)
@@ -190,11 +207,17 @@ bool UAvaRundownPageTransition::CanStart(bool& bOutShouldDiscard)
 		TSet<FSoftObjectPath> Assets;
 		for (const TWeakObjectPtr<UAvaRundownPagePlayer>& PlayerWeak : EnterPlayersWeak)
 		{
-			if (const UAvaRundownPagePlayer* Player = PlayerWeak.Get())
+			const UAvaRundownPagePlayer* Player = PlayerWeak.Get();
+
+			// Only collect references and assets if there is at least one playable that isn't remote.
+			// If all playables are remote, no loading needs to take place locally.
+			if (Player && HasLocalPlayable(Player))
 			{
-				const TSharedPtr<FAvaPlayableRemoteControlValues> Values = GetRemoteControlValues(Player);
-				FAvaPlayableRemoteControlValues::CollectReferencedAssetPaths(Values->ControllerValues, Assets);
-				FAvaPlayableRemoteControlValues::CollectReferencedAssetPaths(Values->EntityValues, Assets);
+				if (const TSharedPtr<const FAvaPlayableRemoteControlValues> Values = GetRemoteControlValues(Player))
+				{
+					FAvaPlayableRemoteControlValues::CollectReferencedAssetPaths(Values->ControllerValues, Assets);
+					FAvaPlayableRemoteControlValues::CollectReferencedAssetPaths(Values->EntityValues, Assets);
+				}
 			}
 		}
 		AsyncAssetLoader->BeginLoadingAssets(Assets.Array());
diff --git a/Engine/Plugins/VirtualProduction/Avalanche/Source/AvalancheMedia/Public/Playable/AvaPlayable.h b/Engine/Plugins/VirtualProduction/Avalanche/Source/AvalancheMedia/Public/Playable/AvaPlayable.h
--- a/Engine/Plugins/VirtualProduction/Avalanche/Source/AvalancheMedia/Public/Playable/AvaPlayable.h
+++ b/Engine/Plugins/VirtualProduction/Avalanche/Source/AvalancheMedia/Public/Playable/AvaPlayable.h
@@ -124,7 +124,7 @@ public:
 	 */
 	virtual bool IsReadyToPlay() const
 	{
-		return false;
+		return true;
 	}
 
 	/**
