Motion Design: Cleaned up hotfixed issue where dynamic mesh builder does not call mesh change events, causing the dynamic mesh component to not update its bounds and triggering all sorts of visual errors like no shadows casted or the mesh being incorrectly culled before it's fully out of the frustum.

diff --git a/Engine/Plugins/VirtualProduction/Avalanche/Source/AvalancheModifiers/Private/Modifiers/AvaDynamicMeshConverterModifier.cpp b/Engine/Plugins/VirtualProduction/Avalanche/Source/AvalancheModifiers/Private/Modifiers/AvaDynamicMeshConverterModifier.cpp
--- a/Engine/Plugins/VirtualProduction/Avalanche/Source/AvalancheModifiers/Private/Modifiers/AvaDynamicMeshConverterModifier.cpp
+++ b/Engine/Plugins/VirtualProduction/Avalanche/Source/AvalancheModifiers/Private/Modifiers/AvaDynamicMeshConverterModifier.cpp
@@ -567,27 +567,9 @@ bool UAvaDynamicMeshConverterModifier::ConvertComponents(TArray<TWeakObjectPtr<U
 
 	ConvertedComponents = MoveTemp(NewConvertedComponents);
 
-	UDynamicMesh* const DynamicMesh = DynamicMeshComponent->GetDynamicMesh();
-	if (!IsValid(DynamicMesh))
-	{
-		return false;
-	}
-
-	// Copy ChangeInfo from how it gets built in FCEMeshBuilder::BuildDynamicMesh 
-	FDynamicMeshChangeInfo ChangeInfo;
-	ChangeInfo.Type = EDynamicMeshChangeType::GeneralEdit;
-	ChangeInfo.Flags = EDynamicMeshAttributeChangeFlags::Unknown;
-
-	// Call PreMeshChange to have the same behavior as UDynamicMesh::EditMeshInternal when bDeferChangeEvents is set to false
-	DynamicMesh->OnPreMeshChanged().Broadcast(DynamicMesh, ChangeInfo);
-
-	const bool bResult = MeshBuilder.BuildDynamicMesh(DynamicMesh, OutMaterialsWeak);
-
-	// Call OnMeshChanged and BP Event to have the same behavior as UDynamicMesh::EditMeshInternal when bDeferChangeEvents is set to false
-	DynamicMesh->OnMeshChanged().Broadcast(DynamicMesh, ChangeInfo);
-	DynamicMesh->MeshModifiedBPEvent.Broadcast(DynamicMesh);
-
-	return bResult;
+	FCEMeshBuilder::FCEMeshBuilderParams BuilderParams;
+	BuilderParams.bNotifyMeshChangeEvents = true;
+	return MeshBuilder.BuildDynamicMesh(DynamicMeshComponent->GetDynamicMesh(), OutMaterialsWeak, BuilderParams);
 }
 
 bool UAvaDynamicMeshConverterModifier::HasFlag(EAvaDynamicMeshConverterModifierType InFlag) const
diff --git a/Engine/Plugins/VirtualProduction/ClonerEffector/Source/ClonerEffectorMeshBuilder/Private/CEMeshBuilder.cpp b/Engine/Plugins/VirtualProduction/ClonerEffector/Source/ClonerEffectorMeshBuilder/Private/CEMeshBuilder.cpp
--- a/Engine/Plugins/VirtualProduction/ClonerEffector/Source/ClonerEffectorMeshBuilder/Private/CEMeshBuilder.cpp
+++ b/Engine/Plugins/VirtualProduction/ClonerEffector/Source/ClonerEffectorMeshBuilder/Private/CEMeshBuilder.cpp
@@ -755,6 +755,8 @@ bool FCEMeshBuilder::BuildDynamicMesh(UDynamicMesh* OutMesh, TArray<TWeakObjectP
 
 	OutMaterials.Empty();
 
+	const bool bDeferMeshChangeEvents = !InParams.bNotifyMeshChangeEvents;
+
 	// Lets combine all meshes components from this actor together
 	OutMesh->EditMesh([this, &OutMaterials, &InParams](FDynamicMesh3& InMergedMesh)
 	{
@@ -839,7 +841,7 @@ bool FCEMeshBuilder::BuildDynamicMesh(UDynamicMesh* OutMesh, TArray<TWeakObjectP
 			FMergeCoincidentMeshEdges WeldOp(&InMergedMesh);
 			WeldOp.Apply();
 		}
-	}, EDynamicMeshChangeType::GeneralEdit, EDynamicMeshAttributeChangeFlags::Unknown, /** bDeferChange */true);
+	}, EDynamicMeshChangeType::GeneralEdit, EDynamicMeshAttributeChangeFlags::Unknown, bDeferMeshChangeEvents);
 
 	return true;
 }
diff --git a/Engine/Plugins/VirtualProduction/ClonerEffector/Source/ClonerEffectorMeshBuilder/Public/CEMeshBuilder.h b/Engine/Plugins/VirtualProduction/ClonerEffector/Source/ClonerEffectorMeshBuilder/Public/CEMeshBuilder.h
--- a/Engine/Plugins/VirtualProduction/ClonerEffector/Source/ClonerEffectorMeshBuilder/Public/CEMeshBuilder.h
+++ b/Engine/Plugins/VirtualProduction/ClonerEffector/Source/ClonerEffectorMeshBuilder/Public/CEMeshBuilder.h
@@ -54,6 +54,8 @@ struct CLONEREFFECTORMESHBUILDER_API FCEMeshBuilder
 	{
 		/** Merge same meshes material slot in the final result */
 		bool bMergeMaterials = false;
+		/** Whether to call the PreMeshChangedEvent and MeshChangedEvents during the build */
+		bool bNotifyMeshChangeEvents = false;
 	};
 
 	struct FCEMeshBuilderAppendParams
