Motion Design: added cache shader step when loading in a page, via the option '-CacheShaders'.

diff --git a/Engine/Plugins/VirtualProduction/Avalanche/Source/AvalancheMedia/Internal/Playback/AvaPlaybackUtils.h b/Engine/Plugins/VirtualProduction/Avalanche/Source/AvalancheMedia/Internal/Playback/AvaPlaybackUtils.h
--- a/Engine/Plugins/VirtualProduction/Avalanche/Source/AvalancheMedia/Internal/Playback/AvaPlaybackUtils.h
+++ b/Engine/Plugins/VirtualProduction/Avalanche/Source/AvalancheMedia/Internal/Playback/AvaPlaybackUtils.h
@@ -63,6 +63,13 @@ namespace UE::AvaPlayback::Utils
 	 */
 	AVALANCHEMEDIA_API FString GetBriefFrameInfo();
 
+	/**
+	 * Processes shader compilation results if any.
+	 * @param InTimeBudget the time budget limit, in seconds, to process async results
+	 * @return true if shaders are compiling and async results (if any) began processing, false otherwise.
+	 */
+	bool ProcessShadersCompiling(TOptional<float> InTimeBudget);
+
 	template<typename InEnumType>
 	FString StaticEnumToString(InEnumType InValue)
 	{
diff --git a/Engine/Plugins/VirtualProduction/Avalanche/Source/AvalancheMedia/Private/Playable/Playables/AvaPlayableLevelStreaming.cpp b/Engine/Plugins/VirtualProduction/Avalanche/Source/AvalancheMedia/Private/Playable/Playables/AvaPlayableLevelStreaming.cpp
--- a/Engine/Plugins/VirtualProduction/Avalanche/Source/AvalancheMedia/Private/Playable/Playables/AvaPlayableLevelStreaming.cpp
+++ b/Engine/Plugins/VirtualProduction/Avalanche/Source/AvalancheMedia/Private/Playable/Playables/AvaPlayableLevelStreaming.cpp
@@ -111,6 +111,15 @@ namespace UE::AvaMedia::LevelStreamingPlayable::Private
 		return false;
 	}
 
+	/**
+	 * Utility function to determine whether to cache shaders from load options.
+	 * returns true if shaders should be cached on load.
+	 */
+	bool ShouldCacheShadersOnLoad(const FString& InOptions)
+	{
+		return FParse::Param(*InOptions, TEXT("CacheShaders"));
+	}
+
 	/**
 	 * Version of FLevelUtils::ApplyLevelTransform specialized to work with animated objects with sequencer.
 	 * All root actors get added to a root null and the root is moved.
@@ -281,11 +290,12 @@ namespace UE::AvaMedia::LevelStreamingPlayable::Private
 	}
 
 	/** Calls cache shaders for each material in the set, until the time budget is exceeded */
-	void CacheMaterials(const TSet<UMaterialInterface*>& InMaterialsToCache, double InTimeBudget)
+	void CacheMaterials(const TSet<UMaterialInterface*>& InMaterialsToCache, TOptional<float> InTimeBudget)
 	{
 		TRACE_CPUPROFILER_EVENT_SCOPE(MotionDesignPlayable_CacheMaterials)
 
 		const double StartTime = FPlatformTime::Seconds();
+		const double TimeBudget = InTimeBudget.Get(0.f);
 
 		FMaterialUpdateContext UpdateContext(0);
 		for (UMaterialInterface* MaterialInterface : InMaterialsToCache)
@@ -297,7 +307,7 @@ namespace UE::AvaMedia::LevelStreamingPlayable::Private
 			UpdateContext.AddMaterialInterface(MaterialInterface);
 			MaterialInterface->CacheShaders(PrecompileMode);
 
-			if (FPlatformTime::Seconds() - StartTime > InTimeBudget)
+			if (InTimeBudget.IsSet() && FPlatformTime::Seconds() - StartTime > TimeBudget)
 			{
 				return;
 			}
@@ -318,6 +328,9 @@ bool UAvaPlayableLevelStreaming::LoadAsset(const FAvaSoftAssetPtr& InSourceAsset
 	// Setup the level transform from the load options.
 	bHasTransform = UE::AvaMedia::LevelStreamingPlayable::Private::GetTransformFromOptions(LevelTransform, InLoadOptions, PlayableGroup, InSourceAsset);
 
+	// Determine if shaders should be cached on load
+	bCacheShadersOnLoad = UE::AvaMedia::LevelStreamingPlayable::Private::ShouldCacheShadersOnLoad(InLoadOptions);
+
 	const FAvaInstanceSettings& PlaybackInstanceSettings = IAvaMediaModule::Get().GetAvaInstanceSettings();
 	bLoadSubPlayables = PlaybackInstanceSettings.bEnableLoadSubPlayables;
 	
@@ -437,57 +450,10 @@ bool UAvaPlayableLevelStreaming::IsReadyToPlay() const
 		return false;
 	}
 
-	TSet<UMaterialInterface*> MaterialsToCache;
-
-	const UE::Ava::FMaterialBridgeRegistry& MaterialBridgeRegistry = UE::Ava::FMaterialBridgeRegistry::Get();
-
-	for (AActor* Actor : Level->Actors)
-	{
-		if (const UE::Ava::FMaterialBridge* ActorMaterialBridge = MaterialBridgeRegistry.GetMaterialBridgeForObject(Actor))
-		{
-			ActorMaterialBridge->ForEachUsedMaterial(UE::Ava::FConstDataView(Actor),
-				[&MaterialsToCache](UMaterialInterface* InMaterialInterface)
-				{
-					if (!InMaterialInterface)
-					{
-						return;
-					}
-
-					UMaterial* const Material = InMaterialInterface->GetMaterial();
-					if (Material && Material != InMaterialInterface && !Material->IsComplete())
-					{
-						MaterialsToCache.Add(Material);
-					}
-					if (!InMaterialInterface->IsComplete())
-					{
-						MaterialsToCache.Add(InMaterialInterface);
-					}
-				});
-		}
-	}
-
-	if (MaterialsToCache.IsEmpty())
-	{
-		bool bFinishedCompilation = true;
-
-		if (GShaderCompilingManager)
-		{
-			if (GShaderCompilingManager->IsCompiling())
-			{
-				constexpr bool bLimitExecutionTime = true;
-				constexpr bool bBlockOnGlobalShaderCompletion = false;
-				GShaderCompilingManager->ProcessAsyncResults(bLimitExecutionTime, bBlockOnGlobalShaderCompletion);
-
-				bFinishedCompilation = false;
-			}
-		}
-
-		return bFinishedCompilation;
-	}
-
-	constexpr double TimeBudget = 0.001;
-	Private::CacheMaterials(MaterialsToCache, TimeBudget);
-	return false;
+	// Check whether the compilation has finished
+	constexpr float TimeBudget = 0.001f;
+	const bool bIsCachingShaders = CacheMaterials(Level, TimeBudget);
+	return !bIsCachingShaders;
 }
 
 IAvaSceneInterface* UAvaPlayableLevelStreaming::GetSceneInterface() const
@@ -588,6 +554,50 @@ bool UAvaPlayableLevelStreaming::LoadLevel(const TSoftObjectPtr<UWorld>& InSourc
 	return true;
 }
 
+bool UAvaPlayableLevelStreaming::CacheMaterials(TNotNull<const ULevel*> InLevel, TOptional<float> InTimeBudget) const
+{
+	TRACE_CPUPROFILER_EVENT_SCOPE(UAvaPlayableLevelStreaming::CacheMaterials)
+
+	using namespace UE::AvaMedia::LevelStreamingPlayable;
+
+	TSet<UMaterialInterface*> MaterialsToCache;
+
+	const UE::Ava::FMaterialBridgeRegistry& MaterialBridgeRegistry = UE::Ava::FMaterialBridgeRegistry::Get();
+	for (AActor* Actor : InLevel->Actors)
+	{
+		if (const UE::Ava::FMaterialBridge* ActorMaterialBridge = MaterialBridgeRegistry.GetMaterialBridgeForObject(Actor))
+		{
+			ActorMaterialBridge->ForEachUsedMaterial(UE::Ava::FConstDataView(Actor),
+				[&MaterialsToCache](UMaterialInterface* InMaterialInterface)
+				{
+					if (!InMaterialInterface)
+					{
+						return;
+					}
+
+					UMaterial* const Material = InMaterialInterface->GetMaterial();
+					if (Material && Material != InMaterialInterface && !Material->IsComplete())
+					{
+						MaterialsToCache.Add(Material);
+					}
+					if (!InMaterialInterface->IsComplete())
+					{
+						MaterialsToCache.Add(InMaterialInterface);
+					}
+				});
+		}
+	}
+
+	if (MaterialsToCache.IsEmpty())
+	{
+		const bool bAreShadersCompiling = UE::AvaPlayback::Utils::ProcessShadersCompiling(InTimeBudget);
+		return bAreShadersCompiling;
+	}
+
+	Private::CacheMaterials(MaterialsToCache, InTimeBudget);
+	return true;
+}
+
 void UAvaPlayableLevelStreaming::HandleTransitionEvent(UAvaPlayable* InPlayable, UAvaPlayableTransition* InTransition, EAvaPlayableTransitionEventFlags InTransitionFlags)
 {
 	if (InPlayable == this && EnumHasAnyFlags(InTransitionFlags, EAvaPlayableTransitionEventFlags::ShowPlayable))
@@ -696,6 +706,36 @@ void UAvaPlayableLevelStreaming::OnLevelStreamingStateChanged_Synchronized(ELeve
 		
 		// Resolve the ava scene for the other operations.
 		ResolveScene(Level);
+
+		if (bCacheShadersOnLoad)
+		{
+			static constexpr float TimeBudget = 0.001f;
+
+			// If there was more than one material that required caching,
+			// kick off a ticker that incrementally caches materials with same time budget,
+			// until there's no materials to cache and all shaders have completed compilation
+			if (CacheMaterials(Level, TimeBudget))
+			{
+				bCachingShadersOnLoad = true;
+
+				FTSTicker::GetCoreTicker().AddTicker(FTickerDelegate::CreateWeakLambda(this,
+					[this](float)
+					{
+						const ULevel* const Level = LevelStreaming ? LevelStreaming->GetLoadedLevel() : nullptr;
+						if (Level)
+						{
+							bCachingShadersOnLoad = CacheMaterials(Level, TimeBudget);
+						}
+						else
+						{
+							bCachingShadersOnLoad = false;
+						}
+
+						// Continue ticking until caching shaders stops
+						return bCachingShadersOnLoad;
+					}));
+			}
+		}
 	}
 	
 	// Important: playable status gets updated in the synchronized event handler.
diff --git a/Engine/Plugins/VirtualProduction/Avalanche/Source/AvalancheMedia/Private/Playable/Playables/AvaPlayableLevelStreaming.h b/Engine/Plugins/VirtualProduction/Avalanche/Source/AvalancheMedia/Private/Playable/Playables/AvaPlayableLevelStreaming.h
--- a/Engine/Plugins/VirtualProduction/Avalanche/Source/AvalancheMedia/Private/Playable/Playables/AvaPlayableLevelStreaming.h
+++ b/Engine/Plugins/VirtualProduction/Avalanche/Source/AvalancheMedia/Private/Playable/Playables/AvaPlayableLevelStreaming.h
@@ -49,6 +49,14 @@ public:
 protected:
 	bool LoadLevel(const TSoftObjectPtr<UWorld>& InSourceLevel, const FTransform& InTransform, bool bInInitiallyVisible);
 
+	/**
+	 * Caches the shaders for the given level.
+	 * @param InLevel the level to cache shaders for
+	 * @param InTimeBudget the time budget, in seconds, given to cache materials
+	 * @return true if caching is taking place. false otherwise. 
+	 */
+	bool CacheMaterials(TNotNull<const ULevel*> InLevel, TOptional<float> InTimeBudget) const;
+
 	void HandleTransitionEvent(UAvaPlayable* InPlayable, UAvaPlayableTransition* InTransition, EAvaPlayableTransitionEventFlags InTransitionFlags);
 
 	void OnLevelStreamingStateChanged(UWorld* InWorld
@@ -125,6 +133,12 @@ protected:
 	/** If true, all primitives from the playable will hidden. */
 	bool bShouldBeHidden = false;
 
+	/** True if shaders should be cached when loading the level */
+	bool bCacheShadersOnLoad = false;
+
+	/** True if shaders are currently being compiled */
+	bool bCachingShadersOnLoad = false;
+
 	/** True if the load command specified a transform. */
 	bool bHasTransform = false;
 
diff --git a/Engine/Plugins/VirtualProduction/Avalanche/Source/AvalancheMedia/Private/Playback/AvaPlaybackUtils.cpp b/Engine/Plugins/VirtualProduction/Avalanche/Source/AvalancheMedia/Private/Playback/AvaPlaybackUtils.cpp
--- a/Engine/Plugins/VirtualProduction/Avalanche/Source/AvalancheMedia/Private/Playback/AvaPlaybackUtils.cpp
+++ b/Engine/Plugins/VirtualProduction/Avalanche/Source/AvalancheMedia/Private/Playback/AvaPlaybackUtils.cpp
@@ -240,6 +240,20 @@ namespace UE::AvaPlayback::Utils
 		return FString::Printf(TEXT("[%d]"), GFrameNumber);
 	}
 
+	bool ProcessShadersCompiling(TOptional<float> InTimeBudget)
+	{
+		if (GShaderCompilingManager)
+		{
+			if (GShaderCompilingManager->IsCompiling())
+			{
+				constexpr bool bBlockOnGlobalShaderCompletion = false;
+				GShaderCompilingManager->ProcessAsyncResults(InTimeBudget.Get(0.f), bBlockOnGlobalShaderCompletion);
+				return true;
+			}
+		}
+		return false;
+	}
+
 	void FAsyncAssetLoader::HandleStreamableObject(UObject* InObject)
 	{
 		if (!Private::CVarRuntimePrestreamRCReferencedAssets.GetValueOnGameThread())
@@ -309,22 +323,9 @@ namespace UE::AvaPlayback::Utils
 			return false;
 		}
 
-		bool bIsCompilationComplete = true;
-
-		// Only wait for compilation if there were streaming assets that were added
-		if (bWaitForStreamingAssets && GShaderCompilingManager)
-		{
-			if (GShaderCompilingManager->IsCompiling())
-			{
-				constexpr bool bLimitExecutionTime = true;
-				constexpr bool bBlockOnGlobalShaderCompletion = false;
-				GShaderCompilingManager->ProcessAsyncResults(bLimitExecutionTime, bBlockOnGlobalShaderCompletion);
-
-				bIsCompilationComplete = false;
-			}
-		}
-
-		return bIsCompilationComplete;
+		constexpr float TimeBudget = 0.001f;
+		const bool bAreShadersCompiling = bWaitForStreamingAssets && ProcessShadersCompiling(TimeBudget);
+		return !bAreShadersCompiling;
 	}
 
 	void FAsyncAssetLoader::CleanupStreamingAssets() const
diff --git a/Engine/Plugins/VirtualProduction/Avalanche/Source/AvalancheMedia/Private/Rundown/AvaRundown.cpp b/Engine/Plugins/VirtualProduction/Avalanche/Source/AvalancheMedia/Private/Rundown/AvaRundown.cpp
--- a/Engine/Plugins/VirtualProduction/Avalanche/Source/AvalancheMedia/Private/Rundown/AvaRundown.cpp
+++ b/Engine/Plugins/VirtualProduction/Avalanche/Source/AvalancheMedia/Private/Rundown/AvaRundown.cpp
@@ -1172,6 +1172,9 @@ TArray<UAvaRundown::FLoadedInstanceInfo> UAvaRundown::LoadPage(int32 InPageId,
 	FString LoadOptions;
 	FAvaRundownPageCommandContext PageCommandContext = {*this, Page, ChannelName};
 
+	// Built-in Load Options:
+	LoadOptions += TEXT(" -CacheShaders");
+
 	Page.ForEachInstancedCommands(
 		[&PageCommandContext, &LoadOptions](const FAvaRundownPageCommand& InCommand, const FAvaRundownPage& InPage)
 		{
