Motion Design: renamed material bridge function that iterates materials to VisitMaterials and changed the callable signature to pass in a context with information about the material and owner (rather than just material).
+ Added StaticMesh material bridge

diff --git a/Engine/Plugins/VirtualProduction/Avalanche/Source/AvalancheMaterial/Private/AvaActorMaterialBridge.cpp b/Engine/Plugins/VirtualProduction/Avalanche/Source/AvalancheMaterial/Private/AvaActorMaterialBridge.cpp
--- a/Engine/Plugins/VirtualProduction/Avalanche/Source/AvalancheMaterial/Private/AvaActorMaterialBridge.cpp
+++ b/Engine/Plugins/VirtualProduction/Avalanche/Source/AvalancheMaterial/Private/AvaActorMaterialBridge.cpp
@@ -12,18 +12,20 @@ const UStruct* FActorMaterialBridge::OnGetBridgedType() const
 	return AActor::StaticClass();
 }
 
-void FActorMaterialBridge::OnForEachUsedMaterial(FConstDataView InDataView, TFunctionRef<void(UMaterialInterface*)> InFunc, const FGetUsedMaterialOptions& InOptions) const
+void FActorMaterialBridge::OnVisitMaterials(FConstDataView InDataView, TFunctionRef<void(const FVisitedContext&)> InFunc, const FVisitOptions& InOptions) const
 {
 	const FMaterialBridgeRegistry& MaterialBridgeRegistry = FMaterialBridgeRegistry::Get();
 
 	const AActor& Actor = InDataView.Get<const AActor>();
 
-	Actor.ForEachComponent(InOptions.bIncludeAttachedMaterials,
+	constexpr bool bIncludeFromChildActors = true;
+
+	Actor.ForEachComponent(bIncludeFromChildActors,
 		[&MaterialBridgeRegistry, &InFunc, &InOptions](UActorComponent* InComponent)
 		{
 			if (const FMaterialBridge* MaterialBridge = MaterialBridgeRegistry.GetMaterialBridgeForObject(InComponent))
 			{
-				MaterialBridge->ForEachUsedMaterial(FConstDataView(InComponent), InFunc, InOptions);
+				MaterialBridge->VisitMaterials(FConstDataView(InComponent), InFunc, InOptions);
 			}
 		});
 }
diff --git a/Engine/Plugins/VirtualProduction/Avalanche/Source/AvalancheMaterial/Private/AvaDecalComponentMaterialBridge.cpp b/Engine/Plugins/VirtualProduction/Avalanche/Source/AvalancheMaterial/Private/AvaDecalComponentMaterialBridge.cpp
--- a/Engine/Plugins/VirtualProduction/Avalanche/Source/AvalancheMaterial/Private/AvaDecalComponentMaterialBridge.cpp
+++ b/Engine/Plugins/VirtualProduction/Avalanche/Source/AvalancheMaterial/Private/AvaDecalComponentMaterialBridge.cpp
@@ -12,10 +12,15 @@ const UStruct* FDecalComponentMaterialBridge::OnGetBridgedType() const
 	return UDecalComponent::StaticClass();
 }
 
-void FDecalComponentMaterialBridge::OnForEachUsedMaterial(FConstDataView InDataView, TFunctionRef<void(UMaterialInterface*)> InFunc, const FGetUsedMaterialOptions& InOptions) const
+void FDecalComponentMaterialBridge::OnVisitMaterials(FConstDataView InDataView, TFunctionRef<void(const FVisitedContext&)> InFunc, const FVisitOptions& InOptions) const
 {
 	const UDecalComponent& Component = InDataView.Get<const UDecalComponent>();
-	InFunc(Component.GetDecalMaterial());
+
+	FVisitedContext Context;
+	Context.Owner = InDataView;
+	Context.Material = Component.GetDecalMaterial();
+
+	InFunc(Context);
 }
 
 } // UE::Ava
diff --git a/Engine/Plugins/VirtualProduction/Avalanche/Source/AvalancheMaterial/Private/AvaMaterialBridge.cpp b/Engine/Plugins/VirtualProduction/Avalanche/Source/AvalancheMaterial/Private/AvaMaterialBridge.cpp
--- a/Engine/Plugins/VirtualProduction/Avalanche/Source/AvalancheMaterial/Private/AvaMaterialBridge.cpp
+++ b/Engine/Plugins/VirtualProduction/Avalanche/Source/AvalancheMaterial/Private/AvaMaterialBridge.cpp
@@ -15,11 +15,11 @@ const UStruct* FMaterialBridge::GetBridgedType() const
 	return BridgedType;
 }
 
-void FMaterialBridge::ForEachUsedMaterial(FConstDataView InDataView, TFunctionRef<void(UMaterialInterface*)> InFunc, const FGetUsedMaterialOptions& InOptions) const
+void FMaterialBridge::VisitMaterials(FConstDataView InDataView, TFunctionRef<void(const FVisitedContext&)> InFunc, const FVisitOptions& InOptions) const
 {
 	if (InDataView.IsValidFor(BridgedType))
 	{
-		OnForEachUsedMaterial(InDataView, InFunc, InOptions);
+		OnVisitMaterials(InDataView, InFunc, InOptions);
 	}
 }
 
diff --git a/Engine/Plugins/VirtualProduction/Avalanche/Source/AvalancheMaterial/Private/AvaMaterialModule.cpp b/Engine/Plugins/VirtualProduction/Avalanche/Source/AvalancheMaterial/Private/AvaMaterialModule.cpp
--- a/Engine/Plugins/VirtualProduction/Avalanche/Source/AvalancheMaterial/Private/AvaMaterialModule.cpp
+++ b/Engine/Plugins/VirtualProduction/Avalanche/Source/AvalancheMaterial/Private/AvaMaterialModule.cpp
@@ -6,6 +6,7 @@
 #include "AvaMaterialBridge.h"
 #include "AvaMaterialBridgeRegistry.h"
 #include "AvaPrimitiveComponentMaterialBridge.h"
+#include "AvaStaticMeshMaterialBridge.h"
 #include "Modules/ModuleManager.h"
 
 IMPLEMENT_MODULE(FAvaMaterialModule, AvalancheMaterial)
@@ -18,4 +19,5 @@ void FAvaMaterialModule::StartupModule()
 	MaterialBridgeRegistry.Register<UE::Ava::FActorMaterialBridge>();
 	MaterialBridgeRegistry.Register<UE::Ava::FPrimitiveComponentMaterialBridge>();
 	MaterialBridgeRegistry.Register<UE::Ava::FDecalComponentMaterialBridge>();
+	MaterialBridgeRegistry.Register<UE::Ava::FStaticMeshMaterialBridge>();
 }
diff --git a/Engine/Plugins/VirtualProduction/Avalanche/Source/AvalancheMaterial/Private/AvaPrimitiveComponentMaterialBridge.cpp b/Engine/Plugins/VirtualProduction/Avalanche/Source/AvalancheMaterial/Private/AvaPrimitiveComponentMaterialBridge.cpp
--- a/Engine/Plugins/VirtualProduction/Avalanche/Source/AvalancheMaterial/Private/AvaPrimitiveComponentMaterialBridge.cpp
+++ b/Engine/Plugins/VirtualProduction/Avalanche/Source/AvalancheMaterial/Private/AvaPrimitiveComponentMaterialBridge.cpp
@@ -12,16 +12,20 @@ const UStruct* FPrimitiveComponentMaterialBridge::OnGetBridgedType() const
 	return UPrimitiveComponent::StaticClass();
 }
 
-void FPrimitiveComponentMaterialBridge::OnForEachUsedMaterial(FConstDataView InDataView, TFunctionRef<void(UMaterialInterface*)> InFunc, const FGetUsedMaterialOptions& InOptions) const
+void FPrimitiveComponentMaterialBridge::OnVisitMaterials(FConstDataView InDataView, TFunctionRef<void(const FVisitedContext&)> InFunc, const FVisitOptions& InOptions) const
 {
 	const UPrimitiveComponent& Component = InDataView.Get<const UPrimitiveComponent>();
 
 	TArray<UMaterialInterface*> Materials;
-	Component.GetUsedMaterials(Materials, InOptions.bIncludeDebugMaterials);
+	Component.GetUsedMaterials(Materials, InOptions.bVisitDebugMaterials);
+
+	FVisitedContext Context;
+	Context.Owner = InDataView;
 
 	for (UMaterialInterface* Material : Materials)
 	{
-		InFunc(Material);
+		Context.Material = Material;
+		InFunc(Context);
 	}
 }
 
diff --git a/Engine/Plugins/VirtualProduction/Avalanche/Source/AvalancheMaterial/Private/AvaStaticMeshMaterialBridge.cpp b/Engine/Plugins/VirtualProduction/Avalanche/Source/AvalancheMaterial/Private/AvaStaticMeshMaterialBridge.cpp
new file mode 100644
--- /dev/null
+++ b/Engine/Plugins/VirtualProduction/Avalanche/Source/AvalancheMaterial/Private/AvaStaticMeshMaterialBridge.cpp
@@ -0,0 +1,31 @@
+// Copyright Epic Games, Inc. All Rights Reserved.
+
+#include "AvaStaticMeshMaterialBridge.h"
+#include "Engine/StaticMesh.h"
+
+namespace UE::Ava
+{
+
+const UStruct* FStaticMeshMaterialBridge::OnGetBridgedType() const
+{
+	return UStaticMesh::StaticClass();
+}
+
+void FStaticMeshMaterialBridge::OnVisitMaterials(FConstDataView InDataView, TFunctionRef<void(const FVisitedContext&)> InFunc, const FVisitOptions& InOptions) const
+{
+	const UStaticMesh& StaticMesh = InDataView.Get<const UStaticMesh>();
+
+	FVisitedContext Context;
+	Context.Owner = InDataView;
+
+	for (const FStaticMaterial& StaticMaterial : StaticMesh.GetStaticMaterials())
+	{
+		Context.Material = StaticMaterial.MaterialInterface;
+		InFunc(Context);
+
+		Context.Material = StaticMaterial.OverlayMaterialInterface;
+		InFunc(Context);
+	}
+}
+
+} // UE::Ava
diff --git a/Engine/Plugins/VirtualProduction/Avalanche/Source/AvalancheMaterial/Public/AvaActorMaterialBridge.h b/Engine/Plugins/VirtualProduction/Avalanche/Source/AvalancheMaterial/Public/AvaActorMaterialBridge.h
--- a/Engine/Plugins/VirtualProduction/Avalanche/Source/AvalancheMaterial/Public/AvaActorMaterialBridge.h
+++ b/Engine/Plugins/VirtualProduction/Avalanche/Source/AvalancheMaterial/Public/AvaActorMaterialBridge.h
@@ -15,7 +15,7 @@ class FActorMaterialBridge : public FMaterialBridge
 protected:
 	//~ Begin FMaterialBridge
 	UE_API virtual const UStruct* OnGetBridgedType() const override;
-	UE_API virtual void OnForEachUsedMaterial(FConstDataView InDataView, TFunctionRef<void(UMaterialInterface*)> InFunc, const FGetUsedMaterialOptions& InOptions) const override;
+	UE_API virtual void OnVisitMaterials(FConstDataView InDataView, TFunctionRef<void(const FVisitedContext&)> InFunc, const FVisitOptions& InOptions) const override;
 	//~ End FMaterialBridge
 };
 
diff --git a/Engine/Plugins/VirtualProduction/Avalanche/Source/AvalancheMaterial/Public/AvaDecalComponentMaterialBridge.h b/Engine/Plugins/VirtualProduction/Avalanche/Source/AvalancheMaterial/Public/AvaDecalComponentMaterialBridge.h
--- a/Engine/Plugins/VirtualProduction/Avalanche/Source/AvalancheMaterial/Public/AvaDecalComponentMaterialBridge.h
+++ b/Engine/Plugins/VirtualProduction/Avalanche/Source/AvalancheMaterial/Public/AvaDecalComponentMaterialBridge.h
@@ -15,7 +15,7 @@ class FDecalComponentMaterialBridge : public FMaterialBridge
 protected:
 	//~ Begin FMaterialBridge
 	UE_API virtual const UStruct* OnGetBridgedType() const override;
-	UE_API virtual void OnForEachUsedMaterial(FConstDataView InDataView, TFunctionRef<void(UMaterialInterface*)> InFunc, const FGetUsedMaterialOptions& InOptions) const override;
+	UE_API virtual void OnVisitMaterials(FConstDataView InDataView, TFunctionRef<void(const FVisitedContext&)> InFunc, const FVisitOptions& InOptions) const override;
 	//~ End FMaterialBridge
 };
 
diff --git a/Engine/Plugins/VirtualProduction/Avalanche/Source/AvalancheMaterial/Public/AvaMaterialBridge.h b/Engine/Plugins/VirtualProduction/Avalanche/Source/AvalancheMaterial/Public/AvaMaterialBridge.h
--- a/Engine/Plugins/VirtualProduction/Avalanche/Source/AvalancheMaterial/Public/AvaMaterialBridge.h
+++ b/Engine/Plugins/VirtualProduction/Avalanche/Source/AvalancheMaterial/Public/AvaMaterialBridge.h
@@ -28,29 +28,34 @@ public:
 	/** Retrieves the type that this bridge supports */
 	UE_API const UStruct* GetBridgedType() const;
 
-	struct FGetUsedMaterialOptions
+	struct FVisitedContext
+	{
+		/** Owner of the visited material */
+		FConstDataView Owner;
+		/** The visited material */
+		TObjectPtr<UMaterialInterface> Material;
+	};
+	struct FVisitOptions
 	{
 		/** Declare default ctor explicitly due to https://gcc.gnu.org/bugzilla/show_bug.cgi?id=88165 */
-		FGetUsedMaterialOptions() {};
-		/** Whether to include materials from attachments (e.g. whether to include from attached actors when the object is an actor) */
-		bool bIncludeAttachedMaterials = false;
+		FVisitOptions() {};
 		/** Whether to include debug materials */
-		bool bIncludeDebugMaterials = false;
+		bool bVisitDebugMaterials = false;
 	};
 	/**
-	 * Iterates all the used materials in the given data view.
+	 * Visits all the used materials in the given data view.
 	 * @param InDataView the data view to find materials from.
 	 * @param InFunc function to call for each material found.
-	 * @param InOptions additional options to get materials.
+	 * @param InOptions additional visit options.
 	 */
-	UE_API void ForEachUsedMaterial(FConstDataView InDataView, TFunctionRef<void(UMaterialInterface*)> InFunc, const FGetUsedMaterialOptions& InOptions = {}) const;
+	UE_API void VisitMaterials(FConstDataView InDataView, TFunctionRef<void(const FVisitedContext&)> InFunc, const FVisitOptions& InOptions = {}) const;
 
 protected:
 	/** Retrieves the type that this bridge supports. Only called once in Initialize() */
 	virtual const UStruct* OnGetBridgedType() const = 0;
 
 	/** Get all the used materials in the given data view. */
-	virtual void OnForEachUsedMaterial(FConstDataView InDataView, TFunctionRef<void(UMaterialInterface*)> InFunc, const FGetUsedMaterialOptions& InOptions) const = 0;
+	virtual void OnVisitMaterials(FConstDataView InDataView, TFunctionRef<void(const FVisitedContext&)> InFunc, const FVisitOptions& InOptions) const = 0;
 
 private:
 	/** The type that this Material Bridge works on */
diff --git a/Engine/Plugins/VirtualProduction/Avalanche/Source/AvalancheMaterial/Public/AvaMaterialBridgeRegistry.h b/Engine/Plugins/VirtualProduction/Avalanche/Source/AvalancheMaterial/Public/AvaMaterialBridgeRegistry.h
--- a/Engine/Plugins/VirtualProduction/Avalanche/Source/AvalancheMaterial/Public/AvaMaterialBridgeRegistry.h
+++ b/Engine/Plugins/VirtualProduction/Avalanche/Source/AvalancheMaterial/Public/AvaMaterialBridgeRegistry.h
@@ -30,7 +30,7 @@ public:
 		this->RegisterInternal(MakeUnique<InMaterialBridgeType>(InArgs...));
 	}
 
-	/** Tries to find the most relevant material bridge for a given object, returning a Null Material Bridge if object is invalid, or if a relevant material bridge was not found. */
+	/** Tries to find the most relevant material bridge for a given object, returning a null if object is invalid, or if a relevant material bridge was not found. */
 	UE_API const FMaterialBridge* GetMaterialBridgeForObject(const UObject* InObject) const;
 
 	/** Tries to find the most relevant material bridge for a given type, returning null if a relevant material bridge was not found. */
diff --git a/Engine/Plugins/VirtualProduction/Avalanche/Source/AvalancheMaterial/Public/AvaPrimitiveComponentMaterialBridge.h b/Engine/Plugins/VirtualProduction/Avalanche/Source/AvalancheMaterial/Public/AvaPrimitiveComponentMaterialBridge.h
--- a/Engine/Plugins/VirtualProduction/Avalanche/Source/AvalancheMaterial/Public/AvaPrimitiveComponentMaterialBridge.h
+++ b/Engine/Plugins/VirtualProduction/Avalanche/Source/AvalancheMaterial/Public/AvaPrimitiveComponentMaterialBridge.h
@@ -15,7 +15,7 @@ class FPrimitiveComponentMaterialBridge : public FMaterialBridge
 protected:
 	//~ Begin FMaterialBridge
 	UE_API virtual const UStruct* OnGetBridgedType() const override;
-	UE_API virtual void OnForEachUsedMaterial(FConstDataView InDataView, TFunctionRef<void(UMaterialInterface*)> InFunc, const FGetUsedMaterialOptions& InOptions) const override;
+	UE_API virtual void OnVisitMaterials(FConstDataView InDataView, TFunctionRef<void(const FVisitedContext&)> InFunc, const FVisitOptions& InOptions) const override;
 	//~ End FMaterialBridge
 };
 
diff --git a/Engine/Plugins/VirtualProduction/Avalanche/Source/AvalancheMaterial/Public/AvaStaticMeshMaterialBridge.h b/Engine/Plugins/VirtualProduction/Avalanche/Source/AvalancheMaterial/Public/AvaStaticMeshMaterialBridge.h
new file mode 100644
--- /dev/null
+++ b/Engine/Plugins/VirtualProduction/Avalanche/Source/AvalancheMaterial/Public/AvaStaticMeshMaterialBridge.h
@@ -0,0 +1,27 @@
+// Copyright Epic Games, Inc. All Rights Reserved.
+
+#pragma once
+
+#include "AvaMaterialBridge.h"
+
+#define UE_API AVALANCHEMATERIAL_API
+
+namespace UE::Ava
+{
+
+/**
+ * Material Bridge for Static Meshes
+ * Only deals with the Static Mesh's static materials.
+ */
+class FStaticMeshMaterialBridge : public FMaterialBridge
+{
+protected:
+	//~ Begin FMaterialBridge
+	UE_API virtual const UStruct* OnGetBridgedType() const override;
+	UE_API virtual void OnVisitMaterials(FConstDataView InDataView, TFunctionRef<void(const FVisitedContext&)> InFunc, const FVisitOptions& InOptions) const override;
+	//~ End FMaterialBridge
+};
+
+} // UE::Ava
+
+#undef UE_API
diff --git a/Engine/Plugins/VirtualProduction/Avalanche/Source/AvalancheMedia/Private/Playable/Playables/AvaPlayableLevelStreaming.cpp b/Engine/Plugins/VirtualProduction/Avalanche/Source/AvalancheMedia/Private/Playable/Playables/AvaPlayableLevelStreaming.cpp
--- a/Engine/Plugins/VirtualProduction/Avalanche/Source/AvalancheMedia/Private/Playable/Playables/AvaPlayableLevelStreaming.cpp
+++ b/Engine/Plugins/VirtualProduction/Avalanche/Source/AvalancheMedia/Private/Playable/Playables/AvaPlayableLevelStreaming.cpp
@@ -567,22 +567,22 @@ bool UAvaPlayableLevelStreaming::CacheMaterials(TNotNull<const ULevel*> InLevel,
 	{
 		if (const UE::Ava::FMaterialBridge* ActorMaterialBridge = MaterialBridgeRegistry.GetMaterialBridgeForObject(Actor))
 		{
-			ActorMaterialBridge->ForEachUsedMaterial(UE::Ava::FConstDataView(Actor),
-				[&MaterialsToCache](UMaterialInterface* InMaterialInterface)
+			ActorMaterialBridge->VisitMaterials(UE::Ava::FConstDataView(Actor),
+				[&MaterialsToCache](const UE::Ava::FMaterialBridge::FVisitedContext& InContext)
 				{
-					if (!InMaterialInterface)
+					if (!InContext.Material)
 					{
 						return;
 					}
 
-					UMaterial* const Material = InMaterialInterface->GetMaterial();
-					if (Material && Material != InMaterialInterface && !Material->IsComplete())
+					UMaterial* const Material = InContext.Material->GetMaterial();
+					if (Material && Material != InContext.Material && !Material->IsComplete())
 					{
 						MaterialsToCache.Add(Material);
 					}
-					if (!InMaterialInterface->IsComplete())
+					if (!InContext.Material->IsComplete())
 					{
-						MaterialsToCache.Add(InMaterialInterface);
+						MaterialsToCache.Add(InContext.Material);
 					}
 				});
 		}
diff --git a/Engine/Plugins/VirtualProduction/Avalanche/Source/AvalancheShapes/Private/AvaShapeComponentMaterialBridge.cpp b/Engine/Plugins/VirtualProduction/Avalanche/Source/AvalancheShapes/Private/AvaShapeComponentMaterialBridge.cpp
--- a/Engine/Plugins/VirtualProduction/Avalanche/Source/AvalancheShapes/Private/AvaShapeComponentMaterialBridge.cpp
+++ b/Engine/Plugins/VirtualProduction/Avalanche/Source/AvalancheShapes/Private/AvaShapeComponentMaterialBridge.cpp
@@ -11,13 +11,17 @@ const UStruct* FShapeComponentMaterialBridge::OnGetBridgedType() const
 	return UAvaShapeDynamicMeshBase::StaticClass();
 }
 
-void FShapeComponentMaterialBridge::OnForEachUsedMaterial(FConstDataView InDataView, TFunctionRef<void(UMaterialInterface*)> InFunc, const FGetUsedMaterialOptions& InOptions) const
+void FShapeComponentMaterialBridge::OnVisitMaterials(FConstDataView InDataView, TFunctionRef<void(const FVisitedContext&)> InFunc, const FVisitOptions& InOptions) const
 {
 	const UAvaShapeDynamicMeshBase& Component = InDataView.Get<const UAvaShapeDynamicMeshBase>();
 
+	FVisitedContext Context;
+	Context.Owner = InDataView;
+
 	for (const TPair<int32, FAvaShapeMeshData>& Pair : Component.GetMeshDatas())
 	{
-		InFunc(Pair.Value.GetMaterial());
+		Context.Material = Pair.Value.GetMaterial();
+		InFunc(Context);
 	}
 }
 
diff --git a/Engine/Plugins/VirtualProduction/Avalanche/Source/AvalancheShapes/Public/AvaShapeComponentMaterialBridge.h b/Engine/Plugins/VirtualProduction/Avalanche/Source/AvalancheShapes/Public/AvaShapeComponentMaterialBridge.h
--- a/Engine/Plugins/VirtualProduction/Avalanche/Source/AvalancheShapes/Public/AvaShapeComponentMaterialBridge.h
+++ b/Engine/Plugins/VirtualProduction/Avalanche/Source/AvalancheShapes/Public/AvaShapeComponentMaterialBridge.h
@@ -15,7 +15,7 @@ class FShapeComponentMaterialBridge : public FMaterialBridge
 protected:
 	//~ Begin FMaterialBridge
 	UE_API virtual const UStruct* OnGetBridgedType() const override;
-	UE_API virtual void OnForEachUsedMaterial(FConstDataView InDataView, TFunctionRef<void(UMaterialInterface*)> InFunc, const FGetUsedMaterialOptions& InOptions) const override;
+	UE_API virtual void OnVisitMaterials(FConstDataView InDataView, TFunctionRef<void(const FVisitedContext&)> InFunc, const FVisitOptions& InOptions) const override;
 	//~ End FMaterialBridge
 };
 
diff --git a/Engine/Plugins/VirtualProduction/Avalanche/Source/AvalancheText/Private/AvaText3DComponentMaterialBridge.cpp b/Engine/Plugins/VirtualProduction/Avalanche/Source/AvalancheText/Private/AvaText3DComponentMaterialBridge.cpp
--- a/Engine/Plugins/VirtualProduction/Avalanche/Source/AvalancheText/Private/AvaText3DComponentMaterialBridge.cpp
+++ b/Engine/Plugins/VirtualProduction/Avalanche/Source/AvalancheText/Private/AvaText3DComponentMaterialBridge.cpp
@@ -13,16 +13,21 @@ const UStruct* FText3DComponentMaterialBridge::OnGetBridgedType() const
 	return UText3DComponent::StaticClass();
 }
 
-void FText3DComponentMaterialBridge::OnForEachUsedMaterial(FConstDataView InDataView, TFunctionRef<void(UMaterialInterface*)> InFunc, const FGetUsedMaterialOptions& InOptions) const
+void FText3DComponentMaterialBridge::OnVisitMaterials(FConstDataView InDataView, TFunctionRef<void(const FVisitedContext&)> InFunc, const FVisitOptions& InOptions) const
 {
 	const UText3DComponent& Component = InDataView.Get<const UText3DComponent>();
 
 	if (const UText3DMaterialExtensionBase* MaterialExtension = Component.GetMaterialExtension())
 	{
+		FVisitedContext VisitContext;
+		VisitContext.Owner = FConstDataView(MaterialExtension);
+
 		MaterialExtension->ForEachMaterial(
-			[&InFunc](const UE::Text3D::Material::FMaterialParameters&, UMaterialInterface* InMaterial)->bool
+			[&InFunc, &VisitContext](const UE::Text3D::Material::FMaterialParameters&, UMaterialInterface* InMaterial)->bool
 			{
-				InFunc(InMaterial);
+				VisitContext.Material = InMaterial;
+
+				InFunc(VisitContext);
 				return true; // continue				
 			});
 	}
diff --git a/Engine/Plugins/VirtualProduction/Avalanche/Source/AvalancheText/Public/AvaText3DComponentMaterialBridge.h b/Engine/Plugins/VirtualProduction/Avalanche/Source/AvalancheText/Public/AvaText3DComponentMaterialBridge.h
--- a/Engine/Plugins/VirtualProduction/Avalanche/Source/AvalancheText/Public/AvaText3DComponentMaterialBridge.h
+++ b/Engine/Plugins/VirtualProduction/Avalanche/Source/AvalancheText/Public/AvaText3DComponentMaterialBridge.h
@@ -15,7 +15,7 @@ class FText3DComponentMaterialBridge : public FMaterialBridge
 protected:
 	//~ Begin FMaterialBridge
 	UE_API virtual const UStruct* OnGetBridgedType() const override;
-	UE_API virtual void OnForEachUsedMaterial(FConstDataView InDataView, TFunctionRef<void(UMaterialInterface*)> InFunc, const FGetUsedMaterialOptions& InOptions) const override;
+	UE_API virtual void OnVisitMaterials(FConstDataView InDataView, TFunctionRef<void(const FVisitedContext&)> InFunc, const FVisitOptions& InOptions) const override;
 	//~ End FMaterialBridge
 };
 
