Motion Design: moved material caching logic to the materials module with a 'Material Cache Helper' that skips re-gathering materials for the same objects.

diff --git a/Engine/Plugins/VirtualProduction/Avalanche/Source/AvalancheMaterial/AvalancheMaterial.Build.cs b/Engine/Plugins/VirtualProduction/Avalanche/Source/AvalancheMaterial/AvalancheMaterial.Build.cs
--- a/Engine/Plugins/VirtualProduction/Avalanche/Source/AvalancheMaterial/AvalancheMaterial.Build.cs
+++ b/Engine/Plugins/VirtualProduction/Avalanche/Source/AvalancheMaterial/AvalancheMaterial.Build.cs
@@ -17,5 +17,12 @@ public class AvalancheMaterial : ModuleRules
                 "Engine",
             }
         );
+
+        PrivateDependencyModuleNames.AddRange(
+            new string[]
+            {
+                "RHI",
+            }
+        );
     }
 }
diff --git a/Engine/Plugins/VirtualProduction/Avalanche/Source/AvalancheMaterial/Internal/AvaMaterialCacheHelper.h b/Engine/Plugins/VirtualProduction/Avalanche/Source/AvalancheMaterial/Internal/AvaMaterialCacheHelper.h
new file mode 100644
--- /dev/null
+++ b/Engine/Plugins/VirtualProduction/Avalanche/Source/AvalancheMaterial/Internal/AvaMaterialCacheHelper.h
@@ -0,0 +1,107 @@
+// Copyright Epic Games, Inc. All Rights Reserved.
+
+#pragma once
+
+#include "Templates/SharedPointer.h"
+#include "UObject/GCObject.h"
+#include "UObject/ObjectKey.h"
+
+#define UE_API AVALANCHEMATERIAL_API
+
+class UMaterialInterface;
+
+namespace UE::Ava
+{
+
+/** Helper to cache materials over multiple frames */
+class FMaterialCacheHelper : public TSharedFromThis<FMaterialCacheHelper>, public FGCObject
+{
+	enum class EMaterialCacheStatus : uint8
+	{
+		/** Material is waiting to call CacheShaders */
+		Waiting,
+		/** Material called CacheShaders */
+		Caching,
+	};
+
+	struct FObjectData
+	{
+		/** All the materials belonging to an object that are still incomplete */
+		TSet<TObjectKey<UMaterialInterface>> IncompleteMaterials;
+	};
+
+	struct FMaterialData
+	{
+		/** Material that is pending completion */
+		TObjectPtr<UMaterialInterface> Material;
+		/** Gives the status of cache */
+		EMaterialCacheStatus CacheStatus = EMaterialCacheStatus::Waiting;
+	};
+
+public:
+	UE_API static FMaterialCacheHelper& Get();
+
+	/**
+	 * Caches the materials in the given object.
+	 * The materials are found via the material bridge api
+	 */
+	UE_API void RequestCacheMaterials(const UObject* InObject);
+
+	/** Processes materials pending completion with a limited time budget */
+	UE_API void Tick();
+
+	/** Returns true if there are still any materials being cached */
+	UE_API bool IsCaching() const;
+
+	/** Returns true if there are materials from the given object that are still being cached */
+	UE_API bool IsCaching(const UObject* InObject) const;
+
+	/** Dumps all the materials that are still being processed (i.e. still incomplete) to the output log */
+	UE_API void DumpMaterials();
+
+	//~ Begin FGCObject
+	virtual void AddReferencedObjects(FReferenceCollector& InCollector) override;
+	virtual FString GetReferencerName() const override;
+	//~ End FGCObject
+
+private:
+	/** Resets the object and material map */
+	void Reset();
+
+	/** Registers the tick function to get called every frame */
+	void RegisterTick();
+
+	/** Unregisters the tick function from being called every frame */
+	void UnregisterTick();
+
+	/**
+	 * Caches as many queued materials as the time budget allows, and updates the status of each material to 'caching'
+	 * It also removes any material that is fully complete.
+	 */
+	void ProcessMaterials();
+
+	/** Removes the complete materials from the objects' IncompleteMaterials set */
+	void RemoveCompletedMaterials();
+
+	/** Process async results if shader is compiling */
+	void ProcessShadersCompiling();
+
+	/** Gathers the material data for an object, or skips it if the object's materials have been gathered before. */
+	void GatherMaterialData(const UObject* InObject);
+
+	/** Objects that have already been gathered, mapped to the incomplete materials that were gathered for the object */
+	TMap<FObjectKey, FObjectData> Objects;
+
+	/** All the materials that have already been gathered and are pending completion. Once materials are complete, they're removed from this map */
+	TMap<TObjectKey<UMaterialInterface>, FMaterialData> Materials;
+
+	/** Current time budget left for the tick. Resets every tick */
+	float RemainingTimeBudget = 0.f;
+
+	/** Handle to the delegate for ticking */
+	FDelegateHandle TickHandle;
+};
+
+} // UE::Ava
+
+#undef UE_API
diff --git a/Engine/Plugins/VirtualProduction/Avalanche/Source/AvalancheMaterial/Private/AvaMaterialCacheHelper.cpp b/Engine/Plugins/VirtualProduction/Avalanche/Source/AvalancheMaterial/Private/AvaMaterialCacheHelper.cpp
new file mode 100644
--- /dev/null
+++ b/Engine/Plugins/VirtualProduction/Avalanche/Source/AvalancheMaterial/Private/AvaMaterialCacheHelper.cpp
@@ -0,0 +1,307 @@
+// Copyright Epic Games, Inc. All Rights Reserved.
+
+#include "AvaMaterialCacheHelper.h"
+#include "AvaDataView.h"
+#include "AvaMaterialBridge.h"
+#include "AvaMaterialBridgeRegistry.h"
+#include "MaterialShared.h"
+#include "Materials/Material.h"
+#include "ShaderCompiler.h"
+
+DEFINE_LOG_CATEGORY_STATIC(LogAvaMaterialCache, Log, All);
+
+namespace UE::Ava
+{
+
+namespace Private
+{
+
+FAutoConsoleCommand MaterialCacheDumpMaterials(TEXT("MotionDesign.MaterialCache.DumpMaterials")
+	, TEXT("Dumps all the materials that are currently being processed in the Material Cache Helper")
+	, FConsoleCommandDelegate::CreateLambda([]
+		{
+			FMaterialCacheHelper::Get().DumpMaterials();
+		}));
+
+void CacheMaterial(FMaterialUpdateContext& InContext, TNotNull<UMaterialInterface*> InMaterial)
+{
+	constexpr EMaterialShaderPrecompileMode PrecompileMode = EMaterialShaderPrecompileMode::Background;
+
+	// This is needed because CacheShaders blindly recreates uniform buffers
+	// which can only be done if the draw command is going to be re-cached.
+	InContext.AddMaterialInterface(InMaterial);
+	InMaterial->CacheShaders(PrecompileMode);
+}
+
+} // UE::Ava::Private
+
+FMaterialCacheHelper& FMaterialCacheHelper::Get()
+{
+	static TSharedRef<FMaterialCacheHelper> MaterialCacheHelper = MakeShared<FMaterialCacheHelper>();
+	return *MaterialCacheHelper;
+}
+
+void FMaterialCacheHelper::RequestCacheMaterials(const UObject* InObject)
+{
+	if (!InObject)
+	{
+		return;
+	}
+
+	GatherMaterialData(InObject);
+	RegisterTick();
+}
+
+void FMaterialCacheHelper::Tick()
+{
+	TRACE_CPUPROFILER_EVENT_SCOPE(FMaterialCacheHelper::Tick);
+
+	constexpr float MaxTimeBudget = 0.001f;
+	RemainingTimeBudget = MaxTimeBudget;
+
+	ProcessMaterials();
+
+	if (IsCaching())
+	{
+		ProcessShadersCompiling();
+		RemoveCompletedMaterials();
+	}
+	else
+	{
+		Reset();
+	}
+}
+
+bool FMaterialCacheHelper::IsCaching() const
+{
+	return !Materials.IsEmpty();
+}
+
+bool FMaterialCacheHelper::IsCaching(const UObject* InObject) const
+{
+	if (!IsCaching())
+	{
+		return false;
+	}
+
+	const FObjectData* ObjectData = Objects.Find(InObject);
+	if (!ObjectData)
+	{
+		return false;
+	}
+
+	for (TObjectKey<UMaterialInterface> Material : ObjectData->IncompleteMaterials)
+	{
+		if (Materials.Contains(Material))
+		{
+			return true;
+		}
+	}
+	return false;
+}
+
+void FMaterialCacheHelper::DumpMaterials()
+{
+	constexpr const TCHAR* CacheStatusDisplays[] =
+		{
+			TEXT("Waiting"),
+			TEXT("Caching"),
+		};
+
+	UE_LOG(LogAvaMaterialCache, Log, TEXT("----- DUMP MATERIALS--------------------------------------------"));
+	for (const TPair<TObjectKey<UMaterialInterface>, FMaterialData>& Pair : Materials)
+	{
+		const FMaterialData& MaterialData = Pair.Value;
+
+		UE_LOG(LogAvaMaterialCache, Log, TEXT("\t'%s' --> [%s]")
+			, MaterialData.Material ? *MaterialData.Material->GetFullName() : TEXT("(invalid)")
+			, CacheStatusDisplays[static_cast<uint8>(MaterialData.CacheStatus)]);
+	}
+	UE_LOG(LogAvaMaterialCache, Log, TEXT("----------------------------------------------------------------"));
+}
+
+void FMaterialCacheHelper::AddReferencedObjects(FReferenceCollector& InCollector)
+{
+	for (TPair<TObjectKey<UMaterialInterface>, FMaterialData>& Pair : Materials)
+	{
+		InCollector.AddReferencedObject(Pair.Value.Material);
+	}
+}
+
+FString FMaterialCacheHelper::GetReferencerName() const
+{
+	return TEXT("UE::Ava::FMaterialCacheHelper");
+}
+
+void FMaterialCacheHelper::Reset()
+{
+	Materials.Empty();
+	Objects.Empty();
+	UnregisterTick();
+}
+
+void FMaterialCacheHelper::RegisterTick()
+{
+	if (!TickHandle.IsValid())
+	{
+		TickHandle = FCoreDelegates::OnEndFrame.AddSP(this, &FMaterialCacheHelper::Tick);
+	}
+}
+
+void FMaterialCacheHelper::UnregisterTick()
+{
+	if (TickHandle.IsValid())
+	{
+		FCoreDelegates::OnEndFrame.Remove(TickHandle);
+		TickHandle.Reset();
+	}
+}
+
+void FMaterialCacheHelper::ProcessMaterials()
+{
+	if (RemainingTimeBudget <= 0.f || Materials.IsEmpty())
+	{
+		return;
+	}
+
+	TRACE_CPUPROFILER_EVENT_SCOPE(FMaterialCacheHelper::ProcessMaterials);
+
+	// Time when budget runs out
+	const double EndBudgetTime = FPlatformTime::Seconds() + RemainingTimeBudget;
+
+	TOptional<FMaterialUpdateContext> UpdateContext;
+
+	for (TMap<TObjectKey<UMaterialInterface>, FMaterialData>::TIterator Iter(Materials); Iter; ++Iter)
+	{
+		FMaterialData& MaterialData = Iter.Value();
+
+		if (!MaterialData.Material)
+		{
+			Iter.RemoveCurrent();
+			continue;
+		}
+
+		switch (MaterialData.CacheStatus)
+		{
+		case EMaterialCacheStatus::Waiting:
+			if (!UpdateContext.IsSet())
+			{
+				UpdateContext.Emplace(0);
+			}
+			Private::CacheMaterial(*UpdateContext, MaterialData.Material);
+			MaterialData.CacheStatus = EMaterialCacheStatus::Caching;
+			break;
+
+		case EMaterialCacheStatus::Caching:
+			if (MaterialData.Material->IsComplete())
+			{
+				Iter.RemoveCurrent();
+			}
+			break;
+		}
+
+		// Check if time is past end budget time
+		if (FPlatformTime::Seconds() >= EndBudgetTime)
+		{
+			break;
+		}
+	}
+
+	RemainingTimeBudget = EndBudgetTime - FPlatformTime::Seconds();
+}
+
+void FMaterialCacheHelper::RemoveCompletedMaterials()
+{
+	if (RemainingTimeBudget <= 0.f)
+	{
+		return;
+	}
+
+	const double StartTime = FPlatformTime::Seconds();
+
+	for (TPair<FObjectKey, FObjectData>& Object : Objects)
+	{
+		for (TSet<TObjectKey<UMaterialInterface>>::TIterator MaterialIter(Object.Value.IncompleteMaterials); MaterialIter; ++MaterialIter)
+		{
+			if (!Materials.Contains(*MaterialIter))
+			{
+				MaterialIter.RemoveCurrent();
+			}
+		}
+	}
+
+	RemainingTimeBudget -= (FPlatformTime::Seconds() - StartTime);
+}
+
+void FMaterialCacheHelper::ProcessShadersCompiling()
+{
+	if (RemainingTimeBudget <= 0.f || !GShaderCompilingManager || !GShaderCompilingManager->IsCompiling())
+	{
+		return;
+	}
+
+	TRACE_CPUPROFILER_EVENT_SCOPE(FMaterialCacheHelper::ProcessShadersCompiling);
+
+	const double StartTime = FPlatformTime::Seconds();
+
+	constexpr bool bBlockOnGlobalShaderCompletion = false;
+	GShaderCompilingManager->ProcessAsyncResults(RemainingTimeBudget, bBlockOnGlobalShaderCompletion);
+
+	RemainingTimeBudget -= (FPlatformTime::Seconds() - StartTime);
+}
+
+void FMaterialCacheHelper::GatherMaterialData(const UObject* InObject)
+{
+	if (!InObject || Objects.Contains(InObject))
+	{
+		return;
+	}
+
+	TRACE_CPUPROFILER_EVENT_SCOPE(FMaterialCacheHelper::GatherMaterialData);
+
+	const FMaterialBridge* MaterialBridge = FMaterialBridgeRegistry::Get().GetMaterialBridgeForObject(InObject);
+	if (!MaterialBridge)
+	{
+		UE_LOG(LogAvaMaterialCache, Warning, TEXT("Object '%s' does not have a valid bridge to cache materials!"), *GetNameSafe(InObject));
+		return;
+	}
+
+	FObjectData& ObjectData = Objects.Add(InObject);
+
+	auto AddIncompleteMaterial =
+		[&Materials = Materials, &ObjectData](TNotNull<UMaterialInterface*> InMaterial)
+		{
+			FMaterialData* MaterialData = Materials.Find(InMaterial);
+			if (!MaterialData)
+			{
+				MaterialData = &Materials.Add(InMaterial);
+				MaterialData->Material = InMaterial;
+				MaterialData->CacheStatus = EMaterialCacheStatus::Waiting;
+			}
+			ObjectData.IncompleteMaterials.Add(InMaterial);
+		};
+
+	// Gather all materials for object
+	MaterialBridge->VisitMaterials(FConstDataView(InObject),
+		[&AddIncompleteMaterial](const FMaterialBridge::FVisitedContext& InContext)
+		{
+			if (!InContext.Material)
+			{
+				return;
+			}
+
+			UMaterial* const Material = InContext.Material->GetMaterial();
+
+			// Add interface's material if it is different from the interface
+			if (Material && Material != InContext.Material && !Material->IsComplete())
+			{
+				AddIncompleteMaterial(Material);
+			}
+			if (!InContext.Material->IsComplete())
+			{
+				AddIncompleteMaterial(InContext.Material);
+			}
+		});
+}
+
+} // UE::Ava
\ No newline at end of file
diff --git a/Engine/Plugins/VirtualProduction/Avalanche/Source/AvalancheMedia/Internal/Playback/AvaPlaybackUtils.h b/Engine/Plugins/VirtualProduction/Avalanche/Source/AvalancheMedia/Internal/Playback/AvaPlaybackUtils.h
--- a/Engine/Plugins/VirtualProduction/Avalanche/Source/AvalancheMedia/Internal/Playback/AvaPlaybackUtils.h
+++ b/Engine/Plugins/VirtualProduction/Avalanche/Source/AvalancheMedia/Internal/Playback/AvaPlaybackUtils.h
@@ -63,13 +63,6 @@ namespace UE::AvaPlayback::Utils
 	 */
 	AVALANCHEMEDIA_API FString GetBriefFrameInfo();
 
-	/**
-	 * Processes shader compilation results if any.
-	 * @param InTimeBudget the time budget limit, in seconds, to process async results
-	 * @return true if shaders are compiling and async results (if any) began processing, false otherwise.
-	 */
-	bool ProcessShadersCompiling(TOptional<float> InTimeBudget);
-
 	template<typename InEnumType>
 	FString StaticEnumToString(InEnumType InValue)
 	{
@@ -105,6 +98,5 @@ namespace UE::AvaPlayback::Utils
 	private:
 		TSet<FSoftObjectPath> PendingAssets;
 		mutable TArray<TObjectPtr<UStreamableRenderAsset>> PendingStreamingAssets;
-		bool bWaitForStreamingAssets = false;
 	};
 }
diff --git a/Engine/Plugins/VirtualProduction/Avalanche/Source/AvalancheMedia/Private/Playable/Playables/AvaPlayableLevelStreaming.cpp b/Engine/Plugins/VirtualProduction/Avalanche/Source/AvalancheMedia/Private/Playable/Playables/AvaPlayableLevelStreaming.cpp
--- a/Engine/Plugins/VirtualProduction/Avalanche/Source/AvalancheMedia/Private/Playable/Playables/AvaPlayableLevelStreaming.cpp
+++ b/Engine/Plugins/VirtualProduction/Avalanche/Source/AvalancheMedia/Private/Playable/Playables/AvaPlayableLevelStreaming.cpp
@@ -1,8 +1,7 @@
 // Copyright Epic Games, Inc. All Rights Reserved.
 
 #include "Playable/Playables/AvaPlayableLevelStreaming.h"
-#include "AvaMaterialBridge.h"
-#include "AvaMaterialBridgeRegistry.h"
+#include "AvaMaterialCacheHelper.h"
 #include "AvaRemoteControlRebind.h"
 #include "AvaRemoteControlUtils.h"
 #include "AvaScene.h"
@@ -15,16 +14,12 @@
 #include "Framework/AvaSoftAssetPtr.h"
 #include "IAvaMediaModule.h"
 #include "IAvaRemoteControlInterface.h"
-#include "MaterialShared.h"
-#include "Materials/Material.h"
-#include "Materials/MaterialInterface.h"
 #include "Playable/AvaPlayableAssetUserData.h"
 #include "Playable/AvaPlayableGroup.h"
 #include "Playable/AvaPlayableGroupManager.h"
 #include "Playable/AvaPlayableUtils.h"
 #include "Playback/AvaPlaybackUtils.h"
 #include "SceneView.h"
-#include "ShaderCompiler.h"
 #include "UnrealEngine.h"
 #include "Streaming/LevelStreamingDelegates.h"
 #include "UObject/Package.h"
@@ -288,31 +283,6 @@ namespace UE::AvaMedia::LevelStreamingPlayable::Private
 		
 		return GetPlayableStatusRelevanceShouldBeUnloaded(InStatus);
 	}
-
-	/** Calls cache shaders for each material in the set, until the time budget is exceeded */
-	void CacheMaterials(const TSet<UMaterialInterface*>& InMaterialsToCache, TOptional<float> InTimeBudget)
-	{
-		TRACE_CPUPROFILER_EVENT_SCOPE(MotionDesignPlayable_CacheMaterials)
-
-		const double StartTime = FPlatformTime::Seconds();
-		const double TimeBudget = InTimeBudget.Get(0.f);
-
-		FMaterialUpdateContext UpdateContext(0);
-		for (UMaterialInterface* MaterialInterface : InMaterialsToCache)
-		{
-			constexpr EMaterialShaderPrecompileMode PrecompileMode = EMaterialShaderPrecompileMode::Background;
-
-			// This is needed because CacheShaders blindly recreates uniform buffers
-			// which can only be done if the draw command is going to be re-cached.
-			UpdateContext.AddMaterialInterface(MaterialInterface);
-			MaterialInterface->CacheShaders(PrecompileMode);
-
-			if (InTimeBudget.IsSet() && FPlatformTime::Seconds() - StartTime > TimeBudget)
-			{
-				return;
-			}
-		}
-	}
 }
 
 bool UAvaPlayableLevelStreaming::LoadAsset(const FAvaSoftAssetPtr& InSourceAsset, bool bInInitiallyVisible, const FString& InLoadOptions)
@@ -451,8 +421,7 @@ bool UAvaPlayableLevelStreaming::IsReadyToPlay() const
 	}
 
 	// Check whether the compilation has finished
-	constexpr float TimeBudget = 0.001f;
-	const bool bIsCachingShaders = CacheMaterials(Level, TimeBudget);
+	const bool bIsCachingShaders = CacheMaterials(Level);
 	return !bIsCachingShaders;
 }
 
@@ -554,48 +523,27 @@ bool UAvaPlayableLevelStreaming::LoadLevel(const TSoftObjectPtr<UWorld>& InSourc
 	return true;
 }
 
-bool UAvaPlayableLevelStreaming::CacheMaterials(TNotNull<const ULevel*> InLevel, TOptional<float> InTimeBudget) const
+bool UAvaPlayableLevelStreaming::CacheMaterials(TNotNull<const ULevel*> InLevel) const
 {
 	TRACE_CPUPROFILER_EVENT_SCOPE(UAvaPlayableLevelStreaming::CacheMaterials)
 
-	using namespace UE::AvaMedia::LevelStreamingPlayable;
-
-	TSet<UMaterialInterface*> MaterialsToCache;
+	UE::Ava::FMaterialCacheHelper& MaterialCacheHelper = UE::Ava::FMaterialCacheHelper::Get();
 
-	const UE::Ava::FMaterialBridgeRegistry& MaterialBridgeRegistry = UE::Ava::FMaterialBridgeRegistry::Get();
 	for (AActor* Actor : InLevel->Actors)
 	{
-		if (const UE::Ava::FMaterialBridge* ActorMaterialBridge = MaterialBridgeRegistry.GetMaterialBridgeForObject(Actor))
-		{
-			ActorMaterialBridge->VisitMaterials(UE::Ava::FConstDataView(Actor),
-				[&MaterialsToCache](const UE::Ava::FMaterialBridge::FVisitedContext& InContext)
-				{
-					if (!InContext.Material)
-					{
-						return;
-					}
-
-					UMaterial* const Material = InContext.Material->GetMaterial();
-					if (Material && Material != InContext.Material && !Material->IsComplete())
-					{
-						MaterialsToCache.Add(Material);
-					}
-					if (!InContext.Material->IsComplete())
-					{
-						MaterialsToCache.Add(InContext.Material);
-					}
-				});
-		}
+		MaterialCacheHelper.RequestCacheMaterials(Actor);
 	}
 
-	if (MaterialsToCache.IsEmpty())
+	bool bIsCaching = false;
+	for (AActor* Actor : InLevel->Actors)
 	{
-		const bool bAreShadersCompiling = UE::AvaPlayback::Utils::ProcessShadersCompiling(InTimeBudget);
-		return bAreShadersCompiling;
+		if (MaterialCacheHelper.IsCaching(Actor))
+		{
+			bIsCaching = true;
+			break;
+		}
 	}
-
-	Private::CacheMaterials(MaterialsToCache, InTimeBudget);
-	return true;
+	return bIsCaching;
 }
 
 void UAvaPlayableLevelStreaming::HandleTransitionEvent(UAvaPlayable* InPlayable, UAvaPlayableTransition* InTransition, EAvaPlayableTransitionEventFlags InTransitionFlags)
@@ -709,12 +657,10 @@ void UAvaPlayableLevelStreaming::OnLevelStreamingStateChanged_Synchronized(ELeve
 
 		if (bCacheShadersOnLoad)
 		{
-			static constexpr float TimeBudget = 0.001f;
-
 			// If there was more than one material that required caching,
-			// kick off a ticker that incrementally caches materials with same time budget,
+			// kick off a ticker that incrementally caches materials,
 			// until there's no materials to cache and all shaders have completed compilation
-			if (CacheMaterials(Level, TimeBudget))
+			if (CacheMaterials(Level))
 			{
 				bCachingShadersOnLoad = true;
 
@@ -724,7 +670,7 @@ void UAvaPlayableLevelStreaming::OnLevelStreamingStateChanged_Synchronized(ELeve
 						const ULevel* const Level = LevelStreaming ? LevelStreaming->GetLoadedLevel() : nullptr;
 						if (Level)
 						{
-							bCachingShadersOnLoad = CacheMaterials(Level, TimeBudget);
+							bCachingShadersOnLoad = CacheMaterials(Level);
 						}
 						else
 						{
diff --git a/Engine/Plugins/VirtualProduction/Avalanche/Source/AvalancheMedia/Private/Playable/Playables/AvaPlayableLevelStreaming.h b/Engine/Plugins/VirtualProduction/Avalanche/Source/AvalancheMedia/Private/Playable/Playables/AvaPlayableLevelStreaming.h
--- a/Engine/Plugins/VirtualProduction/Avalanche/Source/AvalancheMedia/Private/Playable/Playables/AvaPlayableLevelStreaming.h
+++ b/Engine/Plugins/VirtualProduction/Avalanche/Source/AvalancheMedia/Private/Playable/Playables/AvaPlayableLevelStreaming.h
@@ -52,10 +52,9 @@ protected:
 	/**
 	 * Caches the shaders for the given level.
 	 * @param InLevel the level to cache shaders for
-	 * @param InTimeBudget the time budget, in seconds, given to cache materials
 	 * @return true if caching is taking place. false otherwise. 
 	 */
-	bool CacheMaterials(TNotNull<const ULevel*> InLevel, TOptional<float> InTimeBudget) const;
+	bool CacheMaterials(TNotNull<const ULevel*> InLevel) const;
 
 	void HandleTransitionEvent(UAvaPlayable* InPlayable, UAvaPlayableTransition* InTransition, EAvaPlayableTransitionEventFlags InTransitionFlags);
 
diff --git a/Engine/Plugins/VirtualProduction/Avalanche/Source/AvalancheMedia/Private/Playback/AvaPlaybackUtils.cpp b/Engine/Plugins/VirtualProduction/Avalanche/Source/AvalancheMedia/Private/Playback/AvaPlaybackUtils.cpp
--- a/Engine/Plugins/VirtualProduction/Avalanche/Source/AvalancheMedia/Private/Playback/AvaPlaybackUtils.cpp
+++ b/Engine/Plugins/VirtualProduction/Avalanche/Source/AvalancheMedia/Private/Playback/AvaPlaybackUtils.cpp
@@ -5,12 +5,12 @@
 #include "AssetRegistry/AssetData.h"
 #include "AssetRegistry/IAssetRegistry.h"
 #include "AvaAssetTags.h"
+#include "AvaMaterialCacheHelper.h"
 #include "ContentStreaming.h"
 #include "Engine/StreamableRenderAsset.h"
 #include "Framework/AvaSoftAssetPtr.h"
 #include "HAL/PlatformFileManager.h"
 #include "IAvaMediaModule.h"
-#include "Materials/MaterialInterface.h"
 #include "Misc/Paths.h"
 #include "RenderingThread.h"
 #include "ShaderCompiler.h"
@@ -240,25 +240,11 @@ namespace UE::AvaPlayback::Utils
 		return FString::Printf(TEXT("[%d]"), GFrameNumber);
 	}
 
-	bool ProcessShadersCompiling(TOptional<float> InTimeBudget)
-	{
-		if (GShaderCompilingManager)
-		{
-			if (GShaderCompilingManager->IsCompiling())
-			{
-				constexpr bool bBlockOnGlobalShaderCompletion = false;
-				GShaderCompilingManager->ProcessAsyncResults(InTimeBudget.Get(0.f), bBlockOnGlobalShaderCompletion);
-				return true;
-			}
-		}
-		return false;
-	}
-
 	void FAsyncAssetLoader::HandleStreamableObject(UObject* InObject)
 	{
-		if (!Private::CVarRuntimePrestreamRCReferencedAssets.GetValueOnGameThread())
+		if (!InObject || !Private::CVarRuntimePrestreamRCReferencedAssets.GetValueOnGameThread())
 		{
-			// Not allowed to stream assets manually.
+			// Not allowed to stream assets.
 			return;
 		}
 
@@ -273,7 +259,6 @@ namespace UE::AvaPlayback::Utils
 			StreamableAsset->StreamIn(FStreamableRenderResourceState::MAX_LOD_COUNT, bHighPriority);
 
 			PendingStreamingAssets.Add(StreamableAsset);
-			bWaitForStreamingAssets = true;
 		}
 	}
 
@@ -323,18 +308,20 @@ namespace UE::AvaPlayback::Utils
 			return false;
 		}
 
-		constexpr float TimeBudget = 0.001f;
-		const bool bAreShadersCompiling = bWaitForStreamingAssets && ProcessShadersCompiling(TimeBudget);
-		return !bAreShadersCompiling;
+		return true;
 	}
 
 	void FAsyncAssetLoader::CleanupStreamingAssets() const
 	{
+		UE::Ava::FMaterialCacheHelper& MaterialCacheHelper = UE::Ava::FMaterialCacheHelper::Get();
+
 		PendingStreamingAssets.RemoveAll(
-			[](UStreamableRenderAsset* InStreamableRenderAsset)
+			[&MaterialCacheHelper](UStreamableRenderAsset* InObject)
 			{
 				constexpr bool bWaitForLODTransition = true;
-				return !InStreamableRenderAsset->HasPendingInitOrStreaming(bWaitForLODTransition) && InStreamableRenderAsset->IsFullyStreamedIn();
+				return !InObject->HasPendingInitOrStreaming(bWaitForLODTransition)
+					&& InObject->IsFullyStreamedIn()
+					&& !MaterialCacheHelper.IsCaching(InObject);
 			});
 	}
 }
diff --git a/Engine/Plugins/VirtualProduction/Avalanche/Source/AvalancheMedia/Private/Rundown/AvaRundown.cpp b/Engine/Plugins/VirtualProduction/Avalanche/Source/AvalancheMedia/Private/Rundown/AvaRundown.cpp
--- a/Engine/Plugins/VirtualProduction/Avalanche/Source/AvalancheMedia/Private/Rundown/AvaRundown.cpp
+++ b/Engine/Plugins/VirtualProduction/Avalanche/Source/AvalancheMedia/Private/Rundown/AvaRundown.cpp
@@ -2,6 +2,7 @@
 
 #include "Rundown/AvaRundown.h"
 
+#include "AvaMaterialCacheHelper.h"
 #include "AvaMediaSettings.h"
 #include "Broadcast/AvaBroadcast.h"
 #include "IAvaMediaModule.h"
@@ -1205,6 +1206,35 @@ TArray<UAvaRundown::FLoadedInstanceInfo> UAvaRundown::LoadPage(int32 InPageId,
 		PlaybackInstance->Recycle();
 		LoadedInstances.Add({PlaybackInstance->GetInstanceId(), AssetPath});
 	}
+
+	// Load RC Assets Async and request cache its materials
+	{
+		TSet<FSoftObjectPath> RemoteControlAssets;
+
+		const FAvaPlayableRemoteControlValues& RemoteControlValues = Page.GetRemoteControlValues();
+		FAvaPlayableRemoteControlValues::CollectReferencedAssetPaths(RemoteControlValues.ControllerValues, RemoteControlAssets);
+		FAvaPlayableRemoteControlValues::CollectReferencedAssetPaths(RemoteControlValues.EntityValues, RemoteControlAssets);
+
+		for (const FSoftObjectPath& Asset : RemoteControlAssets)
+		{
+			if (UObject* Object = Asset.ResolveObject())
+			{
+				UE::Ava::FMaterialCacheHelper::Get().RequestCacheMaterials(Object);
+			}
+			else
+			{
+				Asset.LoadAsync(FLoadSoftObjectPathAsyncDelegate::CreateLambda(
+					[](const FSoftObjectPath&, UObject* InObject)
+					{
+						if (InObject)
+						{
+							UE::Ava::FMaterialCacheHelper::Get().RequestCacheMaterials(InObject);
+						}
+					}));
+			}
+		}
+	}
+
 	return LoadedInstances;
 }
 
