Motion Design: Added Load Page as a command for instanced page lists

diff --git a/Engine/Plugins/VirtualProduction/Avalanche/Source/AvalancheMediaEditor/Private/AvaMediaEditorStyle.cpp b/Engine/Plugins/VirtualProduction/Avalanche/Source/AvalancheMediaEditor/Private/AvaMediaEditorStyle.cpp
--- a/Engine/Plugins/VirtualProduction/Avalanche/Source/AvalancheMediaEditor/Private/AvaMediaEditorStyle.cpp
+++ b/Engine/Plugins/VirtualProduction/Avalanche/Source/AvalancheMediaEditor/Private/AvaMediaEditorStyle.cpp
@@ -45,8 +45,9 @@ FAvaMediaEditorStyle::FAvaMediaEditorStyle()
 	Set("AvaRundown.RemovePage"      , new IMAGE_BRUSH_SVG("Icons/MediaIcons/RemovePage"   , Icon16x16));
 	Set("AvaRundown.RenumberPage"    , new IMAGE_BRUSH_SVG("Icons/MediaIcons/RenumberPage" , Icon16x16));
 	Set("AvaRundown.ReimportPage"    , new IMAGE_BRUSH_SVG("Icons/MediaIcons/ReimportPage" , Icon16x16));
-	Set("AvaRundown.EditPageSource"  , new IMAGE_BRUSH_SVG("Icons/MediaIcons/OpenAsset"	  , Icon16x16));
+	Set("AvaRundown.EditPageSource"  , new IMAGE_BRUSH_SVG("Icons/MediaIcons/OpenAsset"	   , Icon16x16));
 	Set("AvaRundown.Play"            , new IMAGE_BRUSH_SVG("Icons/MediaIcons/Play"         , Icon16x16));
+	Set("AvaRundown.LoadPage"        , new IMAGE_BRUSH_SVG("Icons/MediaIcons/LoadPage"     , Icon16x16));
 	Set("AvaRundown.UpdateValues"    , new IMAGE_BRUSH_SVG("Icons/MediaIcons/UpdateValues" , Icon16x16));
 	Set("AvaRundown.Stop"            , new IMAGE_BRUSH_SVG("Icons/MediaIcons/Stop"         , Icon16x16));
 	Set("AvaRundown.ForceStop"       , new IMAGE_BRUSH_SVG("Icons/MediaIcons/Stop"         , Icon16x16));
diff --git a/Engine/Plugins/VirtualProduction/Avalanche/Source/AvalancheMediaEditor/Private/Rundown/AvaRundownCommands.cpp b/Engine/Plugins/VirtualProduction/Avalanche/Source/AvalancheMediaEditor/Private/Rundown/AvaRundownCommands.cpp
--- a/Engine/Plugins/VirtualProduction/Avalanche/Source/AvalancheMediaEditor/Private/Rundown/AvaRundownCommands.cpp
+++ b/Engine/Plugins/VirtualProduction/Avalanche/Source/AvalancheMediaEditor/Private/Rundown/AvaRundownCommands.cpp
@@ -114,6 +114,12 @@ void FAvaRundownCommands::RegisterCommands()
 		, EUserInterfaceActionType::Button
 		, FInputChord(EModifierKey::Shift, EKeys::Enter));
 
+	UI_COMMAND(LoadPage
+		, "Load Page"
+		, "Loads the selected pages for play in their respective channels"
+		, EUserInterfaceActionType::Button
+		, FInputChord());
+
 	UI_COMMAND(UpdateValues
 		, "Update Values"
 		, "Update Values of the Selected Page displayed in the selected channel"
diff --git a/Engine/Plugins/VirtualProduction/Avalanche/Source/AvalancheMediaEditor/Private/Rundown/AvaRundownCommands.h b/Engine/Plugins/VirtualProduction/Avalanche/Source/AvalancheMediaEditor/Private/Rundown/AvaRundownCommands.h
--- a/Engine/Plugins/VirtualProduction/Avalanche/Source/AvalancheMediaEditor/Private/Rundown/AvaRundownCommands.h
+++ b/Engine/Plugins/VirtualProduction/Avalanche/Source/AvalancheMediaEditor/Private/Rundown/AvaRundownCommands.h
@@ -37,6 +37,7 @@ public:
 
 	/** Show Control / Rundown Commands */
 	TSharedPtr<FUICommandInfo> Play;
+	TSharedPtr<FUICommandInfo> LoadPage;
 	TSharedPtr<FUICommandInfo> UpdateValues;
 	TSharedPtr<FUICommandInfo> Stop;
 	TSharedPtr<FUICommandInfo> ForceStop;
diff --git a/Engine/Plugins/VirtualProduction/Avalanche/Source/AvalancheMediaEditor/Private/Rundown/Pages/AvaRundownPageContextMenu.cpp b/Engine/Plugins/VirtualProduction/Avalanche/Source/AvalancheMediaEditor/Private/Rundown/Pages/AvaRundownPageContextMenu.cpp
--- a/Engine/Plugins/VirtualProduction/Avalanche/Source/AvalancheMediaEditor/Private/Rundown/Pages/AvaRundownPageContextMenu.cpp
+++ b/Engine/Plugins/VirtualProduction/Avalanche/Source/AvalancheMediaEditor/Private/Rundown/Pages/AvaRundownPageContextMenu.cpp
@@ -157,6 +157,7 @@ void FAvaRundownPageContextMenu::PopulatePageContextMenu(UToolMenu& InMenu, UAva
 		FToolMenuSection& Section = InMenu.FindOrAddSection(TEXT("ShowControlOperations"), LOCTEXT("ShowControlOperationsHeader", "Show Control Actions"));
 
 		Section.AddMenuEntry(RundownCommands.Play);
+		Section.AddMenuEntry(RundownCommands.LoadPage);
 		Section.AddMenuEntry(RundownCommands.UpdateValues);
 		Section.AddMenuEntry(RundownCommands.Continue);
 		Section.AddMenuEntry(RundownCommands.Stop);
diff --git a/Engine/Plugins/VirtualProduction/Avalanche/Source/AvalancheMediaEditor/Private/Rundown/Pages/Slate/SAvaRundownInstancedPageList.cpp b/Engine/Plugins/VirtualProduction/Avalanche/Source/AvalancheMediaEditor/Private/Rundown/Pages/Slate/SAvaRundownInstancedPageList.cpp
--- a/Engine/Plugins/VirtualProduction/Avalanche/Source/AvalancheMediaEditor/Private/Rundown/Pages/Slate/SAvaRundownInstancedPageList.cpp
+++ b/Engine/Plugins/VirtualProduction/Avalanche/Source/AvalancheMediaEditor/Private/Rundown/Pages/Slate/SAvaRundownInstancedPageList.cpp
@@ -296,6 +296,10 @@ void SAvaRundownInstancedPageList::BindCommands()
 			FExecuteAction::CreateSP(this, &SAvaRundownInstancedPageList::PlaySelectedPage),
 			FCanExecuteAction::CreateSP(this, &SAvaRundownInstancedPageList::CanPlaySelectedPage));
 
+		CommandList->MapAction(RundownCommands.LoadPage,
+			FExecuteAction::CreateSP(this, &SAvaRundownInstancedPageList::LoadSelectedPages),
+			FCanExecuteAction::CreateSP(this, &SAvaRundownInstancedPageList::CanLoadSelectedPages));
+
 		CommandList->MapAction(RundownCommands.UpdateValues,
 			FExecuteAction::CreateSP(this, &SAvaRundownInstancedPageList::UpdateValuesOnSelectedPage),
 			FCanExecuteAction::CreateSP(this, &SAvaRundownInstancedPageList::CanUpdateValuesOnSelectedPage));
@@ -1128,6 +1132,30 @@ bool SAvaRundownInstancedPageList::CanPlaySelectedPage() const
 	return !PageIds.IsEmpty();
 }
 
+bool SAvaRundownInstancedPageList::CanLoadSelectedPages() const
+{
+	const TArray<int32> PageIds = GetPagesToLoad();
+	return !PageIds.IsEmpty();
+}
+
+void SAvaRundownInstancedPageList::LoadSelectedPages() const
+{
+	if (UAvaRundown* Rundown = GetRundown())
+	{
+		const TArray<int32> PageIds = GetPagesToLoad();
+
+		// Settings to load program, not preview
+		constexpr bool bIsPreview = false;
+		const FName PreviewChannel = NAME_None;
+
+		IAvaRundownPageLoadingManager& PageLoadingManager = Rundown->GetPageLoadingManager();
+		for (int32 PageId : PageIds)
+		{
+			PageLoadingManager.RequestLoadPage(PageId, bIsPreview, PreviewChannel);
+		}
+	}
+}
+
 bool SAvaRundownInstancedPageList::CanUpdateValuesOnSelectedPage() const
 {
 	const TArray<int32> PageIds = GetPagesToUpdate();
@@ -1476,6 +1504,25 @@ TArray<int32> SAvaRundownInstancedPageList::GetPagesToUpdate() const
 	return FilterPageSetForProgram(KeepPagesToUpdate, PageSet);
 }
 
+TArray<int32> SAvaRundownInstancedPageList::GetPagesToLoad() const
+{
+	// Only allow pages that can be played.
+	auto FilterPagesToLoad = [](const UAvaRundown* InRundown, const TArray<int32>& InPageIds)
+		{
+			TArray<int32> OutPageIds;
+			OutPageIds.Reserve(InPageIds.Num());
+			for (int32 PageId : InPageIds)
+			{
+				if (InRundown->CanPlayPage(PageId, false))
+				{
+					OutPageIds.Add(PageId);
+				}
+			}
+			return OutPageIds;
+		};
+	return FilterSelectedPages(FilterPagesToLoad);
+}
+
 int32 SAvaRundownInstancedPageList::GetPageIdToTakeNext() const
 {
 	return FAvaRundownPlaybackUtils::GetPageIdToPlayNext(GetRundown(), GetPageListReference(), /*bInPreview*/ false, NAME_None);
diff --git a/Engine/Plugins/VirtualProduction/Avalanche/Source/AvalancheMediaEditor/Private/Rundown/Pages/Slate/SAvaRundownInstancedPageList.h b/Engine/Plugins/VirtualProduction/Avalanche/Source/AvalancheMediaEditor/Private/Rundown/Pages/Slate/SAvaRundownInstancedPageList.h
--- a/Engine/Plugins/VirtualProduction/Avalanche/Source/AvalancheMediaEditor/Private/Rundown/Pages/Slate/SAvaRundownInstancedPageList.h
+++ b/Engine/Plugins/VirtualProduction/Avalanche/Source/AvalancheMediaEditor/Private/Rundown/Pages/Slate/SAvaRundownInstancedPageList.h
@@ -43,6 +43,10 @@ public:
 	void PlaySelectedPage() const;
 	bool CanPlaySelectedPage() const;
 
+	/** Requests load of the selected pages that are able to play in program */
+	bool CanLoadSelectedPages() const;
+	void LoadSelectedPages() const;
+
 	bool CanUpdateValuesOnSelectedPage() const;
 	void UpdateValuesOnSelectedPage();
 	
@@ -77,6 +81,7 @@ private:
 	TArray<int32> GetPagesToTakeOut(bool bInForce) const;
 	TArray<int32> GetPagesToContinue() const;
 	TArray<int32> GetPagesToUpdate() const;
+	TArray<int32> GetPagesToLoad() const;
 	int32 GetPageIdToTakeNext() const;
 
 protected:
