Motion Design: fixed regression where taking in a page that changes a mesh (or a texture, etc) via its RC values would not work if that mesh wasn't loaded ahead of time.  'HasLocalPlayable' returns false because there are no playables available. This was fixed by moving the util func that determines whether a local or remote playable should be created to a centralized place that the page transition object can also access and use instead.

#Synced-CL 48030442

diff --git a/Engine/Plugins/VirtualProduction/Avalanche/Source/AvalancheMedia/Private/Playable/AvaPlayable.cpp b/Engine/Plugins/VirtualProduction/Avalanche/Source/AvalancheMedia/Private/Playable/AvaPlayable.cpp
--- a/Engine/Plugins/VirtualProduction/Avalanche/Source/AvalancheMedia/Private/Playable/AvaPlayable.cpp
+++ b/Engine/Plugins/VirtualProduction/Avalanche/Source/AvalancheMedia/Private/Playable/AvaPlayable.cpp
@@ -2,6 +2,7 @@
 
 #include "Playable/AvaPlayable.h"
 
+#include "AvaPlayableUtils.h"
 #include "AvaSequence.h"
 #include "AvaSequencePlayer.h"
 #include "Broadcast/AvaBroadcast.h"
@@ -53,32 +54,6 @@ static TAutoConsoleVariable<bool> CVarExecuteNonOverlappingRCControllerBehavior(
 
 namespace UE::AvaPlayable::Private
 {
-	bool ShouldCreateLocalPlayable(const FName& InChannelName, const UAvaBroadcast& InBroadcast)
-	{
-		const FAvaBroadcastOutputChannel& Channel = InBroadcast.GetCurrentProfile().GetChannel(InChannelName);
-
-		// If there is no broadcast channel defined, like for preview (by default), then this is a local playable.
-		if (!Channel.IsValidChannel())
-		{
-			return true;
-		}
-
-		// For non-preview, the commands will be executed locally if the channel has at least one local outputs or no outputs.
-		// The "no outputs" condition is considered valid. Empty channels run locally.
-		if (Channel.HasAnyLocalMediaOutputs() || Channel.GetMediaOutputs().IsEmpty())
-		{
-			return true;
-		}
-		
-		return false;
-	}
-
-	bool HasRemoteOutputs(const FName& InChannelName, const UAvaBroadcast& InBroadcast)
-	{
-		const FAvaBroadcastOutputChannel& Channel = InBroadcast.GetCurrentProfile().GetChannel(InChannelName);
-		return Channel.IsValidChannel() && Channel.HasAnyRemoteMediaOutputs();
-	}
-
 	FString GetPrettyPlayableInfo(const UAvaPlayable* InPlayable)
 	{
 		if (InPlayable)
@@ -131,10 +106,10 @@ UAvaPlayable* UAvaPlayable::Create(UObject* InOuter, const FPlayableCreationInfo
 	// - The case of forked local and remote channels will lead to a local playable and a remote proxy playable.
 	//   It would require wrapping the playables in a composite (or facade?) proxy. TODO
 	
-	if (ShouldCreateLocalPlayable(InPlayableInfo.ChannelName, Broadcast))
+	if (UE::AvaMedia::PlayableUtils::HasLocalPlay(Broadcast, InPlayableInfo.ChannelName))
 	{
 		// For the moment, remote outputs will be ignored.
-		if (HasRemoteOutputs(InPlayableInfo.ChannelName, Broadcast))
+		if (UE::AvaMedia::PlayableUtils::HasRemotePlay(Broadcast, InPlayableInfo.ChannelName))
 		{
 			UE_LOG(LogAvaPlayable, Error, TEXT("Forked Channels with both local and remote outputs are not supported in this version. Only local instance will be created."));
 		}
diff --git a/Engine/Plugins/VirtualProduction/Avalanche/Source/AvalancheMedia/Private/Playable/AvaPlayableUtils.cpp b/Engine/Plugins/VirtualProduction/Avalanche/Source/AvalancheMedia/Private/Playable/AvaPlayableUtils.cpp
--- a/Engine/Plugins/VirtualProduction/Avalanche/Source/AvalancheMedia/Private/Playable/AvaPlayableUtils.cpp
+++ b/Engine/Plugins/VirtualProduction/Avalanche/Source/AvalancheMedia/Private/Playable/AvaPlayableUtils.cpp
@@ -2,6 +2,8 @@
 
 #include "AvaPlayableUtils.h"
 
+#include "Broadcast/AvaBroadcast.h"
+#include "Broadcast/Channel/AvaBroadcastOutputChannel.h"
 #include "Components/PrimitiveComponent.h"
 
 namespace UE::AvaMedia::PlayableUtils
@@ -28,4 +30,30 @@ namespace UE::AvaMedia::PlayableUtils
 			}
 		}
 	}
-}
\ No newline at end of file
+
+	bool HasLocalPlay(const UAvaBroadcast& InBroadcast, FName InChannelName)
+	{
+		const FAvaBroadcastOutputChannel& Channel = InBroadcast.GetCurrentProfile().GetChannel(InChannelName);
+
+		// If there is no broadcast channel defined, like for preview (by default), then this is for local play.
+		if (!Channel.IsValidChannel())
+		{
+			return true;
+		}
+
+		// For non-preview, the commands will be executed locally if the channel has at least one local outputs or no outputs.
+		// The "no outputs" condition is considered valid. Empty channels run locally.
+		if (Channel.HasAnyLocalMediaOutputs() || Channel.GetMediaOutputs().IsEmpty())
+		{
+			return true;
+		}
+		
+		return false;
+	}
+
+	bool HasRemotePlay(const UAvaBroadcast& InBroadcast, FName InChannelName)
+	{
+		const FAvaBroadcastOutputChannel& Channel = InBroadcast.GetCurrentProfile().GetChannel(InChannelName);
+		return Channel.IsValidChannel() && Channel.HasAnyRemoteMediaOutputs();
+	}
+}
diff --git a/Engine/Plugins/VirtualProduction/Avalanche/Source/AvalancheMedia/Private/Playable/AvaPlayableUtils.h b/Engine/Plugins/VirtualProduction/Avalanche/Source/AvalancheMedia/Private/Playable/AvaPlayableUtils.h
--- a/Engine/Plugins/VirtualProduction/Avalanche/Source/AvalancheMedia/Private/Playable/AvaPlayableUtils.h
+++ b/Engine/Plugins/VirtualProduction/Avalanche/Source/AvalancheMedia/Private/Playable/AvaPlayableUtils.h
@@ -6,8 +6,15 @@
 
 class AActor;
 class FPrimitiveComponentId;
+class UAvaBroadcast;
 
 namespace UE::AvaMedia::PlayableUtils
 {
 	void AddPrimitiveComponentIds(const AActor* InActor, TSet<FPrimitiveComponentId>& InComponentIds);
+
+	/** Returns true if the broadcast's current profile for the given channel is configured to have local play */
+	bool HasLocalPlay(const UAvaBroadcast& InBroadcast, FName InChannelName);
+
+	/** Returns true if the broadcast's current profile for the given channel is configured to have remote play */
+	bool HasRemotePlay(const UAvaBroadcast& InBroadcast, FName InChannelName);
 }
\ No newline at end of file
diff --git a/Engine/Plugins/VirtualProduction/Avalanche/Source/AvalancheMedia/Private/Rundown/Transition/AvaRundownPageTransition.cpp b/Engine/Plugins/VirtualProduction/Avalanche/Source/AvalancheMedia/Private/Rundown/Transition/AvaRundownPageTransition.cpp
--- a/Engine/Plugins/VirtualProduction/Avalanche/Source/AvalancheMedia/Private/Rundown/Transition/AvaRundownPageTransition.cpp
+++ b/Engine/Plugins/VirtualProduction/Avalanche/Source/AvalancheMedia/Private/Rundown/Transition/AvaRundownPageTransition.cpp
@@ -2,7 +2,9 @@
 
 #include "Rundown/Transition/AvaRundownPageTransition.h"
 
+#include "Broadcast/AvaBroadcast.h"
 #include "Playable/AvaPlayableGroupManager.h"
+#include "Playable/AvaPlayableUtils.h"
 #include "Playable/Transition/AvaPlayableTransition.h"
 #include "Playback/AvaPlaybackManager.h"
 #include "Playback/AvaPlaybackUtils.h"
@@ -92,23 +94,6 @@ namespace UE::AvaMedia::Rundown::PageTransition::Private
 		}
 		return nullptr;
 	}
-
-	bool HasLocalPlayable(TNotNull<const UAvaRundownPagePlayer*> InPlayer)
-	{
-		// Check if at least one playable is local
-		const int32 NumInstancePlayers = InPlayer->GetNumInstancePlayers();
-		for (int32 InstancePlayerIndex = 0; InstancePlayerIndex < NumInstancePlayers; ++InstancePlayerIndex)
-		{
-			if (const UAvaPlayable* Playable = GetPlayable(InPlayer, InstancePlayerIndex))
-			{
-				if (!Playable->IsRemoteProxy())
-				{
-					return true;
-				}
-			}
-		}
-		return false;
-	}
 }
 
 UAvaRundownPageTransition* UAvaRundownPageTransition::MakeNew(UAvaRundown* InRundown)
@@ -205,13 +190,19 @@ bool UAvaRundownPageTransition::CanStart(bool& bOutShouldDiscard)
 		// Start loading the RC referenced assets.
 		AsyncAssetLoader = MakeShared<UE::AvaPlayback::Utils::FAsyncAssetLoader>();
 		TSet<FSoftObjectPath> Assets;
+
+		const UAvaBroadcast& Broadcast = UAvaBroadcast::Get();
+
 		for (const TWeakObjectPtr<UAvaRundownPagePlayer>& PlayerWeak : EnterPlayersWeak)
 		{
 			const UAvaRundownPagePlayer* Player = PlayerWeak.Get();
+			if (!Player)
+			{
+				continue;
+			}
 
-			// Only collect references and assets if there is at least one playable that isn't remote.
-			// If all playables are remote, no loading needs to take place locally.
-			if (Player && HasLocalPlayable(Player))
+			// Only collect references and assets to load for player with local output.
+			if (Player->bIsPreview || UE::AvaMedia::PlayableUtils::HasLocalPlay(Broadcast, Player->ChannelFName))
 			{
 				if (const TSharedPtr<const FAvaPlayableRemoteControlValues> Values = GetRemoteControlValues(Player))
 				{
