Motion Design: added execute controller column to be able to force execute a controller within page details. This is so that the user avoids having to commit the controller's value to a different one and back.
Also added an option to execute controllers from top to bottom in the column.

#Synced-CL 48030442

diff --git a/Engine/Plugins/VirtualProduction/Avalanche/Source/AvalancheMediaEditor/Private/Rundown/DetailsView/RemoteControl/Controllers/AvaRundownRCControllerItem.cpp b/Engine/Plugins/VirtualProduction/Avalanche/Source/AvalancheMediaEditor/Private/Rundown/DetailsView/RemoteControl/Controllers/AvaRundownRCControllerItem.cpp
--- a/Engine/Plugins/VirtualProduction/Avalanche/Source/AvalancheMediaEditor/Private/Rundown/DetailsView/RemoteControl/Controllers/AvaRundownRCControllerItem.cpp
+++ b/Engine/Plugins/VirtualProduction/Avalanche/Source/AvalancheMediaEditor/Private/Rundown/DetailsView/RemoteControl/Controllers/AvaRundownRCControllerItem.cpp
@@ -44,7 +44,7 @@ FAvaRundownRCControllerItem::FAvaRundownRCControllerItem(int32 InInstanceIndex,
 {
 	AssetName = InAssetName;
 	InstanceIndex = InInstanceIndex;
-	Controller = InController;
+	ControllerWeak = InController;
 	
 	if (InController)
 	{
@@ -131,7 +131,20 @@ TSharedRef<ITableRow> FAvaRundownRCControllerItem::CreateWidget(TSharedRef<SAvaR
 
 URCController* FAvaRundownRCControllerItem::GetController() const
 {
-	return Controller.Get();
+	return ControllerWeak.Get();
 }
 
-#undef LOCTEXT_NAMESPACE
\ No newline at end of file
+void FAvaRundownRCControllerItem::ExecuteController() const
+{
+	if (URCController* Controller = GetController())
+	{
+		Controller->ExecuteBehaviours();
+
+		if (URemoteControlPreset* Preset = Controller->PresetWeakPtr.Get())
+		{
+			Preset->OnControllerModified().Broadcast(Preset, {Controller->Id});
+		}
+	}
+}
+
+#undef LOCTEXT_NAMESPACE
diff --git a/Engine/Plugins/VirtualProduction/Avalanche/Source/AvalancheMediaEditor/Private/Rundown/DetailsView/RemoteControl/Controllers/AvaRundownRCControllerItem.h b/Engine/Plugins/VirtualProduction/Avalanche/Source/AvalancheMediaEditor/Private/Rundown/DetailsView/RemoteControl/Controllers/AvaRundownRCControllerItem.h
--- a/Engine/Plugins/VirtualProduction/Avalanche/Source/AvalancheMediaEditor/Private/Rundown/DetailsView/RemoteControl/Controllers/AvaRundownRCControllerItem.h
+++ b/Engine/Plugins/VirtualProduction/Avalanche/Source/AvalancheMediaEditor/Private/Rundown/DetailsView/RemoteControl/Controllers/AvaRundownRCControllerItem.h
@@ -30,7 +30,10 @@ public:
 	FName GetAssetName() const { return AssetName; }
 
 	URCController* GetController() const;
-	
+
+	/** Executes the behaviors of the controller that this item references */
+	void ExecuteController() const;
+
 private:
 	int32 InstanceIndex = 0;
 	int32 DisplayIndex = 0;
@@ -40,5 +43,5 @@ private:
 	FText ToolTipText;
 	FNodeWidgets NodeWidgets;
 
-	TWeakObjectPtr<URCController> Controller;
+	TWeakObjectPtr<URCController> ControllerWeak;
 };
diff --git a/Engine/Plugins/VirtualProduction/Avalanche/Source/AvalancheMediaEditor/Private/Rundown/DetailsView/RemoteControl/Controllers/SAvaRundownRCControllerItemRow.cpp b/Engine/Plugins/VirtualProduction/Avalanche/Source/AvalancheMediaEditor/Private/Rundown/DetailsView/RemoteControl/Controllers/SAvaRundownRCControllerItemRow.cpp
--- a/Engine/Plugins/VirtualProduction/Avalanche/Source/AvalancheMediaEditor/Private/Rundown/DetailsView/RemoteControl/Controllers/SAvaRundownRCControllerItemRow.cpp
+++ b/Engine/Plugins/VirtualProduction/Avalanche/Source/AvalancheMediaEditor/Private/Rundown/DetailsView/RemoteControl/Controllers/SAvaRundownRCControllerItemRow.cpp
@@ -1,12 +1,15 @@
 // Copyright Epic Games, Inc. All Rights Reserved.
 
 #include "SAvaRundownRCControllerItemRow.h"
-
 #include "AvaRundownRCControllerItem.h"
+#include "Widgets/Images/SImage.h"
+#include "Widgets/Input/SButton.h"
 #include "Widgets/Layout/SScissorRectBox.h"
 #include "Widgets/SNullWidget.h"
 #include "Widgets/Text/STextBlock.h"
 
+#define LOCTEXT_NAMESPACE "SAvaRundownRCControllerItemRow"
+
 void SAvaRundownRCControllerItemRow::Construct(const FArguments& InArgs, TSharedRef<SAvaRundownRCControllerPanel> InControllerPanel,
 	const TSharedRef<STableViewBase>& InOwnerTableView, const TSharedPtr<const FAvaRundownRCControllerItem>& InRowItem)
 {
@@ -61,27 +64,39 @@ TSharedRef<SWidget> SAvaRundownRCControllerItemRow::GenerateWidgetForColumn(cons
 				return ItemPtr->GetNodeWidgets().WholeRowWidget.ToSharedRef();
 			}
 		}
-		else
+		else if (InColumnName == SAvaRundownRCControllerPanel::ExecuteColumnName)
 		{
-			TSharedPtr<SAvaRundownRCControllerPanel> ControllerPanel = ControllerPanelWeak.Pin();
-
-			if (ControllerPanel.IsValid())
-			{
-				TSharedPtr<SWidget> Cell = nullptr;
-				const TArray<FAvaRundownRCControllerTableRowExtensionDelegate>& TableRowExtensionDelegates = ControllerPanel->GetTableRowExtensionDelegates(InColumnName);
-
-				for (const FAvaRundownRCControllerTableRowExtensionDelegate& TableRowExtensionDelegate : TableRowExtensionDelegates)
-				{
-					TableRowExtensionDelegate.ExecuteIfBound(ControllerPanel.ToSharedRef(), ItemPtr.ToSharedRef(), Cell);
-				}
-
-				if (Cell.IsValid())
-				{
-					return Cell.ToSharedRef();
-				}
-			}
+			return SNew(SButton)
+				.HAlign(HAlign_Center)
+				.VAlign(VAlign_Center)
+				.ButtonStyle(FAppStyle::Get(), "SimpleButton")
+				.ToolTipText(LOCTEXT("ExecuteController_Tooltip", "Executes the controller's behaviors, as if its value had been committed."))
+				.ContentPadding(FMargin(4, 0))
+				.Visibility(this, &SAvaRundownRCControllerItemRow::GetExecuteButtonVisibility)
+				.OnClicked(this, &SAvaRundownRCControllerItemRow::OnExecuteController)
+				[
+					SNew(SImage)
+					.ColorAndOpacity(FSlateColor::UseForeground())
+					.Image(FAppStyle::Get().GetBrush("Icons.Play"))
+				];
 		}
 	}
 
 	return SNullWidget::NullWidget;
 }
+
+EVisibility SAvaRundownRCControllerItemRow::GetExecuteButtonVisibility() const
+{
+	return IsHovered() ? EVisibility::Visible : EVisibility::Hidden;
+}
+
+FReply SAvaRundownRCControllerItemRow::OnExecuteController()
+{
+	if (TSharedPtr<const FAvaRundownRCControllerItem> Item = ItemPtrWeak.Pin())
+	{
+		Item->ExecuteController();
+	}
+	return FReply::Handled();
+}
+
+#undef LOCTEXT_NAMESPACE
diff --git a/Engine/Plugins/VirtualProduction/Avalanche/Source/AvalancheMediaEditor/Private/Rundown/DetailsView/RemoteControl/Controllers/SAvaRundownRCControllerItemRow.h b/Engine/Plugins/VirtualProduction/Avalanche/Source/AvalancheMediaEditor/Private/Rundown/DetailsView/RemoteControl/Controllers/SAvaRundownRCControllerItemRow.h
--- a/Engine/Plugins/VirtualProduction/Avalanche/Source/AvalancheMediaEditor/Private/Rundown/DetailsView/RemoteControl/Controllers/SAvaRundownRCControllerItemRow.h
+++ b/Engine/Plugins/VirtualProduction/Avalanche/Source/AvalancheMediaEditor/Private/Rundown/DetailsView/RemoteControl/Controllers/SAvaRundownRCControllerItemRow.h
@@ -20,6 +20,9 @@ public:
 	virtual TSharedRef<SWidget> GenerateWidgetForColumn(const FName& InColumnName) override;
 
 protected:
+	EVisibility GetExecuteButtonVisibility() const;
+	FReply OnExecuteController();
+
 	TWeakPtr<const FAvaRundownRCControllerItem> ItemPtrWeak;
 	TWeakPtr<SAvaRundownRCControllerPanel> ControllerPanelWeak;
 };
diff --git a/Engine/Plugins/VirtualProduction/Avalanche/Source/AvalancheMediaEditor/Private/Rundown/DetailsView/RemoteControl/Controllers/SAvaRundownRCControllerPanel.cpp b/Engine/Plugins/VirtualProduction/Avalanche/Source/AvalancheMediaEditor/Private/Rundown/DetailsView/RemoteControl/Controllers/SAvaRundownRCControllerPanel.cpp
--- a/Engine/Plugins/VirtualProduction/Avalanche/Source/AvalancheMediaEditor/Private/Rundown/DetailsView/RemoteControl/Controllers/SAvaRundownRCControllerPanel.cpp
+++ b/Engine/Plugins/VirtualProduction/Avalanche/Source/AvalancheMediaEditor/Private/Rundown/DetailsView/RemoteControl/Controllers/SAvaRundownRCControllerPanel.cpp
@@ -20,17 +20,10 @@
 
 #define LOCTEXT_NAMESPACE "SAvaRundownRCControllerPanel"
 
-const FName SAvaRundownRCControllerPanel::ControllerColumnName = "ControllerColumn";
-const FName SAvaRundownRCControllerPanel::DescriptionColumnName = "DescriptionColumn";
-const FName SAvaRundownRCControllerPanel::ValueColumnName = "ValueColumn";
-
-FAvaRundownRCControllerHeaderRowExtensionDelegate SAvaRundownRCControllerPanel::HeaderRowExtensionDelegate;
-TMap<FName, TArray<FAvaRundownRCControllerTableRowExtensionDelegate>> SAvaRundownRCControllerPanel::TableRowExtensionDelegates;
-
-TArray<FAvaRundownRCControllerTableRowExtensionDelegate>& SAvaRundownRCControllerPanel::GetTableRowExtensionDelegates(FName InExtensionName)
-{
-	return TableRowExtensionDelegates.FindOrAdd(InExtensionName);
-}
+const FLazyName SAvaRundownRCControllerPanel::ControllerColumnName = "ControllerColumn";
+const FLazyName SAvaRundownRCControllerPanel::DescriptionColumnName = "DescriptionColumn";
+const FLazyName SAvaRundownRCControllerPanel::ValueColumnName = "ValueColumn";
+const FLazyName SAvaRundownRCControllerPanel::ExecuteColumnName = "ExecuteColumn";
 
 void SAvaRundownRCControllerPanel::Construct(const FArguments& InArgs, const TSharedPtr<FAvaRundownEditor>& InRundownEditor)
 {
@@ -64,6 +57,24 @@ void SAvaRundownRCControllerPanel::Construct(const FArguments& InArgs, const TSh
 				+ SHeaderRow::Column(ValueColumnName)
 				.DefaultLabel(LOCTEXT("Value", "Value"))
 				.FillWidth(0.7f)
+				+ SHeaderRow::Column(ExecuteColumnName)
+				.DefaultLabel(LOCTEXT("Execute", "Execute"))
+				.ManualWidth(64)
+				.HeaderContent()
+				[
+					SNew(SButton)
+					.HAlign(HAlign_Center)
+					.VAlign(VAlign_Center)
+					.ButtonStyle(FAppStyle::Get(), "SimpleButton")
+					.ToolTipText(LOCTEXT("ExecuteController_Tooltip", "Executes all the controllers' behaviors in this list from top to bottom order, as if their values had been committed."))
+					.ContentPadding(FMargin(4, 0))
+					.OnClicked(this, &SAvaRundownRCControllerPanel::OnExecuteAllControllers)
+					[
+						SNew(SImage)
+						.ColorAndOpacity(FSlateColor::UseForeground())
+						.Image(FAppStyle::Get().GetBrush("Icons.Play"))
+					]
+				]
 			)
 		]
 	];
@@ -427,6 +438,18 @@ void SAvaRundownRCControllerPanel::UpdatePageSummary(bool bInForceUpdate)
 	}
 }
 
+FReply SAvaRundownRCControllerPanel::OnExecuteAllControllers()
+{
+	for (FAvaRundownRCControllerItemPtr& ControllerItem : ControllerItems)
+	{
+		if (ControllerItem.IsValid())
+		{
+			ControllerItem->ExecuteController();
+		}
+	}
+	return FReply::Handled();
+}
+
 UAvaRundown* SAvaRundownRCControllerPanel::GetRundown() const
 {
 	if (const TSharedPtr<FAvaRundownEditor> RundownEditor = RundownEditorWeak.Pin())
diff --git a/Engine/Plugins/VirtualProduction/Avalanche/Source/AvalancheMediaEditor/Private/Rundown/DetailsView/RemoteControl/Controllers/SAvaRundownRCControllerPanel.h b/Engine/Plugins/VirtualProduction/Avalanche/Source/AvalancheMediaEditor/Private/Rundown/DetailsView/RemoteControl/Controllers/SAvaRundownRCControllerPanel.h
--- a/Engine/Plugins/VirtualProduction/Avalanche/Source/AvalancheMediaEditor/Private/Rundown/DetailsView/RemoteControl/Controllers/SAvaRundownRCControllerPanel.h
+++ b/Engine/Plugins/VirtualProduction/Avalanche/Source/AvalancheMediaEditor/Private/Rundown/DetailsView/RemoteControl/Controllers/SAvaRundownRCControllerPanel.h
@@ -27,20 +27,13 @@ template <typename ItemType> class SListView;
 
 using FAvaRundownRCControllerItemPtr = TSharedPtr<FAvaRundownRCControllerItem>;
 
-DECLARE_MULTICAST_DELEGATE_TwoParams(FAvaRundownRCControllerHeaderRowExtensionDelegate, TSharedRef<SAvaRundownRCControllerPanel> Panel, 
-	TSharedRef<SHeaderRow>& HeaderRow)
-DECLARE_DELEGATE_ThreeParams(FAvaRundownRCControllerTableRowExtensionDelegate, TSharedRef<SAvaRundownRCControllerPanel> Panel, 
-	TSharedRef<const FAvaRundownRCControllerItem> ItemPtr, TSharedPtr<SWidget>& CurrentWidget)
-
 class SAvaRundownRCControllerPanel : public SCompoundWidget
 {
 public:
-	static const FName ControllerColumnName;
-	static const FName DescriptionColumnName;
-	static const FName ValueColumnName;
-
-	static FAvaRundownRCControllerHeaderRowExtensionDelegate& GetHeaderRowExtensionDelegate() { return HeaderRowExtensionDelegate; }
-	static TArray<FAvaRundownRCControllerTableRowExtensionDelegate>& GetTableRowExtensionDelegates(FName InExtensionName);
+	static const FLazyName ControllerColumnName;
+	static const FLazyName DescriptionColumnName;
+	static const FLazyName ValueColumnName;
+	static const FLazyName ExecuteColumnName;
 
 	SLATE_BEGIN_ARGS(SAvaRundownRCControllerPanel) {}
 	SLATE_END_ARGS()
@@ -58,8 +51,8 @@ public:
 	void UpdatePageSummary(bool bInForceUpdate);
 
 private:
-	static FAvaRundownRCControllerHeaderRowExtensionDelegate HeaderRowExtensionDelegate;
-	static TMap<FName, TArray<FAvaRundownRCControllerTableRowExtensionDelegate>> TableRowExtensionDelegates;
+	/** Executes all the controllers from the items in this panel, in array order */
+	FReply OnExecuteAllControllers();
 
 	void UpdatePropertyRowGenerators(int32 InNumGenerators);
 	
