Motion Design: Moved some Asset Async Loader functionality so that Rundown Load Page can reuse.

#Synced-CL 48030442

diff --git a/Engine/Plugins/VirtualProduction/Avalanche/Source/AvalancheMaterial/Internal/AvaMaterialCacheHelper.h b/Engine/Plugins/VirtualProduction/Avalanche/Source/AvalancheMaterial/Internal/AvaMaterialCacheHelper.h
--- a/Engine/Plugins/VirtualProduction/Avalanche/Source/AvalancheMaterial/Internal/AvaMaterialCacheHelper.h
+++ b/Engine/Plugins/VirtualProduction/Avalanche/Source/AvalancheMaterial/Internal/AvaMaterialCacheHelper.h
@@ -43,9 +43,10 @@ public:
 
 	/**
 	 * Caches the materials in the given object.
-	 * The materials are found via the material bridge api
+	 * The materials are found via the material bridge api.
+	 * Returns true if materials were gathered for the object, false otherwise. (e.g. if object's materials were already gathered, or no material bridge found)
 	 */
-	UE_API void RequestCacheMaterials(const UObject* InObject);
+	UE_API bool RequestCacheMaterials(const UObject* InObject);
 
 	/** Processes materials pending completion with a limited time budget */
 	UE_API void Tick();
@@ -86,8 +87,11 @@ private:
 	/** Process async results if shader is compiling */
 	void ProcessShadersCompiling();
 
-	/** Gathers the material data for an object, or skips it if the object's materials have been gathered before. */
-	void GatherMaterialData(const UObject* InObject);
+	/**
+	 * Gathers the material data for an object, or skips it if the object's materials have been gathered before.
+	 * returns true if a material bridge was found for the object and visited its materials
+	 */
+	bool GatherMaterialData(const UObject* InObject);
 
 	/** Objects that have already been gathered, mapped to the incomplete materials that were gathered for the object */
 	TMap<FObjectKey, FObjectData> Objects;
diff --git a/Engine/Plugins/VirtualProduction/Avalanche/Source/AvalancheMaterial/Private/AvaMaterialCacheHelper.cpp b/Engine/Plugins/VirtualProduction/Avalanche/Source/AvalancheMaterial/Private/AvaMaterialCacheHelper.cpp
--- a/Engine/Plugins/VirtualProduction/Avalanche/Source/AvalancheMaterial/Private/AvaMaterialCacheHelper.cpp
+++ b/Engine/Plugins/VirtualProduction/Avalanche/Source/AvalancheMaterial/Private/AvaMaterialCacheHelper.cpp
@@ -41,15 +41,16 @@ FMaterialCacheHelper& FMaterialCacheHelper::Get()
 	return *MaterialCacheHelper;
 }
 
-void FMaterialCacheHelper::RequestCacheMaterials(const UObject* InObject)
+bool FMaterialCacheHelper::RequestCacheMaterials(const UObject* InObject)
 {
 	if (!InObject)
 	{
-		return;
+		return false;
 	}
 
-	GatherMaterialData(InObject);
+	const bool bMaterialsGathered = GatherMaterialData(InObject);
 	RegisterTick();
+	return bMaterialsGathered;
 }
 
 void FMaterialCacheHelper::Tick()
@@ -250,11 +251,11 @@ void FMaterialCacheHelper::ProcessShadersCompiling()
 	RemainingTimeBudget -= (FPlatformTime::Seconds() - StartTime);
 }
 
-void FMaterialCacheHelper::GatherMaterialData(const UObject* InObject)
+bool FMaterialCacheHelper::GatherMaterialData(const UObject* InObject)
 {
 	if (!InObject || Objects.Contains(InObject))
 	{
-		return;
+		return false;
 	}
 
 	TRACE_CPUPROFILER_EVENT_SCOPE(FMaterialCacheHelper::GatherMaterialData);
@@ -262,8 +263,8 @@ void FMaterialCacheHelper::GatherMaterialData(const UObject* InObject)
 	const FMaterialBridge* MaterialBridge = FMaterialBridgeRegistry::Get().GetMaterialBridgeForObject(InObject);
 	if (!MaterialBridge)
 	{
-		UE_LOG(LogAvaMaterialCache, Warning, TEXT("Object '%s' does not have a valid bridge to cache materials!"), *GetNameSafe(InObject));
-		return;
+		UE_LOG(LogAvaMaterialCache, Verbose, TEXT("Object '%s' does not have a valid bridge to cache materials!"), *GetNameSafe(InObject));
+		return false;
 	}
 
 	FObjectData& ObjectData = Objects.Add(InObject);
@@ -302,6 +303,8 @@ void FMaterialCacheHelper::GatherMaterialData(const UObject* InObject)
 				AddIncompleteMaterial(InContext.Material);
 			}
 		});
+
+	return true;
 }
 
 } // UE::Ava
\ No newline at end of file
diff --git a/Engine/Plugins/VirtualProduction/Avalanche/Source/AvalancheMedia/Internal/Playback/AvaPlaybackUtils.h b/Engine/Plugins/VirtualProduction/Avalanche/Source/AvalancheMedia/Internal/Playback/AvaPlaybackUtils.h
--- a/Engine/Plugins/VirtualProduction/Avalanche/Source/AvalancheMedia/Internal/Playback/AvaPlaybackUtils.h
+++ b/Engine/Plugins/VirtualProduction/Avalanche/Source/AvalancheMedia/Internal/Playback/AvaPlaybackUtils.h
@@ -69,13 +69,22 @@ namespace UE::AvaPlayback::Utils
 		return StaticEnum<InEnumType>()->GetNameStringByValue(static_cast<int64>(InValue));
 	}
 
+	/**
+	 * Streams in the given object if it's a streamable asset.
+	 * returns a valid pointer to the object as a render asset if the asset was streamable
+	 */
+	UStreamableRenderAsset* StreamRenderAsset(UObject* InObject);
+
+	/** Caches the materials of the given object using the material bridge api */
+	bool CacheObjectMaterials(UObject* InObject);
+
 	/**
 	 * Utility class to perform async load of assets. 
 	 */
 	class FAsyncAssetLoader : public TSharedFromThis<FAsyncAssetLoader>
 	{
 	public:
-		void HandleStreamableObject(UObject* InObject);
+		void ProcessLoadedObject(UObject* InObject);
 
 		/**
 		 * @brief Issue the async load command. Returns immediately.
@@ -88,8 +97,11 @@ namespace UE::AvaPlayback::Utils
 		/** Returns true when all the requested assets have finished loading. */
 		AVALANCHEMEDIA_API bool IsLoadingCompleted() const;
 
+		/** Removes objects that have all their materials completed */
+		void CleanupObjectsWithCompleteMaterials() const;
+
 		/** Removes streamable render assets that have fully been streamed in */
-		AVALANCHEMEDIA_API void CleanupStreamingAssets() const;
+		void CleanupStreamingAssets() const;
 
 		/** Delegate called when all assets have finished loading. */
 		DECLARE_MULTICAST_DELEGATE(FOnLoadingCompleted);
@@ -97,6 +109,7 @@ namespace UE::AvaPlayback::Utils
 
 	private:
 		TSet<FSoftObjectPath> PendingAssets;
+		mutable TArray<TObjectPtr<UObject>> PendingMaterialCompletionObjects;
 		mutable TArray<TObjectPtr<UStreamableRenderAsset>> PendingStreamingAssets;
 	};
 }
diff --git a/Engine/Plugins/VirtualProduction/Avalanche/Source/AvalancheMedia/Private/Playback/AvaPlaybackUtils.cpp b/Engine/Plugins/VirtualProduction/Avalanche/Source/AvalancheMedia/Private/Playback/AvaPlaybackUtils.cpp
--- a/Engine/Plugins/VirtualProduction/Avalanche/Source/AvalancheMedia/Private/Playback/AvaPlaybackUtils.cpp
+++ b/Engine/Plugins/VirtualProduction/Avalanche/Source/AvalancheMedia/Private/Playback/AvaPlaybackUtils.cpp
@@ -34,9 +34,15 @@ namespace UE::AvaPlayback::Utils::Private
 	);
 
 	static TAutoConsoleVariable<bool> CVarRuntimePrestreamRCReferencedAssets(
-		TEXT("MotionDesignPlayback.RuntimePrestreamRemoteControlReferencedAssets"),
+		TEXT("MotionDesignPlayback.RuntimePrestreamRemoteControlObjects"),
 		true,
-		TEXT("Manually streams and waits for all Remote Control referenced assets to be ready while preparing to start play.")
+		TEXT("Manually streams and waits for all Remote Control referenced objects to be ready while preparing to start play.")
+	);
+
+	static TAutoConsoleVariable<bool> CVarRuntimePrecacheRCReferencedAssetsMaterials(
+		TEXT("MotionDesignPlayback.RuntimePrecacheRemoteControlObjectMaterials"),
+		true,
+		TEXT("Manually caches and waits for materials used by Remote Control referenced objects to be complete while preparing to start play.")
 	);
 }
 
@@ -240,25 +246,41 @@ namespace UE::AvaPlayback::Utils
 		return FString::Printf(TEXT("[%d]"), GFrameNumber);
 	}
 
-	void FAsyncAssetLoader::HandleStreamableObject(UObject* InObject)
+	UStreamableRenderAsset* StreamRenderAsset(UObject* InObject)
 	{
-		if (!InObject || !Private::CVarRuntimePrestreamRCReferencedAssets.GetValueOnGameThread())
-		{
-			// Not allowed to stream assets.
-			return;
-		}
-
 		UStreamableRenderAsset* const StreamableAsset = Cast<UStreamableRenderAsset>(InObject);
+
 		if (StreamableAsset && StreamableAsset->IsStreamable())
 		{
-			constexpr float DurationSeconds = 30.f;
-			constexpr uint32 CinematicLODGroupMask = 0xFFFFFFFF;
-			StreamableAsset->SetForceMipLevelsToBeResident(DurationSeconds, CinematicLODGroupMask);
-
-			constexpr bool bHighPriority = true;
+			constexpr bool bHighPriority = false;
 			StreamableAsset->StreamIn(FStreamableRenderResourceState::MAX_LOD_COUNT, bHighPriority);
+			return StreamableAsset;
+		}
+
+		return nullptr;
+	}
+
+	bool CacheObjectMaterials(UObject* InObject)
+	{
+		return UE::Ava::FMaterialCacheHelper::Get().RequestCacheMaterials(InObject);
+	}
+
+	void FAsyncAssetLoader::ProcessLoadedObject(UObject* InObject)
+	{
+		if (Private::CVarRuntimePrestreamRCReferencedAssets.GetValueOnGameThread())
+		{
+			if (UStreamableRenderAsset* RenderAsset = StreamRenderAsset(InObject))
+			{
+				PendingStreamingAssets.Add(RenderAsset);
+			}
+		}
 
-			PendingStreamingAssets.Add(StreamableAsset);
+		if (Private::CVarRuntimePrecacheRCReferencedAssetsMaterials.GetValueOnGameThread())
+		{
+			if (CacheObjectMaterials(InObject))
+			{
+				PendingMaterialCompletionObjects.Add(InObject);
+			}
 		}
 	}
 
@@ -275,8 +297,8 @@ namespace UE::AvaPlayback::Utils
 			UObject* ResolvedObject = AssetToLoad.ResolveObject();
 			if (ResolvedObject)
 			{
-				// Already loaded... but streamable objects could still require streaming in
-				HandleStreamableObject(ResolvedObject);
+				// Already loaded... but streamable objects could still require streaming in, or it could have materials to cache.
+				ProcessLoadedObject(ResolvedObject);
 				continue;
 			}
 
@@ -289,7 +311,7 @@ namespace UE::AvaPlayback::Utils
 				{
 					OnLoadingCompleted.Broadcast();
 				}
-				HandleStreamableObject(InLoadedObject);
+				ProcessLoadedObject(InLoadedObject);
 			}));
 		}
 	}
@@ -302,26 +324,39 @@ namespace UE::AvaPlayback::Utils
 		}
 
 		CleanupStreamingAssets();
-
 		if (!PendingStreamingAssets.IsEmpty())
 		{
 			return false;
 		}
 
+		CleanupObjectsWithCompleteMaterials();
+		if (!PendingMaterialCompletionObjects.IsEmpty())
+		{
+			return false;
+		}
+
 		return true;
 	}
 
 	void FAsyncAssetLoader::CleanupStreamingAssets() const
 	{
-		UE::Ava::FMaterialCacheHelper& MaterialCacheHelper = UE::Ava::FMaterialCacheHelper::Get();
+		IRenderAssetStreamingManager& RenderAssetStreamingManager = IStreamingManager::Get().GetRenderAssetStreamingManager();
 
 		PendingStreamingAssets.RemoveAll(
-			[&MaterialCacheHelper](UStreamableRenderAsset* InObject)
+			[&RenderAssetStreamingManager](UStreamableRenderAsset* InObject)
+			{
+				return !InObject || (!InObject->HasPendingInitOrStreaming() && RenderAssetStreamingManager.IsFullyStreamedIn(InObject));
+			});
+	}
+
+	void FAsyncAssetLoader::CleanupObjectsWithCompleteMaterials() const
+	{
+		UE::Ava::FMaterialCacheHelper& MaterialCacheHelper = UE::Ava::FMaterialCacheHelper::Get();
+
+		PendingMaterialCompletionObjects.RemoveAll(
+			[&MaterialCacheHelper](UObject* InObject)
 			{
-				constexpr bool bWaitForLODTransition = true;
-				return !InObject->HasPendingInitOrStreaming(bWaitForLODTransition)
-					&& InObject->IsFullyStreamedIn()
-					&& !MaterialCacheHelper.IsCaching(InObject);
+				return !InObject || !MaterialCacheHelper.IsCaching(InObject);
 			});
 	}
 }
diff --git a/Engine/Plugins/VirtualProduction/Avalanche/Source/AvalancheMedia/Private/Rundown/AvaRundown.cpp b/Engine/Plugins/VirtualProduction/Avalanche/Source/AvalancheMedia/Private/Rundown/AvaRundown.cpp
--- a/Engine/Plugins/VirtualProduction/Avalanche/Source/AvalancheMedia/Private/Rundown/AvaRundown.cpp
+++ b/Engine/Plugins/VirtualProduction/Avalanche/Source/AvalancheMedia/Private/Rundown/AvaRundown.cpp
@@ -1207,7 +1207,7 @@ TArray<UAvaRundown::FLoadedInstanceInfo> UAvaRundown::LoadPage(int32 InPageId,
 		LoadedInstances.Add({PlaybackInstance->GetInstanceId(), AssetPath});
 	}
 
-	// Load RC Assets Async and request cache its materials
+	// Load RC Assets Async, stream them and request cache its materials
 	{
 		TSet<FSoftObjectPath> RemoteControlAssets;
 
@@ -1219,7 +1219,8 @@ TArray<UAvaRundown::FLoadedInstanceInfo> UAvaRundown::LoadPage(int32 InPageId,
 		{
 			if (UObject* Object = Asset.ResolveObject())
 			{
-				UE::Ava::FMaterialCacheHelper::Get().RequestCacheMaterials(Object);
+				UE::AvaPlayback::Utils::StreamRenderAsset(Object);
+				UE::AvaPlayback::Utils::CacheObjectMaterials(Object);
 			}
 			else
 			{
@@ -1228,7 +1229,8 @@ TArray<UAvaRundown::FLoadedInstanceInfo> UAvaRundown::LoadPage(int32 InPageId,
 					{
 						if (InObject)
 						{
-							UE::Ava::FMaterialCacheHelper::Get().RequestCacheMaterials(InObject);
+							UE::AvaPlayback::Utils::StreamRenderAsset(InObject);
+							UE::AvaPlayback::Utils::CacheObjectMaterials(InObject);
 						}
 					}));
 			}
