Motion Design: added FGCObject to FAssetAsyncLoader

#Synced-CL 48030442

diff --git a/Engine/Plugins/VirtualProduction/Avalanche/Source/AvalancheMedia/Internal/Playback/AvaPlaybackUtils.h b/Engine/Plugins/VirtualProduction/Avalanche/Source/AvalancheMedia/Internal/Playback/AvaPlaybackUtils.h
--- a/Engine/Plugins/VirtualProduction/Avalanche/Source/AvalancheMedia/Internal/Playback/AvaPlaybackUtils.h
+++ b/Engine/Plugins/VirtualProduction/Avalanche/Source/AvalancheMedia/Internal/Playback/AvaPlaybackUtils.h
@@ -6,6 +6,7 @@
 #include "Containers/UnrealString.h"
 #include "Delegates/Delegate.h"
 #include "Templates/SharedPointer.h"
+#include "UObject/GCObject.h"
 #include "UObject/SoftObjectPath.h"
 
 class UPackage;
@@ -81,7 +82,7 @@ namespace UE::AvaPlayback::Utils
 	/**
 	 * Utility class to perform async load of assets. 
 	 */
-	class FAsyncAssetLoader : public TSharedFromThis<FAsyncAssetLoader>
+	class FAsyncAssetLoader : public FGCObject, public TSharedFromThis<FAsyncAssetLoader>
 	{
 	public:
 		void ProcessLoadedObject(UObject* InObject);
@@ -103,6 +104,11 @@ namespace UE::AvaPlayback::Utils
 		/** Removes streamable render assets that have fully been streamed in */
 		void CleanupStreamingAssets() const;
 
+		//~ Begin FGCObject
+		AVALANCHEMEDIA_API virtual void AddReferencedObjects(FReferenceCollector& InCollector) override;
+		AVALANCHEMEDIA_API virtual FString GetReferencerName() const override;
+		//~ End FGCObject
+
 		/** Delegate called when all assets have finished loading. */
 		DECLARE_MULTICAST_DELEGATE(FOnLoadingCompleted);
 		FOnLoadingCompleted OnLoadingCompleted;
diff --git a/Engine/Plugins/VirtualProduction/Avalanche/Source/AvalancheMedia/Private/Playback/AvaPlaybackUtils.cpp b/Engine/Plugins/VirtualProduction/Avalanche/Source/AvalancheMedia/Private/Playback/AvaPlaybackUtils.cpp
--- a/Engine/Plugins/VirtualProduction/Avalanche/Source/AvalancheMedia/Private/Playback/AvaPlaybackUtils.cpp
+++ b/Engine/Plugins/VirtualProduction/Avalanche/Source/AvalancheMedia/Private/Playback/AvaPlaybackUtils.cpp
@@ -359,4 +359,15 @@ namespace UE::AvaPlayback::Utils
 				return !InObject || !MaterialCacheHelper.IsCaching(InObject);
 			});
 	}
+
+	void FAsyncAssetLoader::AddReferencedObjects(FReferenceCollector& InCollector)
+	{
+		InCollector.AddReferencedObjects(PendingMaterialCompletionObjects);
+		InCollector.AddReferencedObjects(PendingStreamingAssets);
+	}
+
+	FString FAsyncAssetLoader::GetReferencerName() const
+	{
+		return TEXT("UE::AvaPlayback::Utils::FAsyncAssetLoader");
+	}
 }
