Motion Design: remove erroneous fully streamed in check

#Synced-CL 48030442

diff --git a/Engine/Plugins/VirtualProduction/Avalanche/Source/AvalancheMedia/Private/Playback/AvaPlaybackUtils.cpp b/Engine/Plugins/VirtualProduction/Avalanche/Source/AvalancheMedia/Private/Playback/AvaPlaybackUtils.cpp
--- a/Engine/Plugins/VirtualProduction/Avalanche/Source/AvalancheMedia/Private/Playback/AvaPlaybackUtils.cpp
+++ b/Engine/Plugins/VirtualProduction/Avalanche/Source/AvalancheMedia/Private/Playback/AvaPlaybackUtils.cpp
@@ -340,12 +340,10 @@ namespace UE::AvaPlayback::Utils
 
 	void FAsyncAssetLoader::CleanupStreamingAssets() const
 	{
-		IRenderAssetStreamingManager& RenderAssetStreamingManager = IStreamingManager::Get().GetRenderAssetStreamingManager();
-
 		PendingStreamingAssets.RemoveAll(
-			[&RenderAssetStreamingManager](UStreamableRenderAsset* InObject)
+			[](UStreamableRenderAsset* InObject)
 			{
-				return !InObject || (!InObject->HasPendingInitOrStreaming() && RenderAssetStreamingManager.IsFullyStreamedIn(InObject));
+				return !InObject || !InObject->HasPendingInitOrStreaming();
 			});
 	}
 
