Motion Design: have playables by default disallow calling FlushRenderingCommands on an actor component's PreEditChange. Playables can allow it via the cvar CVarPlayableAllowFlushRenderingCommandsOnPreEditChange

#Synced-CL 48030442

diff --git a/Engine/Plugins/VirtualProduction/Avalanche/Source/AvalancheMedia/Private/Playable/AvaPlayable.cpp b/Engine/Plugins/VirtualProduction/Avalanche/Source/AvalancheMedia/Private/Playable/AvaPlayable.cpp
--- a/Engine/Plugins/VirtualProduction/Avalanche/Source/AvalancheMedia/Private/Playable/AvaPlayable.cpp
+++ b/Engine/Plugins/VirtualProduction/Avalanche/Source/AvalancheMedia/Private/Playable/AvaPlayable.cpp
@@ -52,6 +52,16 @@ static TAutoConsoleVariable<bool> CVarExecuteNonOverlappingRCControllerBehavior(
 	TEXT("Execute all non-overlapping controller behaviors."),
 	ECVF_Default);
 
+#if WITH_EDITOR
+// UActorComponent::PreEditChange calls FlushRenderingCommands() if GFlushRenderingCommandsOnPreEditChange is true (which it is by default)
+// Allow Playables to turn it off to reduce FlushRenderingCommands() calls in play time.
+static TAutoConsoleVariable<bool> CVarPlayableAllowFlushRenderingCommandsOnPreEditChange(
+	TEXT("MotionDesignPlayable.AllowFlushRenderingCommandsOnPreEditChange"),
+	false,
+	TEXT("Editor only! Determines whether FlushRenderingCommands() is allowed to be called by an Actor Component's PreEditChange when updating RC values in the playable"),
+	ECVF_Default);
+#endif
+
 namespace UE::AvaPlayable::Private
 {
 	FString GetPrettyPlayableInfo(const UAvaPlayable* InPlayable)
@@ -283,6 +293,20 @@ EAvaPlayableCommandResult UAvaPlayable::UpdateRemoteControlCommand(const TShared
 	using namespace UE::AvaPlayable::Private;
 	UE_LOG(LogAvaPlayable, Verbose, TEXT("%s Playable {%s} -> Executing RC Update."), *GetBriefFrameInfo(), *GetPrettyPlayableInfo(this));
 
+#if WITH_EDITOR
+	extern ENGINE_API bool GFlushRenderingCommandsOnPreEditChange;
+	bool bFlushRenderingCommandsOnPreEditChange = GFlushRenderingCommandsOnPreEditChange;
+
+	// If not allowing flush rendering commands on pre edit change, this will force the global to be set to false in this function scope.
+	// NOTE: If allowing the flush, this won't force the global to be true in case callers have set the global to false.
+	if (!CVarPlayableAllowFlushRenderingCommandsOnPreEditChange.GetValueOnGameThread())
+	{
+		bFlushRenderingCommandsOnPreEditChange = false;
+	}
+
+	TGuardValue<bool> ShouldFlushRenderingCommands(GFlushRenderingCommandsOnPreEditChange, bFlushRenderingCommandsOnPreEditChange);
+#endif
+
 	TSet<FGuid> AppliedOrIgnoredEntities;
 	
 	if (CVarUpdateRCControllerValues.GetValueOnGameThread())
