Motion Design: added level material bridge

#Synced-CL 48030442

diff --git a/Engine/Plugins/VirtualProduction/Avalanche/Source/AvalancheMaterial/Private/AvaActorMaterialBridge.cpp b/Engine/Plugins/VirtualProduction/Avalanche/Source/AvalancheMaterial/Private/AvaActorMaterialBridge.cpp
--- a/Engine/Plugins/VirtualProduction/Avalanche/Source/AvalancheMaterial/Private/AvaActorMaterialBridge.cpp
+++ b/Engine/Plugins/VirtualProduction/Avalanche/Source/AvalancheMaterial/Private/AvaActorMaterialBridge.cpp
@@ -14,16 +14,14 @@ const UStruct* FActorMaterialBridge::OnGetBridgedType() const
 
 void FActorMaterialBridge::OnVisitMaterials(FConstDataView InDataView, TFunctionRef<void(const FVisitedContext&)> InFunc, const FVisitOptions& InOptions) const
 {
-	const FMaterialBridgeRegistry& MaterialBridgeRegistry = FMaterialBridgeRegistry::Get();
-
 	const AActor& Actor = InDataView.Get<const AActor>();
 
 	constexpr bool bIncludeFromChildActors = true;
 
 	Actor.ForEachComponent(bIncludeFromChildActors,
-		[&MaterialBridgeRegistry, &InFunc, &InOptions](UActorComponent* InComponent)
+		[&InFunc, &InOptions](UActorComponent* InComponent)
 		{
-			if (const FMaterialBridge* MaterialBridge = MaterialBridgeRegistry.GetMaterialBridgeForObject(InComponent))
+			if (const FMaterialBridge* MaterialBridge = InOptions.MaterialBridgeRegistry->GetMaterialBridgeForObject(InComponent))
 			{
 				MaterialBridge->VisitMaterials(FConstDataView(InComponent), InFunc, InOptions);
 			}
diff --git a/Engine/Plugins/VirtualProduction/Avalanche/Source/AvalancheMaterial/Private/AvaLevelMaterialBridge.cpp b/Engine/Plugins/VirtualProduction/Avalanche/Source/AvalancheMaterial/Private/AvaLevelMaterialBridge.cpp
new file mode 100644
--- /dev/null
+++ b/Engine/Plugins/VirtualProduction/Avalanche/Source/AvalancheMaterial/Private/AvaLevelMaterialBridge.cpp
@@ -0,0 +1,30 @@
+// Copyright Epic Games, Inc. All Rights Reserved.
+
+#include "AvaLevelMaterialBridge.h"
+#include "AvaMaterialBridgeRegistry.h"
+#include "AvaMaterialCacheHelper.h"
+#include "Engine/Level.h"
+#include "GameFramework/Actor.h"
+
+namespace UE::Ava
+{
+
+const UStruct* FLevelMaterialBridge::OnGetBridgedType() const
+{
+	return ULevel::StaticClass();
+}
+
+void FLevelMaterialBridge::OnVisitMaterials(FConstDataView InDataView, TFunctionRef<void(const FVisitedContext&)> InFunc, const FVisitOptions& InOptions) const
+{
+	const ULevel& Level = InDataView.Get<const ULevel>();
+
+	for (const AActor* Actor : Level.Actors)
+	{
+		if (const FMaterialBridge* MaterialBridge = InOptions.MaterialBridgeRegistry->GetMaterialBridgeForObject(Actor))
+		{
+			MaterialBridge->VisitMaterials(FConstDataView(Actor), InFunc, InOptions);
+		}
+	}
+}
+
+} // UE::Ava
diff --git a/Engine/Plugins/VirtualProduction/Avalanche/Source/AvalancheMaterial/Private/AvaMaterialBridge.cpp b/Engine/Plugins/VirtualProduction/Avalanche/Source/AvalancheMaterial/Private/AvaMaterialBridge.cpp
--- a/Engine/Plugins/VirtualProduction/Avalanche/Source/AvalancheMaterial/Private/AvaMaterialBridge.cpp
+++ b/Engine/Plugins/VirtualProduction/Avalanche/Source/AvalancheMaterial/Private/AvaMaterialBridge.cpp
@@ -1,6 +1,7 @@
 // Copyright Epic Games, Inc. All Rights Reserved.
 
 #include "AvaMaterialBridge.h"
+#include "AvaMaterialBridgeRegistry.h"
 
 namespace UE::Ava
 {
@@ -15,6 +16,11 @@ const UStruct* FMaterialBridge::GetBridgedType() const
 	return BridgedType;
 }
 
+FMaterialBridge::FVisitOptions::FVisitOptions()
+	: MaterialBridgeRegistry(&FMaterialBridgeRegistry::Get())
+{
+}
+
 void FMaterialBridge::VisitMaterials(FConstDataView InDataView, TFunctionRef<void(const FVisitedContext&)> InFunc, const FVisitOptions& InOptions) const
 {
 	if (InDataView.IsValidFor(BridgedType))
diff --git a/Engine/Plugins/VirtualProduction/Avalanche/Source/AvalancheMaterial/Private/AvaMaterialModule.cpp b/Engine/Plugins/VirtualProduction/Avalanche/Source/AvalancheMaterial/Private/AvaMaterialModule.cpp
--- a/Engine/Plugins/VirtualProduction/Avalanche/Source/AvalancheMaterial/Private/AvaMaterialModule.cpp
+++ b/Engine/Plugins/VirtualProduction/Avalanche/Source/AvalancheMaterial/Private/AvaMaterialModule.cpp
@@ -3,6 +3,7 @@
 #include "AvaMaterialModule.h"
 #include "AvaActorMaterialBridge.h"
 #include "AvaDecalComponentMaterialBridge.h"
+#include "AvaLevelMaterialBridge.h"
 #include "AvaMaterialBridge.h"
 #include "AvaMaterialBridgeRegistry.h"
 #include "AvaPrimitiveComponentMaterialBridge.h"
@@ -16,6 +17,7 @@ void FAvaMaterialModule::StartupModule()
 	UE::Ava::FMaterialBridgeRegistry& MaterialBridgeRegistry = UE::Ava::FMaterialBridgeRegistry::GetMutable();
 
 	// Register built-in material bridges
+	MaterialBridgeRegistry.Register<UE::Ava::FLevelMaterialBridge>();
 	MaterialBridgeRegistry.Register<UE::Ava::FActorMaterialBridge>();
 	MaterialBridgeRegistry.Register<UE::Ava::FPrimitiveComponentMaterialBridge>();
 	MaterialBridgeRegistry.Register<UE::Ava::FDecalComponentMaterialBridge>();
diff --git a/Engine/Plugins/VirtualProduction/Avalanche/Source/AvalancheMaterial/Public/AvaLevelMaterialBridge.h b/Engine/Plugins/VirtualProduction/Avalanche/Source/AvalancheMaterial/Public/AvaLevelMaterialBridge.h
new file mode 100644
--- /dev/null
+++ b/Engine/Plugins/VirtualProduction/Avalanche/Source/AvalancheMaterial/Public/AvaLevelMaterialBridge.h
@@ -0,0 +1,24 @@
+// Copyright Epic Games, Inc. All Rights Reserved.
+
+#pragma once
+
+#include "AvaMaterialBridge.h"
+
+#define UE_API AVALANCHEMATERIAL_API
+
+namespace UE::Ava
+{
+
+/** Material Bridge for Levels */
+class FLevelMaterialBridge : public FMaterialBridge
+{
+protected:
+	//~ Begin FMaterialBridge
+	UE_API virtual const UStruct* OnGetBridgedType() const override;
+	UE_API virtual void OnVisitMaterials(FConstDataView InDataView, TFunctionRef<void(const FVisitedContext&)> InFunc, const FVisitOptions& InOptions) const override;
+	//~ End FMaterialBridge
+};
+
+} // UE::Ava
+
+#undef UE_API
diff --git a/Engine/Plugins/VirtualProduction/Avalanche/Source/AvalancheMaterial/Public/AvaMaterialBridge.h b/Engine/Plugins/VirtualProduction/Avalanche/Source/AvalancheMaterial/Public/AvaMaterialBridge.h
--- a/Engine/Plugins/VirtualProduction/Avalanche/Source/AvalancheMaterial/Public/AvaMaterialBridge.h
+++ b/Engine/Plugins/VirtualProduction/Avalanche/Source/AvalancheMaterial/Public/AvaMaterialBridge.h
@@ -13,7 +13,12 @@ class UMaterialInterface;
 
 namespace UE::Ava
 {
+	class FMaterialBridgeRegistry;
 
+} // UE::Ava
+
+namespace UE::Ava
+{
 /**
  * Base class that allows interacting with materials from any type of data view (Actors, ActorComponents, etc).
  */
@@ -37,8 +42,9 @@ public:
 	};
 	struct FVisitOptions
 	{
-		/** Declare default ctor explicitly due to https://gcc.gnu.org/bugzilla/show_bug.cgi?id=88165 */
-		FVisitOptions() {};
+		UE_API FVisitOptions();
+		/** The registry to use to visit any materials of sub objects. Defaults to the global registry. */
+		TNotNull<const FMaterialBridgeRegistry*> MaterialBridgeRegistry; 
 		/** Whether to include debug materials */
 		bool bVisitDebugMaterials = false;
 	};
