[AvalancheTests] Automate 'Sequence - Creating and deleting new Sequence'

#rnx

#virtualized
#Synced-CL 45465245

diff --git a/Engine/Plugins/VirtualProduction/Avalanche/Source/AvalancheSequencer/Private/SequenceTree/Widgets/SAvaSequenceTree.h b/Engine/Plugins/VirtualProduction/Avalanche/Source/AvalancheSequencer/Private/SequenceTree/Widgets/SAvaSequenceTree.h
--- a/Engine/Plugins/VirtualProduction/Avalanche/Source/AvalancheSequencer/Private/SequenceTree/Widgets/SAvaSequenceTree.h
+++ b/Engine/Plugins/VirtualProduction/Avalanche/Source/AvalancheSequencer/Private/SequenceTree/Widgets/SAvaSequenceTree.h
@@ -2,7 +2,7 @@
 
 #pragma once
 
-#include "AvaSequenceItemShared.h"
+#include "AvalancheSequencer/Private/AvaSequenceItemShared.h"
 #include "Templates/SharedPointer.h"
 #include "Widgets/SCompoundWidget.h"
 
@@ -28,7 +28,7 @@ public:
 
 	TSharedPtr<STreeView<FAvaSequenceItemPtr>> GetSequenceTreeView() const { return SequenceTreeView; }
 	
-	TSharedPtr<SWidget> OnContextMenuOpening() const;
+	AVALANCHESEQUENCER_API TSharedPtr<SWidget> OnContextMenuOpening() const;
 
 	void BindCommands(const TSharedRef<FAvaSequencer>& InSequencer);
 	
