Motion Design: added option for specific types to be cached instead of the generic UMaterialInterface::CacheShaders. This makes it much more performant, but runs the risk of missing a shader if not specified. LogShaders Verbose and LogAvaRundown Verbose can be used to detect which shaders are getting compiled after a page transition starts.

#Synced-CL 48399641

diff --git a/Engine/Plugins/VirtualProduction/Avalanche/Source/AvalancheMaterial/AvalancheMaterial.Build.cs b/Engine/Plugins/VirtualProduction/Avalanche/Source/AvalancheMaterial/AvalancheMaterial.Build.cs
--- a/Engine/Plugins/VirtualProduction/Avalanche/Source/AvalancheMaterial/AvalancheMaterial.Build.cs
+++ b/Engine/Plugins/VirtualProduction/Avalanche/Source/AvalancheMaterial/AvalancheMaterial.Build.cs
@@ -14,6 +14,7 @@ public class AvalancheMaterial : ModuleRules
                 "AvalancheCore",
                 "Core",
                 "CoreUObject",
+                "DeveloperSettings",
                 "Engine",
             }
         );
@@ -21,6 +22,7 @@ public class AvalancheMaterial : ModuleRules
         PrivateDependencyModuleNames.AddRange(
             new string[]
             {
+                "RenderCore",
                 "RHI",
             }
         );
diff --git a/Engine/Plugins/VirtualProduction/Avalanche/Source/AvalancheMaterial/Internal/AvaMaterialCacheHelper.h b/Engine/Plugins/VirtualProduction/Avalanche/Source/AvalancheMaterial/Internal/AvaMaterialCacheHelper.h
--- a/Engine/Plugins/VirtualProduction/Avalanche/Source/AvalancheMaterial/Internal/AvaMaterialCacheHelper.h
+++ b/Engine/Plugins/VirtualProduction/Avalanche/Source/AvalancheMaterial/Internal/AvaMaterialCacheHelper.h
@@ -8,8 +8,15 @@
 
 #define UE_API AVALANCHEMATERIAL_API
 
+class FMaterialUpdateContext;
+class UAvaMaterialCacheSettings;
 class UMaterialInterface;
 
+namespace UE::Ava
+{
+	struct FResolvedShaderProfile;
+}
+
 namespace UE::Ava
 {
 
@@ -24,10 +31,27 @@ class FMaterialCacheHelper : public TSharedFromThis<FMaterialCacheHelper>, publi
 		Caching,
 	};
 
+	struct FMaterialKey
+	{
+		/** Object key to the Material being cached */
+		TObjectKey<UMaterialInterface> Material;
+		/** Shader profile to use for caching the material shaders */
+		FName ShaderProfile;
+
+		bool operator==(const FMaterialKey& InOther) const
+		{
+			return Material == InOther.Material && ShaderProfile == InOther.ShaderProfile;
+		}
+		friend uint32 GetTypeHash(const FMaterialKey& InKey)
+		{
+			return HashCombineFast(GetTypeHash(InKey.Material), GetTypeHash(InKey.ShaderProfile));
+		}
+	};
+
 	struct FObjectData
 	{
-		/** All the materials belonging to an object that are still incomplete */
-		TSet<TObjectKey<UMaterialInterface>> IncompleteMaterials;
+		/** All the keys to the materials belonging to an object that are still incomplete */
+		TSet<FMaterialKey> IncompleteMaterialKeys;
 	};
 
 	struct FMaterialData
@@ -43,10 +67,11 @@ public:
 
 	/**
 	 * Caches the materials in the given object.
-	 * The materials are found via the material bridge api.
-	 * Returns true if materials were gathered for the object, false otherwise. (e.g. if object's materials were already gathered, or no material bridge found)
+	 * @param InObject the object to gather materials for via the material bridge api
+	 * @param InShaderProfile the shader profile to use for caching shaders. If none, it will cache via UMaterial::CacheShaders
+	 * @return true if materials were gathered for the object, false otherwise. (e.g. if object's materials were already gathered, or no material bridge found)
 	 */
-	UE_API bool RequestCacheMaterials(const UObject* InObject);
+	UE_API bool RequestCacheMaterials(const UObject* InObject, FName InShaderProfile = NAME_None);
 
 	/** Processes materials pending completion with a limited time budget */
 	UE_API void Tick();
@@ -81,6 +106,15 @@ private:
 	 */
 	void ProcessMaterials();
 
+	/** Finds the shader profile matching the given name */
+	const FResolvedShaderProfile* FindShaderProfile(FName InShaderProfile) const;
+
+	/** Caches the given material's shaders with the given shader profile. If not found, it will use the UMaterial::CacheShaders for compilation */
+	void CacheMaterial(TNotNull<UMaterialInterface*> InMaterial, FName InShaderProfile);
+
+	/** Determines whether the given material is complete */
+	bool IsMaterialComplete(TNotNull<UMaterialInterface*> InMaterial, FName InShaderProfile) const;
+
 	/** Removes the complete materials from the objects' IncompleteMaterials set */
 	void RemoveCompletedMaterials();
 
@@ -91,19 +125,31 @@ private:
 	 * Gathers the material data for an object, or skips it if the object's materials have been gathered before.
 	 * returns true if a material bridge was found for the object and visited its materials
 	 */
-	bool GatherMaterialData(const UObject* InObject);
+	bool GatherMaterialData(const UObject* InObject, FName InShaderProfile);
 
 	/** Objects that have already been gathered, mapped to the incomplete materials that were gathered for the object */
 	TMap<FObjectKey, FObjectData> Objects;
 
 	/** All the materials that have already been gathered and are pending completion. Once materials are complete, they're removed from this map */
-	TMap<TObjectKey<UMaterialInterface>, FMaterialData> Materials;
+	TMap<FMaterialKey, FMaterialData> Materials;
+
+	/** Time when BeginFrame was called */
+	double BeginFrameTime = 0.0;
 
 	/** Current time budget left for the tick. Resets every tick */
 	float RemainingTimeBudget = 0.f;
 
+	/** Number of frames where Tick was skipped because the frame time was past max frame time */
+	int32 FramesStarved = 0;
+
+	/** Handle to the BeginFrame delegate to keep track of BeginFrameTime */
+	FDelegateHandle BeginFrameHandle;
+
 	/** Handle to the delegate for ticking */
-	FDelegateHandle TickHandle;
+	FDelegateHandle EndFrameHandle;
+
+	/** Reset material cache */
+	bool bResetRequested = false;
 };
 
 } // UE::Ava
diff --git a/Engine/Plugins/VirtualProduction/Avalanche/Source/AvalancheMaterial/Internal/AvaMaterialCacheLog.h b/Engine/Plugins/VirtualProduction/Avalanche/Source/AvalancheMaterial/Internal/AvaMaterialCacheLog.h
new file mode 100644
--- /dev/null
+++ b/Engine/Plugins/VirtualProduction/Avalanche/Source/AvalancheMaterial/Internal/AvaMaterialCacheLog.h
@@ -0,0 +1,7 @@
+// Copyright Epic Games, Inc. All Rights Reserved.
+
+#pragma once
+
+#include "Logging/LogCategory.h"
+
+DECLARE_LOG_CATEGORY_EXTERN(LogAvaMaterialCache, Log, All);
diff --git a/Engine/Plugins/VirtualProduction/Avalanche/Source/AvalancheMaterial/Internal/AvaMaterialCacheSettings.h b/Engine/Plugins/VirtualProduction/Avalanche/Source/AvalancheMaterial/Internal/AvaMaterialCacheSettings.h
new file mode 100644
--- /dev/null
+++ b/Engine/Plugins/VirtualProduction/Avalanche/Source/AvalancheMaterial/Internal/AvaMaterialCacheSettings.h
@@ -0,0 +1,109 @@
+// Copyright Epic Games, Inc. All Rights Reserved.
+
+#pragma once
+
+#include "Engine/DeveloperSettings.h"
+#include "AvaMaterialCacheSettings.generated.h"
+
+#define UE_API AVALANCHEMATERIAL_API
+
+class FShaderPipelineType;
+class FShaderType;
+class FVertexFactoryType;
+
+namespace UE::Ava
+{
+	struct FResolvedShaderProfile
+	{
+		TArray<const FVertexFactoryType*> VertexFactoryTypes;
+		TArray<const FShaderPipelineType*> PipelineTypes;
+		TArray<const FShaderType*> ShaderTypes;
+	};
+} // UE::Ava
+
+USTRUCT()
+struct FAvaShaderCombination
+{
+	GENERATED_BODY()
+
+	FAvaShaderCombination() = default;
+
+	explicit FAvaShaderCombination(FName InVertexFactorType, TConstArrayView<FName> InPipelineTypes, TConstArrayView<FName> InShaderTypes)
+		: VertexFactoryType(InVertexFactorType)
+		, PipelineTypes(InPipelineTypes)
+		, ShaderTypes(InShaderTypes)
+	{
+	}
+
+	UPROPERTY(EditAnywhere, Category="Material Cache")
+	FName VertexFactoryType;
+
+	UPROPERTY(EditAnywhere, Category="Material Cache")
+	TArray<FName> PipelineTypes;
+
+	UPROPERTY(EditAnywhere, Category="Material Cache")
+	TArray<FName> ShaderTypes;
+};
+
+USTRUCT()
+struct FAvaShaderProfile
+{
+	GENERATED_BODY()
+
+	UPROPERTY(EditAnywhere, Category="Material Cache")
+	FName ProfileName;
+
+	UPROPERTY(EditAnywhere, Category="Material Cache")
+	bool bEnabled = true;
+
+	UPROPERTY(EditAnywhere, Category="Material Cache", meta=(EditCondition="bEnabled"))
+	TArray<FAvaShaderCombination> ShaderCombinations;
+};
+
+UCLASS(MinimalAPI, config=Engine, meta=(DisplayName="Material Cache Settings"))
+class UAvaMaterialCacheSettings : public UDeveloperSettings
+{
+	GENERATED_BODY()
+
+public:
+	UE_API static const FLazyName DefaultRealtimeProfile;
+	UE_API static const FLazyName DefaultOfflineProfile;
+
+	UAvaMaterialCacheSettings();
+
+	/** Finds the resolved shader profile for a given profile name */
+	const UE::Ava::FResolvedShaderProfile* FindResolvedShaderProfile(FName InProfileName) const;
+
+protected:
+	//~ Begin UObject
+	virtual void PostInitProperties() override;
+	virtual void PostReloadConfig(FProperty* InPropertyThatWasLoaded) override;
+	//~ End UObject
+
+private:
+	/** Add the built-in shader profiles */
+	void RegisterDefaultShaderProfile();
+
+	/** Resolves all the registered shader profiles to its type pointers */
+	void ResolveShaderProfiles();
+
+	/** Resolves a single shader profile to its type pointers */
+	void ResolveShaderProfile(const FAvaShaderProfile& InShaderProfile, TMap<FName, FHashedName>& InOutHashedNames);
+
+	/** Profile to use in realtime (e.g. when playing levels) */
+	UPROPERTY(Config, EditAnywhere, Category="Material Cache")
+	FName RealtimeProfile;
+
+	/** Profile to use offline (e.g. when preloading levels) */
+	UPROPERTY(Config, EditAnywhere, Category="Material Cache")
+	FName OfflineProfile;
+
+	/** Shader profiles that can be saved to disk */
+	UPROPERTY(Config, EditAnywhere, Category="Material Cache")
+	TArray<FAvaShaderProfile> ShaderProfiles;
+
+	/** The saved shader profiles turned into resolved pointers */
+	TMap<FName, UE::Ava::FResolvedShaderProfile> ResolvedShaderProfiles;
+};
+
+#undef UE_API
diff --git a/Engine/Plugins/VirtualProduction/Avalanche/Source/AvalancheMaterial/Private/AvaMaterialCacheHelper.cpp b/Engine/Plugins/VirtualProduction/Avalanche/Source/AvalancheMaterial/Private/AvaMaterialCacheHelper.cpp
--- a/Engine/Plugins/VirtualProduction/Avalanche/Source/AvalancheMaterial/Private/AvaMaterialCacheHelper.cpp
+++ b/Engine/Plugins/VirtualProduction/Avalanche/Source/AvalancheMaterial/Private/AvaMaterialCacheHelper.cpp
@@ -4,11 +4,14 @@
 #include "AvaDataView.h"
 #include "AvaMaterialBridge.h"
 #include "AvaMaterialBridgeRegistry.h"
+#include "AvaMaterialCacheLog.h"
+#include "AvaMaterialCacheSettings.h"
 #include "MaterialShared.h"
 #include "Materials/Material.h"
 #include "ShaderCompiler.h"
+#include "UnrealEngine.h"
 
-DEFINE_LOG_CATEGORY_STATIC(LogAvaMaterialCache, Log, All);
+DEFINE_LOG_CATEGORY(LogAvaMaterialCache);
 
 namespace UE::Ava
 {
@@ -23,15 +26,42 @@ FAutoConsoleCommand MaterialCacheDumpMaterials(TEXT("MotionDesign.MaterialCache.
 			FMaterialCacheHelper::Get().DumpMaterials();
 		}));
 
-void CacheMaterial(FMaterialUpdateContext& InContext, TNotNull<UMaterialInterface*> InMaterial)
+TAutoConsoleVariable<float> CVarMaterialCacheFrameTimeStarveThreshold(TEXT("MotionDesign.MaterialCache.FrameTimeStarveThreshold")
+	, 1.f / 10.f
+	, TEXT("The frame time threshold for starving tick. If the current duration of the frame surpasses this, FMaterialCacheHelper::Tick will be skipped and starve.")
+	);
+
+TAutoConsoleVariable<int32> CVarMaterialCacheTickMaxStarveFrames(TEXT("MotionDesign.MaterialCache.TickMaxStarveFrames")
+	, 8
+	, TEXT("The maximum number of consecutive frames that FMaterialCacheHelper::Tick can starve for. After this, tick will be forced to execute.")
+	);
+
+TAutoConsoleVariable<float> CVarMaterialCacheTickTimeBudget(TEXT("MotionDesign.MaterialCache.TickTimeBudget")
+	, 0.002f
+	, TEXT("The time budget FMaterialCacheHelper::Tick has to process its pending elements.")
+	);
+
+#if WITH_EDITOR
+void ForEachMaterialResource(TNotNull<UMaterialInterface*> InMaterial, TFunctionRef<bool(FMaterialResource*)> InFunc)
 {
-	constexpr EMaterialShaderPrecompileMode PrecompileMode = EMaterialShaderPrecompileMode::Background;
+	const EMaterialQualityLevel::Type ActiveQualityLevel = GetCachedScalabilityCVars().MaterialQualityLevel;
+	uint32 FeatureLevelsToCompile = UMaterial::GetFeatureLevelsToCompileForAllMaterials();
 
-	// This is needed because CacheShaders blindly recreates uniform buffers
-	// which can only be done if the draw command is going to be re-cached.
-	InContext.AddMaterialInterface(InMaterial);
-	InMaterial->CacheShaders(PrecompileMode);
+	while (FeatureLevelsToCompile != 0)
+	{
+		const ERHIFeatureLevel::Type FeatureLevel = static_cast<ERHIFeatureLevel::Type>(FBitSet::GetAndClearNextBit(FeatureLevelsToCompile));
+		const EShaderPlatform ShaderPlatform = GShaderPlatformForFeatureLevel[FeatureLevel];
+
+		// Only cache shaders for the quality level that will actually be used to render
+		// In cooked build, there is no shader compilation but this is still needed to register the loaded shadermap
+		FMaterialResource* CurrentResource = InMaterial->GetMaterialResource(ShaderPlatform, ActiveQualityLevel);
+		if (ensure(CurrentResource) && !InFunc(CurrentResource))
+		{
+			break;
+		}
+	}
 }
+#endif
 
 } // UE::Ava::Private
 
@@ -41,24 +71,49 @@ FMaterialCacheHelper& FMaterialCacheHelper::Get()
 	return *MaterialCacheHelper;
 }
 
-bool FMaterialCacheHelper::RequestCacheMaterials(const UObject* InObject)
+bool FMaterialCacheHelper::RequestCacheMaterials(const UObject* InObject, FName InShaderProfile)
 {
 	if (!InObject)
 	{
 		return false;
 	}
 
-	const bool bMaterialsGathered = GatherMaterialData(InObject);
+	if (!GatherMaterialData(InObject, InShaderProfile))
+	{
+		return false;
+	}
+
 	RegisterTick();
-	return bMaterialsGathered;
+	return true;
 }
 
 void FMaterialCacheHelper::Tick()
 {
+	if (bResetRequested)
+	{
+		bResetRequested = false;
+		Reset();
+		return;
+	}
+
 	TRACE_CPUPROFILER_EVENT_SCOPE(FMaterialCacheHelper::Tick);
 
-	constexpr float MaxTimeBudget = 0.001f;
-	RemainingTimeBudget = MaxTimeBudget;
+	// Duration of current frame
+	const float CurrentFrameTime = (FPlatformTime::Seconds() - BeginFrameTime);
+
+	// If the current frame time is already past the threshold, tick should be skipped unless it has been starving for a while already.
+	if (CurrentFrameTime >= Private::CVarMaterialCacheFrameTimeStarveThreshold.GetValueOnGameThread())
+	{
+		if (FramesStarved < Private::CVarMaterialCacheTickMaxStarveFrames.GetValueOnGameThread())
+		{
+			++FramesStarved;
+			return;
+		}
+		UE_LOG(LogAvaMaterialCache, Verbose, TEXT("FMaterialCacheHelper::Tick was starved for %d frames. Forcing tick..."), FramesStarved);
+	}
+
+	FramesStarved = 0;
+	RemainingTimeBudget = Private::CVarMaterialCacheTickTimeBudget.GetValueOnGameThread();
 
 	ProcessMaterials();
 
@@ -69,7 +124,7 @@ void FMaterialCacheHelper::Tick()
 	}
 	else
 	{
-		Reset();
+		bResetRequested = true;
 	}
 }
 
@@ -91,9 +146,9 @@ bool FMaterialCacheHelper::IsCaching(const UObject* InObject) const
 		return false;
 	}
 
-	for (TObjectKey<UMaterialInterface> Material : ObjectData->IncompleteMaterials)
+	for (const FMaterialKey& MaterialKey : ObjectData->IncompleteMaterialKeys)
 	{
-		if (Materials.Contains(Material))
+		if (Materials.Contains(MaterialKey))
 		{
 			return true;
 		}
@@ -110,7 +165,7 @@ void FMaterialCacheHelper::DumpMaterials()
 		};
 
 	UE_LOG(LogAvaMaterialCache, Log, TEXT("----- DUMP MATERIALS--------------------------------------------"));
-	for (const TPair<TObjectKey<UMaterialInterface>, FMaterialData>& Pair : Materials)
+	for (const TPair<FMaterialKey, FMaterialData>& Pair : Materials)
 	{
 		const FMaterialData& MaterialData = Pair.Value;
 
@@ -123,7 +178,7 @@ void FMaterialCacheHelper::DumpMaterials()
 
 void FMaterialCacheHelper::AddReferencedObjects(FReferenceCollector& InCollector)
 {
-	for (TPair<TObjectKey<UMaterialInterface>, FMaterialData>& Pair : Materials)
+	for (TPair<FMaterialKey, FMaterialData>& Pair : Materials)
 	{
 		InCollector.AddReferencedObject(Pair.Value.Material);
 	}
@@ -143,18 +198,32 @@ void FMaterialCacheHelper::Reset()
 
 void FMaterialCacheHelper::RegisterTick()
 {
-	if (!TickHandle.IsValid())
+	if (!BeginFrameHandle.IsValid())
+	{
+		BeginFrameHandle = FCoreDelegates::OnBeginFrame.AddSPLambda(this, [this]
+			{
+				BeginFrameTime = FPlatformTime::Seconds();
+			});
+	}
+
+	if (!EndFrameHandle.IsValid())
 	{
-		TickHandle = FCoreDelegates::OnEndFrame.AddSP(this, &FMaterialCacheHelper::Tick);
+		EndFrameHandle = FCoreDelegates::OnEndFrame.AddSP(this, &FMaterialCacheHelper::Tick);
 	}
 }
 
 void FMaterialCacheHelper::UnregisterTick()
 {
-	if (TickHandle.IsValid())
+	if (BeginFrameHandle.IsValid())
 	{
-		FCoreDelegates::OnEndFrame.Remove(TickHandle);
-		TickHandle.Reset();
+		FCoreDelegates::OnBeginFrame.Remove(BeginFrameHandle);
+		BeginFrameHandle.Reset();
+	}
+
+	if (EndFrameHandle.IsValid())
+	{
+		FCoreDelegates::OnEndFrame.Remove(EndFrameHandle);
+		EndFrameHandle.Reset();
 	}
 }
 
@@ -172,8 +241,9 @@ void FMaterialCacheHelper::ProcessMaterials()
 
 	TOptional<FMaterialUpdateContext> UpdateContext;
 
-	for (TMap<TObjectKey<UMaterialInterface>, FMaterialData>::TIterator Iter(Materials); Iter; ++Iter)
+	for (TMap<FMaterialKey, FMaterialData>::TIterator Iter(Materials); Iter; ++Iter)
 	{
+		const FMaterialKey& MaterialKey = Iter.Key();
 		FMaterialData& MaterialData = Iter.Value();
 
 		if (!MaterialData.Material)
@@ -189,12 +259,13 @@ void FMaterialCacheHelper::ProcessMaterials()
 			{
 				UpdateContext.Emplace(0);
 			}
-			Private::CacheMaterial(*UpdateContext, MaterialData.Material);
+			UpdateContext->AddMaterialInterface(MaterialData.Material);
+			CacheMaterial(MaterialData.Material, MaterialKey.ShaderProfile);
 			MaterialData.CacheStatus = EMaterialCacheStatus::Caching;
 			break;
 
 		case EMaterialCacheStatus::Caching:
-			if (MaterialData.Material->IsComplete())
+			if (IsMaterialComplete(MaterialData.Material, MaterialKey.ShaderProfile))
 			{
 				Iter.RemoveCurrent();
 			}
@@ -211,6 +282,62 @@ void FMaterialCacheHelper::ProcessMaterials()
 	RemainingTimeBudget = EndBudgetTime - FPlatformTime::Seconds();
 }
 
+const FResolvedShaderProfile* FMaterialCacheHelper::FindShaderProfile(FName InShaderProfile) const
+{
+	if (const UAvaMaterialCacheSettings* MaterialCacheSettings = GetDefault<UAvaMaterialCacheSettings>())
+	{
+		return MaterialCacheSettings->FindResolvedShaderProfile(InShaderProfile);
+	}
+	return nullptr;
+}
+
+void FMaterialCacheHelper::CacheMaterial(TNotNull<UMaterialInterface*> InMaterial, FName InShaderProfile)
+{
+#if WITH_EDITOR
+	// CacheGivenTypes is EditorOnly
+	if (const FResolvedShaderProfile* const ShaderProfile = FindShaderProfile(InShaderProfile))
+	{
+		InMaterial->CacheShaders(EMaterialShaderPrecompileMode::None);
+
+		Private::ForEachMaterialResource(InMaterial,
+			[ShaderProfile](FMaterialResource* InMaterialResource)->bool
+			{
+				
+				InMaterialResource->CacheGivenTypes(ShaderProfile->VertexFactoryTypes, ShaderProfile->PipelineTypes, ShaderProfile->ShaderTypes);
+				return true; // continue
+			});
+	}
+	else
+#endif
+	{
+		InMaterial->CacheShaders(EMaterialShaderPrecompileMode::Background);
+	}
+}
+
+bool FMaterialCacheHelper::IsMaterialComplete(TNotNull<UMaterialInterface*> InMaterial, FName InShaderProfile) const
+{
+#if WITH_EDITOR
+	if (const FResolvedShaderProfile* const ShaderProfile = FindShaderProfile(InShaderProfile))
+	{
+		bool bIsComplete = true;
+
+		Private::ForEachMaterialResource(InMaterial, 
+			[&bIsComplete](FMaterialResource* InMaterialResource)->bool
+			{
+				if (InMaterialResource->IsCachingShaders())
+				{
+					bIsComplete = false;
+					return false; // break
+				}
+				return true; // continue
+			});
+
+		return bIsComplete;
+	}
+#endif
+	return InMaterial->IsComplete();
+}
+
 void FMaterialCacheHelper::RemoveCompletedMaterials()
 {
 	if (RemainingTimeBudget <= 0.f)
@@ -222,11 +349,11 @@ void FMaterialCacheHelper::RemoveCompletedMaterials()
 
 	for (TPair<FObjectKey, FObjectData>& Object : Objects)
 	{
-		for (TSet<TObjectKey<UMaterialInterface>>::TIterator MaterialIter(Object.Value.IncompleteMaterials); MaterialIter; ++MaterialIter)
+		for (TSet<FMaterialKey>::TIterator MaterialKeyIter(Object.Value.IncompleteMaterialKeys); MaterialKeyIter; ++MaterialKeyIter)
 		{
-			if (!Materials.Contains(*MaterialIter))
+			if (!Materials.Contains(*MaterialKeyIter))
 			{
-				MaterialIter.RemoveCurrent();
+				MaterialKeyIter.RemoveCurrent();
 			}
 		}
 	}
@@ -251,7 +378,7 @@ void FMaterialCacheHelper::ProcessShadersCompiling()
 	RemainingTimeBudget -= (FPlatformTime::Seconds() - StartTime);
 }
 
-bool FMaterialCacheHelper::GatherMaterialData(const UObject* InObject)
+bool FMaterialCacheHelper::GatherMaterialData(const UObject* InObject, FName InShaderProfile)
 {
 	if (!InObject || Objects.Contains(InObject))
 	{
@@ -270,16 +397,20 @@ bool FMaterialCacheHelper::GatherMaterialData(const UObject* InObject)
 	FObjectData& ObjectData = Objects.Add(InObject);
 
 	auto AddIncompleteMaterial =
-		[&Materials = Materials, &ObjectData](TNotNull<UMaterialInterface*> InMaterial)
+		[&Materials = Materials, &ObjectData, InShaderProfile](TNotNull<UMaterialInterface*> InMaterial)
 		{
-			FMaterialData* MaterialData = Materials.Find(InMaterial);
+			FMaterialKey MaterialKey;
+			MaterialKey.Material = InMaterial;
+			MaterialKey.ShaderProfile = InShaderProfile;
+
+			FMaterialData* MaterialData = Materials.Find(MaterialKey);
 			if (!MaterialData)
 			{
-				MaterialData = &Materials.Add(InMaterial);
+				MaterialData = &Materials.Add(MaterialKey);
 				MaterialData->Material = InMaterial;
 				MaterialData->CacheStatus = EMaterialCacheStatus::Waiting;
 			}
-			ObjectData.IncompleteMaterials.Add(InMaterial);
+			ObjectData.IncompleteMaterialKeys.Add(MaterialKey);
 		};
 
 	// Gather all materials for object
diff --git a/Engine/Plugins/VirtualProduction/Avalanche/Source/AvalancheMaterial/Private/AvaMaterialCacheSettings.cpp b/Engine/Plugins/VirtualProduction/Avalanche/Source/AvalancheMaterial/Private/AvaMaterialCacheSettings.cpp
new file mode 100644
--- /dev/null
+++ b/Engine/Plugins/VirtualProduction/Avalanche/Source/AvalancheMaterial/Private/AvaMaterialCacheSettings.cpp
@@ -0,0 +1,216 @@
+// Copyright Epic Games, Inc. All Rights Reserved.
+
+#include "AvaMaterialCacheSettings.h"
+#include "AvaMaterialCacheLog.h"
+#include "VertexFactory.h"
+
+const FLazyName UAvaMaterialCacheSettings::DefaultRealtimeProfile(TEXT("BasicSceneRealtime"));
+const FLazyName UAvaMaterialCacheSettings::DefaultOfflineProfile(TEXT("BasicSceneOffline"));
+
+namespace UE::Ava::Private
+{
+
+TArray<FAvaShaderCombination> MakeDefaultShaderCombinations()
+{
+	// Vertex Factories
+	const FName LocalVertexFactory       = TEXT("FLocalVertexFactory");
+	const FName NiagaraMeshVertexFactory = TEXT("FNiagaraMeshVertexFactory");
+
+	// Pipelines
+	const FName StandardVelocityPipeline = TEXT("StandardVelocityPipeline");
+	const FName DepthPipeline            = TEXT("DepthPipeline");
+
+	// Shaders
+	const FName BasePassVS               = TEXT("TBasePassVSFNoLightMapPolicy");
+	const FName BasePassPS               = TEXT("TBasePassPSFNoLightMapPolicy");
+	const FName BasePassPSSL             = TEXT("TBasePassPSFNoLightMapPolicySkylight");
+	const FName NiagaraShader            = TEXT("FNiagaraShader");
+	const FName MaterialCHS              = TEXT("TMaterialCHSFNoLightMapPolicyFAnyHitShader");
+	const FName LumenCardVS              = TEXT("FLumenCardVS");
+	const FName LumenCardPS              = TEXT("FLumenCardPS");
+	const FName VelocityVS               = TEXT("TVelocityVS<EVelocityPassMode::Velocity_Standard>");
+	const FName VelocityPS               = TEXT("TVelocityPS<EVelocityPassMode::Velocity_Standard>");
+	const FName DepthOnlyVS              = TEXT("TDepthOnlyVS<false>");
+	const FName DepthOnlyPS              = TEXT("FDepthOnlyPS");
+	const FName ShadowDepthVS            = TEXT("TShadowDepthVSCombined");
+	const FName ShadowDepthPS            = TEXT("TShadowDepthPSCombined");
+
+	return 
+	{
+		FAvaShaderCombination(LocalVertexFactory,
+			{
+				StandardVelocityPipeline,
+				DepthPipeline,
+			},
+			{
+				BasePassVS,
+				BasePassPS,
+				BasePassPSSL,
+				NiagaraShader,
+				MaterialCHS,
+				LumenCardVS,
+				LumenCardPS,
+				VelocityVS,
+				VelocityPS,
+				DepthOnlyVS,
+				DepthOnlyPS,
+				ShadowDepthVS,
+				ShadowDepthPS,
+			}
+		),
+		FAvaShaderCombination(NiagaraMeshVertexFactory,
+			{
+				StandardVelocityPipeline
+			},
+			{
+				BasePassVS,
+				BasePassPS,
+				ShadowDepthVS,
+				ShadowDepthPS,
+				VelocityVS,
+				VelocityPS,
+			}
+		)
+	};
+}
+
+} // UE::Ava::Private
+
+UAvaMaterialCacheSettings::UAvaMaterialCacheSettings()
+{
+	RealtimeProfile = UAvaMaterialCacheSettings::DefaultRealtimeProfile;
+	OfflineProfile = UAvaMaterialCacheSettings::DefaultOfflineProfile;
+
+	CategoryName = TEXT("Motion Design");
+	SectionName  = TEXT("Material Cache");
+	RegisterDefaultShaderProfile();
+}
+
+const UE::Ava::FResolvedShaderProfile* UAvaMaterialCacheSettings::FindResolvedShaderProfile(FName InProfileName) const
+{
+	return ResolvedShaderProfiles.Find(InProfileName);
+}
+
+void UAvaMaterialCacheSettings::PostInitProperties()
+{
+	Super::PostInitProperties();
+	ResolveShaderProfiles();
+}
+
+void UAvaMaterialCacheSettings::PostReloadConfig(FProperty* InPropertyThatWasLoaded)
+{
+	Super::PostReloadConfig(InPropertyThatWasLoaded);
+	ResolveShaderProfiles();
+}
+
+void UAvaMaterialCacheSettings::RegisterDefaultShaderProfile()
+{
+	const TArray<FAvaShaderCombination> DefaultShaderCombinations = UE::Ava::Private::MakeDefaultShaderCombinations();
+
+	ShaderProfiles.Add(
+		{ 
+			.ProfileName = UAvaMaterialCacheSettings::DefaultRealtimeProfile,
+			.bEnabled = true,
+			.ShaderCombinations = DefaultShaderCombinations,
+		});
+
+	ShaderProfiles.Add(
+		{ 
+			.ProfileName = UAvaMaterialCacheSettings::DefaultOfflineProfile,
+			.bEnabled = true,
+			.ShaderCombinations = DefaultShaderCombinations,
+		});
+}
+
+void UAvaMaterialCacheSettings::ResolveShaderProfiles()
+{
+	ResolvedShaderProfiles.Empty(ShaderProfiles.Num());
+
+	TMap<FName, FHashedName> HashedNames;
+	for (const FAvaShaderProfile& ShaderProfile : ShaderProfiles)
+	{
+		ResolveShaderProfile(ShaderProfile, HashedNames);
+	}
+}
+
+void UAvaMaterialCacheSettings::ResolveShaderProfile(const FAvaShaderProfile& InShaderProfile, TMap<FName, FHashedName>& InOutHashedNames)
+{
+	if (!InShaderProfile.bEnabled)
+	{
+		UE_LOG(LogAvaMaterialCache, Verbose, TEXT("Shader Profile '%s' not resolved as it's not enabled."), *InShaderProfile.ProfileName.ToString());
+		return;
+	}
+
+	if (ResolvedShaderProfiles.Contains(InShaderProfile.ProfileName))
+	{
+		UE_LOG(LogAvaMaterialCache, Warning, TEXT("Shader Profile '%s' already exists. New entry will be ignored."), *InShaderProfile.ProfileName.ToString());
+		return;
+	}
+
+	UE::Ava::FResolvedShaderProfile& ResolvedShaderProfile = ResolvedShaderProfiles.Add(InShaderProfile.ProfileName);
+
+	// Keep track of FHashedName instances to avoid recomputation
+	auto GetHashedName = [&InOutHashedNames](FName InName)
+		{
+			if (const FHashedName* HashedName = InOutHashedNames.Find(InName))
+			{
+				return *HashedName;
+			}
+			return InOutHashedNames.Add(InName, InName);
+		};
+
+	int32 ReserveCount = 0;
+	for (const FAvaShaderCombination& ShaderCombination : InShaderProfile.ShaderCombinations)
+	{
+		ReserveCount += ShaderCombination.PipelineTypes.Num();
+		ReserveCount += ShaderCombination.ShaderTypes.Num();
+	}
+
+	// The types must all have the same count at the end
+	ResolvedShaderProfile.VertexFactoryTypes.Reserve(ReserveCount);
+	ResolvedShaderProfile.PipelineTypes.Reserve(ReserveCount);
+	ResolvedShaderProfile.ShaderTypes.Reserve(ReserveCount);
+
+	for (const FAvaShaderCombination& ShaderCombination : InShaderProfile.ShaderCombinations)
+	{
+		const FVertexFactoryType* const VertexFactoryType = FVertexFactoryType::GetVFByName(GetHashedName(ShaderCombination.VertexFactoryType));
+		if (!VertexFactoryType)
+		{
+			UE_LOG(LogAvaMaterialCache, Warning, TEXT("Invalid Shader Config. Vertex Factory Type '%s' was not found"), *ShaderCombination.VertexFactoryType.ToString());
+			continue;
+		}
+
+		for (FName PipelineTypeName : ShaderCombination.PipelineTypes)
+		{
+			const FShaderPipelineType* PipelineType = FShaderPipelineType::GetShaderPipelineTypeByName(GetHashedName(PipelineTypeName));
+			if (!PipelineType)
+			{
+				UE_LOG(LogAvaMaterialCache, Warning, TEXT("Invalid Shader Config. Pipeline Type '%s' was found"), *PipelineTypeName.ToString());
+				continue;
+			}
+
+			ResolvedShaderProfile.VertexFactoryTypes.Add(VertexFactoryType);
+			ResolvedShaderProfile.PipelineTypes.Add(PipelineType);
+			ResolvedShaderProfile.ShaderTypes.Add(nullptr);
+		}
+
+		for (FName ShaderTypeName : ShaderCombination.ShaderTypes)
+		{
+			const FShaderType* ShaderType = FindShaderTypeByName(GetHashedName(ShaderTypeName));
+			if (!ShaderType)
+			{
+				UE_LOG(LogAvaMaterialCache, Warning, TEXT("Invalid Shader Config. Shader Type '%s' was found"), *ShaderTypeName.ToString());
+				continue;
+			}
+
+			ResolvedShaderProfile.VertexFactoryTypes.Add(VertexFactoryType);
+			ResolvedShaderProfile.PipelineTypes.Add(nullptr);
+			ResolvedShaderProfile.ShaderTypes.Add(ShaderType);
+		}
+	}
+
+	ensureMsgf(ResolvedShaderProfile.VertexFactoryTypes.Num() == ResolvedShaderProfile.PipelineTypes.Num()
+		&& ResolvedShaderProfile.VertexFactoryTypes.Num() == ResolvedShaderProfile.ShaderTypes.Num()
+		, TEXT("Resolved Shader Profile '%s' unexpectedly has mismatched number of types.")
+		, *InShaderProfile.ProfileName.ToString());
+}
diff --git a/Engine/Plugins/VirtualProduction/Avalanche/Source/AvalancheMedia/Internal/Playback/AvaPlaybackUtils.h b/Engine/Plugins/VirtualProduction/Avalanche/Source/AvalancheMedia/Internal/Playback/AvaPlaybackUtils.h
--- a/Engine/Plugins/VirtualProduction/Avalanche/Source/AvalancheMedia/Internal/Playback/AvaPlaybackUtils.h
+++ b/Engine/Plugins/VirtualProduction/Avalanche/Source/AvalancheMedia/Internal/Playback/AvaPlaybackUtils.h
@@ -76,8 +76,12 @@ namespace UE::AvaPlayback::Utils
 	 */
 	UStreamableRenderAsset* StreamRenderAsset(UObject* InObject);
 
-	/** Caches the materials of the given object using the material bridge api */
-	bool CacheObjectMaterials(UObject* InObject);
+	/** 
+	 * Caches the materials of the given object using the material bridge api
+	 * @param InObject the object to gather and cache materials for
+	 * @param InShaderProfile the shader profile to use for caching
+	 */
+	bool CacheObjectMaterials(UObject* InObject, FName InShaderProfile);
 
 	/**
 	 * Utility class to perform async load of assets. 
diff --git a/Engine/Plugins/VirtualProduction/Avalanche/Source/AvalancheMedia/Private/Playable/Playables/AvaPlayableLevelStreaming.cpp b/Engine/Plugins/VirtualProduction/Avalanche/Source/AvalancheMedia/Private/Playable/Playables/AvaPlayableLevelStreaming.cpp
--- a/Engine/Plugins/VirtualProduction/Avalanche/Source/AvalancheMedia/Private/Playable/Playables/AvaPlayableLevelStreaming.cpp
+++ b/Engine/Plugins/VirtualProduction/Avalanche/Source/AvalancheMedia/Private/Playable/Playables/AvaPlayableLevelStreaming.cpp
@@ -2,6 +2,7 @@
 
 #include "Playable/Playables/AvaPlayableLevelStreaming.h"
 #include "AvaMaterialCacheHelper.h"
+#include "AvaMaterialCacheSettings.h"
 #include "AvaRemoteControlRebind.h"
 #include "AvaRemoteControlUtils.h"
 #include "AvaScene.h"
@@ -421,7 +422,7 @@ bool UAvaPlayableLevelStreaming::IsReadyToPlay() const
 	}
 
 	// Check whether the compilation has finished
-	const bool bIsCachingShaders = CacheMaterials(Level);
+	const bool bIsCachingShaders = CacheMaterials(Level, UAvaMaterialCacheSettings::DefaultRealtimeProfile);
 	return !bIsCachingShaders;
 }
 
@@ -523,27 +524,14 @@ bool UAvaPlayableLevelStreaming::LoadLevel(const TSoftObjectPtr<UWorld>& InSourc
 	return true;
 }
 
-bool UAvaPlayableLevelStreaming::CacheMaterials(TNotNull<const ULevel*> InLevel) const
+bool UAvaPlayableLevelStreaming::CacheMaterials(TNotNull<const ULevel*> InLevel, FName InShaderProfile) const
 {
 	TRACE_CPUPROFILER_EVENT_SCOPE(UAvaPlayableLevelStreaming::CacheMaterials)
 
 	UE::Ava::FMaterialCacheHelper& MaterialCacheHelper = UE::Ava::FMaterialCacheHelper::Get();
 
-	for (AActor* Actor : InLevel->Actors)
-	{
-		MaterialCacheHelper.RequestCacheMaterials(Actor);
-	}
-
-	bool bIsCaching = false;
-	for (AActor* Actor : InLevel->Actors)
-	{
-		if (MaterialCacheHelper.IsCaching(Actor))
-		{
-			bIsCaching = true;
-			break;
-		}
-	}
-	return bIsCaching;
+	MaterialCacheHelper.RequestCacheMaterials(InLevel, InShaderProfile);
+	return MaterialCacheHelper.IsCaching(InLevel);
 }
 
 void UAvaPlayableLevelStreaming::HandleTransitionEvent(UAvaPlayable* InPlayable, UAvaPlayableTransition* InTransition, EAvaPlayableTransitionEventFlags InTransitionFlags)
@@ -660,7 +648,8 @@ void UAvaPlayableLevelStreaming::OnLevelStreamingStateChanged_Synchronized(ELeve
 			// If there was more than one material that required caching,
 			// kick off a ticker that incrementally caches materials,
 			// until there's no materials to cache and all shaders have completed compilation
-			if (CacheMaterials(Level))
+			// NOTE: This uses the default offline shader profile. The real time profile is used in UAvaPlayableLevelStreaming::IsReadyToPlay
+			if (CacheMaterials(Level, UAvaMaterialCacheSettings::DefaultOfflineProfile))
 			{
 				bCachingShadersOnLoad = true;
 
@@ -670,7 +659,7 @@ void UAvaPlayableLevelStreaming::OnLevelStreamingStateChanged_Synchronized(ELeve
 						const ULevel* const Level = LevelStreaming ? LevelStreaming->GetLoadedLevel() : nullptr;
 						if (Level)
 						{
-							bCachingShadersOnLoad = CacheMaterials(Level);
+							bCachingShadersOnLoad = CacheMaterials(Level, UAvaMaterialCacheSettings::DefaultOfflineProfile);
 						}
 						else
 						{
diff --git a/Engine/Plugins/VirtualProduction/Avalanche/Source/AvalancheMedia/Private/Playable/Playables/AvaPlayableLevelStreaming.h b/Engine/Plugins/VirtualProduction/Avalanche/Source/AvalancheMedia/Private/Playable/Playables/AvaPlayableLevelStreaming.h
--- a/Engine/Plugins/VirtualProduction/Avalanche/Source/AvalancheMedia/Private/Playable/Playables/AvaPlayableLevelStreaming.h
+++ b/Engine/Plugins/VirtualProduction/Avalanche/Source/AvalancheMedia/Private/Playable/Playables/AvaPlayableLevelStreaming.h
@@ -52,9 +52,10 @@ protected:
 	/**
 	 * Caches the shaders for the given level.
 	 * @param InLevel the level to cache shaders for
+	 * @param InShaderProfile the shader profile to use for caching
 	 * @return true if caching is taking place. false otherwise. 
 	 */
-	bool CacheMaterials(TNotNull<const ULevel*> InLevel) const;
+	bool CacheMaterials(TNotNull<const ULevel*> InLevel, FName InShaderProfile) const;
 
 	void HandleTransitionEvent(UAvaPlayable* InPlayable, UAvaPlayableTransition* InTransition, EAvaPlayableTransitionEventFlags InTransitionFlags);
 
diff --git a/Engine/Plugins/VirtualProduction/Avalanche/Source/AvalancheMedia/Private/Playback/AvaPlaybackUtils.cpp b/Engine/Plugins/VirtualProduction/Avalanche/Source/AvalancheMedia/Private/Playback/AvaPlaybackUtils.cpp
--- a/Engine/Plugins/VirtualProduction/Avalanche/Source/AvalancheMedia/Private/Playback/AvaPlaybackUtils.cpp
+++ b/Engine/Plugins/VirtualProduction/Avalanche/Source/AvalancheMedia/Private/Playback/AvaPlaybackUtils.cpp
@@ -6,6 +6,7 @@
 #include "AssetRegistry/IAssetRegistry.h"
 #include "AvaAssetTags.h"
 #include "AvaMaterialCacheHelper.h"
+#include "AvaMaterialCacheSettings.h"
 #include "ContentStreaming.h"
 #include "Engine/StreamableRenderAsset.h"
 #include "Framework/AvaSoftAssetPtr.h"
@@ -260,9 +261,9 @@ namespace UE::AvaPlayback::Utils
 		return nullptr;
 	}
 
-	bool CacheObjectMaterials(UObject* InObject)
+	bool CacheObjectMaterials(UObject* InObject, FName InShaderProfile)
 	{
-		return UE::Ava::FMaterialCacheHelper::Get().RequestCacheMaterials(InObject);
+		return UE::Ava::FMaterialCacheHelper::Get().RequestCacheMaterials(InObject, InShaderProfile);
 	}
 
 	void FAsyncAssetLoader::ProcessLoadedObject(UObject* InObject)
@@ -277,7 +278,7 @@ namespace UE::AvaPlayback::Utils
 
 		if (Private::CVarRuntimePrecacheRCReferencedAssetsMaterials.GetValueOnGameThread())
 		{
-			if (CacheObjectMaterials(InObject))
+			if (CacheObjectMaterials(InObject, UAvaMaterialCacheSettings::DefaultRealtimeProfile))
 			{
 				PendingMaterialCompletionObjects.Add(InObject);
 			}
diff --git a/Engine/Plugins/VirtualProduction/Avalanche/Source/AvalancheMedia/Private/Rundown/AvaRundown.cpp b/Engine/Plugins/VirtualProduction/Avalanche/Source/AvalancheMedia/Private/Rundown/AvaRundown.cpp
--- a/Engine/Plugins/VirtualProduction/Avalanche/Source/AvalancheMedia/Private/Rundown/AvaRundown.cpp
+++ b/Engine/Plugins/VirtualProduction/Avalanche/Source/AvalancheMedia/Private/Rundown/AvaRundown.cpp
@@ -3,6 +3,7 @@
 #include "Rundown/AvaRundown.h"
 
 #include "AvaMaterialCacheHelper.h"
+#include "AvaMaterialCacheSettings.h"
 #include "AvaMediaSettings.h"
 #include "Broadcast/AvaBroadcast.h"
 #include "IAvaMediaModule.h"
@@ -1220,7 +1221,7 @@ TArray<UAvaRundown::FLoadedInstanceInfo> UAvaRundown::LoadPage(int32 InPageId,
 			if (UObject* Object = Asset.ResolveObject())
 			{
 				UE::AvaPlayback::Utils::StreamRenderAsset(Object);
-				UE::AvaPlayback::Utils::CacheObjectMaterials(Object);
+				UE::AvaPlayback::Utils::CacheObjectMaterials(Object, UAvaMaterialCacheSettings::DefaultOfflineProfile);
 			}
 			else
 			{
@@ -1230,7 +1231,7 @@ TArray<UAvaRundown::FLoadedInstanceInfo> UAvaRundown::LoadPage(int32 InPageId,
 						if (InObject)
 						{
 							UE::AvaPlayback::Utils::StreamRenderAsset(InObject);
-							UE::AvaPlayback::Utils::CacheObjectMaterials(InObject);
+							UE::AvaPlayback::Utils::CacheObjectMaterials(InObject, UAvaMaterialCacheSettings::DefaultOfflineProfile);
 						}
 					}));
 			}
