Motion Design: added option to wait on incomplete objects in material bridge.
Static Mesh can be in the middle of incremental loading and trying to access its static materials results in a long wait.
An option was added to material bridge visit func so that instead of waiting, a callable is invoked

#Synced-CL 48481255

diff --git a/Engine/Plugins/VirtualProduction/Avalanche/Source/AvalancheMaterial/Private/AvaMaterialModule.cpp b/Engine/Plugins/VirtualProduction/Avalanche/Source/AvalancheMaterial/Private/AvaMaterialModule.cpp
--- a/Engine/Plugins/VirtualProduction/Avalanche/Source/AvalancheMaterial/Private/AvaMaterialModule.cpp
+++ b/Engine/Plugins/VirtualProduction/Avalanche/Source/AvalancheMaterial/Private/AvaMaterialModule.cpp
@@ -5,11 +5,14 @@
 #include "AvaDecalComponentMaterialBridge.h"
 #include "AvaLevelMaterialBridge.h"
 #include "AvaMaterialBridge.h"
+#include "AvaMaterialBridgeLog.h"
 #include "AvaMaterialBridgeRegistry.h"
 #include "AvaPrimitiveComponentMaterialBridge.h"
 #include "AvaStaticMeshMaterialBridge.h"
 #include "Modules/ModuleManager.h"
 
+DEFINE_LOG_CATEGORY(LogAvaMaterialBridge);
+
 IMPLEMENT_MODULE(FAvaMaterialModule, AvalancheMaterial)
 
 void FAvaMaterialModule::StartupModule()
diff --git a/Engine/Plugins/VirtualProduction/Avalanche/Source/AvalancheMaterial/Private/AvaStaticMeshMaterialBridge.cpp b/Engine/Plugins/VirtualProduction/Avalanche/Source/AvalancheMaterial/Private/AvaStaticMeshMaterialBridge.cpp
--- a/Engine/Plugins/VirtualProduction/Avalanche/Source/AvalancheMaterial/Private/AvaStaticMeshMaterialBridge.cpp
+++ b/Engine/Plugins/VirtualProduction/Avalanche/Source/AvalancheMaterial/Private/AvaStaticMeshMaterialBridge.cpp
@@ -1,6 +1,7 @@
 // Copyright Epic Games, Inc. All Rights Reserved.
 
 #include "AvaStaticMeshMaterialBridge.h"
+#include "AvaMaterialBridgeLog.h"
 #include "Engine/StaticMesh.h"
 
 namespace UE::Ava
@@ -18,6 +19,22 @@ void FStaticMeshMaterialBridge::OnVisitMaterials(FConstDataView InDataView, TFun
 	FVisitedContext Context;
 	Context.Owner = InDataView;
 
+	// If static mesh is compiling and visitor does not want to wait on materials, bail out.
+	// NOTE: Only LockedProperties needs to be checked for EStaticMeshAsyncProperties::StaticMaterials, but there's currently no option to check for it directly.
+	if (StaticMesh.IsCompiling() && !InOptions.bTryWaitOnOwnerCompletion)
+	{
+		if (InOptions.OnOwnerPendingCompletion.IsBound())
+		{
+			UE_LOG(LogAvaMaterialBridge, Verbose, TEXT("Material Bridge could not visit materials of static mesh '%s'. It is currently compiling."), *StaticMesh.GetName());
+			InOptions.OnOwnerPendingCompletion.Execute(Context);
+		}
+		else
+		{
+			UE_LOG(LogAvaMaterialBridge, Warning, TEXT("Material Bridge could not visit materials of static mesh '%s'. It is currently compiling."), *StaticMesh.GetName());
+		}
+		return;
+	}
+
 	for (const FStaticMaterial& StaticMaterial : StaticMesh.GetStaticMaterials())
 	{
 		Context.Material = StaticMaterial.MaterialInterface;
diff --git a/Engine/Plugins/VirtualProduction/Avalanche/Source/AvalancheMaterial/Public/AvaMaterialBridge.h b/Engine/Plugins/VirtualProduction/Avalanche/Source/AvalancheMaterial/Public/AvaMaterialBridge.h
--- a/Engine/Plugins/VirtualProduction/Avalanche/Source/AvalancheMaterial/Public/AvaMaterialBridge.h
+++ b/Engine/Plugins/VirtualProduction/Avalanche/Source/AvalancheMaterial/Public/AvaMaterialBridge.h
@@ -44,7 +44,11 @@ public:
 	{
 		UE_API FVisitOptions();
 		/** The registry to use to visit any materials of sub objects. Defaults to the global registry. */
-		TNotNull<const FMaterialBridgeRegistry*> MaterialBridgeRegistry; 
+		TNotNull<const FMaterialBridgeRegistry*> MaterialBridgeRegistry;
+		/** Callback when an owner is not yet complete and access its materials. Only called when bWaitOnOwnerCompletion is false. */
+		TDelegate<void(const FVisitedContext&)> OnOwnerPendingCompletion;
+		/** Whether to try to wait for owner completion to access its materials */
+		bool bTryWaitOnOwnerCompletion = true;
 		/** Whether to include debug materials */
 		bool bVisitDebugMaterials = false;
 	};
diff --git a/Engine/Plugins/VirtualProduction/Avalanche/Source/AvalancheMaterial/Public/AvaMaterialBridgeLog.h b/Engine/Plugins/VirtualProduction/Avalanche/Source/AvalancheMaterial/Public/AvaMaterialBridgeLog.h
new file mode 100644
--- /dev/null
+++ b/Engine/Plugins/VirtualProduction/Avalanche/Source/AvalancheMaterial/Public/AvaMaterialBridgeLog.h
@@ -0,0 +1,7 @@
+// Copyright Epic Games, Inc. All Rights Reserved.
+
+#pragma once
+
+#include "Logging/LogCategory.h"
+
+DECLARE_LOG_CATEGORY_EXTERN(LogAvaMaterialBridge, Log, All);
