Motion Design: moved material bridge visit context  to its own file

#Synced-CL 48481255

diff --git a/Engine/Plugins/VirtualProduction/Avalanche/Source/AvalancheMaterial/Private/AvaActorMaterialBridge.cpp b/Engine/Plugins/VirtualProduction/Avalanche/Source/AvalancheMaterial/Private/AvaActorMaterialBridge.cpp
--- a/Engine/Plugins/VirtualProduction/Avalanche/Source/AvalancheMaterial/Private/AvaActorMaterialBridge.cpp
+++ b/Engine/Plugins/VirtualProduction/Avalanche/Source/AvalancheMaterial/Private/AvaActorMaterialBridge.cpp
@@ -12,18 +12,18 @@ const UStruct* FActorMaterialBridge::OnGetBridgedType() const
 	return AActor::StaticClass();
 }
 
-void FActorMaterialBridge::OnVisitMaterials(FConstDataView InDataView, TFunctionRef<void(const FVisitedContext&)> InFunc, const FVisitOptions& InOptions) const
+void FActorMaterialBridge::OnVisitMaterials(const FMaterialBridgeVisitContext& InContext, TFunctionRef<void(const FMaterialBridgeVisitContext&, UMaterialInterface*)> InFunc, const FVisitOptions& InOptions) const
 {
-	const AActor& Actor = InDataView.Get<const AActor>();
+	const AActor& Actor = InContext.Container.Get<const AActor>();
 
 	constexpr bool bIncludeFromChildActors = true;
 
 	Actor.ForEachComponent(bIncludeFromChildActors,
-		[&InFunc, &InOptions](UActorComponent* InComponent)
+		[&InFunc, &InOptions, &InContext](UActorComponent* InComponent)
 		{
 			if (const FMaterialBridge* MaterialBridge = InOptions.MaterialBridgeRegistry->GetMaterialBridgeForObject(InComponent))
 			{
-				MaterialBridge->VisitMaterials(FConstDataView(InComponent), InFunc, InOptions);
+				MaterialBridge->VisitMaterials(FMaterialBridgeVisitContext(InComponent, &InContext), InFunc, InOptions);
 			}
 		});
 }
diff --git a/Engine/Plugins/VirtualProduction/Avalanche/Source/AvalancheMaterial/Private/AvaDecalComponentMaterialBridge.cpp b/Engine/Plugins/VirtualProduction/Avalanche/Source/AvalancheMaterial/Private/AvaDecalComponentMaterialBridge.cpp
--- a/Engine/Plugins/VirtualProduction/Avalanche/Source/AvalancheMaterial/Private/AvaDecalComponentMaterialBridge.cpp
+++ b/Engine/Plugins/VirtualProduction/Avalanche/Source/AvalancheMaterial/Private/AvaDecalComponentMaterialBridge.cpp
@@ -12,15 +12,10 @@ const UStruct* FDecalComponentMaterialBridge::OnGetBridgedType() const
 	return UDecalComponent::StaticClass();
 }
 
-void FDecalComponentMaterialBridge::OnVisitMaterials(FConstDataView InDataView, TFunctionRef<void(const FVisitedContext&)> InFunc, const FVisitOptions& InOptions) const
+void FDecalComponentMaterialBridge::OnVisitMaterials(const FMaterialBridgeVisitContext& InContext, TFunctionRef<void(const FMaterialBridgeVisitContext&, UMaterialInterface*)> InFunc, const FVisitOptions& InOptions) const
 {
-	const UDecalComponent& Component = InDataView.Get<const UDecalComponent>();
-
-	FVisitedContext Context;
-	Context.Owner = InDataView;
-	Context.Material = Component.GetDecalMaterial();
-
-	InFunc(Context);
+	const UDecalComponent& Component = InContext.Container.Get<const UDecalComponent>();
+	InFunc(InContext, Component.GetDecalMaterial());
 }
 
 } // UE::Ava
diff --git a/Engine/Plugins/VirtualProduction/Avalanche/Source/AvalancheMaterial/Private/AvaLevelMaterialBridge.cpp b/Engine/Plugins/VirtualProduction/Avalanche/Source/AvalancheMaterial/Private/AvaLevelMaterialBridge.cpp
--- a/Engine/Plugins/VirtualProduction/Avalanche/Source/AvalancheMaterial/Private/AvaLevelMaterialBridge.cpp
+++ b/Engine/Plugins/VirtualProduction/Avalanche/Source/AvalancheMaterial/Private/AvaLevelMaterialBridge.cpp
@@ -14,15 +14,15 @@ const UStruct* FLevelMaterialBridge::OnGetBridgedType() const
 	return ULevel::StaticClass();
 }
 
-void FLevelMaterialBridge::OnVisitMaterials(FConstDataView InDataView, TFunctionRef<void(const FVisitedContext&)> InFunc, const FVisitOptions& InOptions) const
+void FLevelMaterialBridge::OnVisitMaterials(const FMaterialBridgeVisitContext& InContext, TFunctionRef<void(const FMaterialBridgeVisitContext&, UMaterialInterface*)> InFunc, const FVisitOptions& InOptions) const
 {
-	const ULevel& Level = InDataView.Get<const ULevel>();
+	const ULevel& Level = InContext.Container.Get<const ULevel>();
 
 	for (const AActor* Actor : Level.Actors)
 	{
 		if (const FMaterialBridge* MaterialBridge = InOptions.MaterialBridgeRegistry->GetMaterialBridgeForObject(Actor))
 		{
-			MaterialBridge->VisitMaterials(FConstDataView(Actor), InFunc, InOptions);
+			MaterialBridge->VisitMaterials(FMaterialBridgeVisitContext(Actor, &InContext), InFunc, InOptions);
 		}
 	}
 }
diff --git a/Engine/Plugins/VirtualProduction/Avalanche/Source/AvalancheMaterial/Private/AvaMaterialBridge.cpp b/Engine/Plugins/VirtualProduction/Avalanche/Source/AvalancheMaterial/Private/AvaMaterialBridge.cpp
--- a/Engine/Plugins/VirtualProduction/Avalanche/Source/AvalancheMaterial/Private/AvaMaterialBridge.cpp
+++ b/Engine/Plugins/VirtualProduction/Avalanche/Source/AvalancheMaterial/Private/AvaMaterialBridge.cpp
@@ -5,6 +5,10 @@
 
 namespace UE::Ava
 {
+FMaterialBridge::FVisitOptions::FVisitOptions()
+	: MaterialBridgeRegistry(&FMaterialBridgeRegistry::Get())
+{
+}
 
 void FMaterialBridge::Initialize()
 {
@@ -16,16 +20,11 @@ const UStruct* FMaterialBridge::GetBridgedType() const
 	return BridgedType;
 }
 
-FMaterialBridge::FVisitOptions::FVisitOptions()
-	: MaterialBridgeRegistry(&FMaterialBridgeRegistry::Get())
-{
-}
-
-void FMaterialBridge::VisitMaterials(FConstDataView InDataView, TFunctionRef<void(const FVisitedContext&)> InFunc, const FVisitOptions& InOptions) const
+void FMaterialBridge::VisitMaterials(const FMaterialBridgeVisitContext& InContext, TFunctionRef<void(const FMaterialBridgeVisitContext&, UMaterialInterface*)> InFunc, const FVisitOptions& InOptions) const
 {
-	if (InDataView.IsValidFor(BridgedType))
+	if (InContext.Container.IsValidFor(BridgedType))
 	{
-		OnVisitMaterials(InDataView, InFunc, InOptions);
+		OnVisitMaterials(InContext, InFunc, InOptions);
 	}
 }
 
diff --git a/Engine/Plugins/VirtualProduction/Avalanche/Source/AvalancheMaterial/Private/AvaMaterialBridgeVisitContext.cpp b/Engine/Plugins/VirtualProduction/Avalanche/Source/AvalancheMaterial/Private/AvaMaterialBridgeVisitContext.cpp
new file mode 100644
--- /dev/null
+++ b/Engine/Plugins/VirtualProduction/Avalanche/Source/AvalancheMaterial/Private/AvaMaterialBridgeVisitContext.cpp
@@ -0,0 +1,19 @@
+// Copyright Epic Games, Inc. All Rights Reserved.
+
+#include "AvaMaterialBridgeVisitContext.h"
+
+namespace UE::Ava
+{
+
+FMaterialBridgeVisitContext::FMaterialBridgeVisitContext(const UObject* InContainer)
+	: Container(InContainer)
+{
+}
+
+FMaterialBridgeVisitContext::FMaterialBridgeVisitContext(const UObject* InContainer, TNotNull<const FMaterialBridgeVisitContext*> InParentContext)
+	: Container(InContainer)
+	, ParentContext(InParentContext)
+{
+}
+
+} // UE::Ava
diff --git a/Engine/Plugins/VirtualProduction/Avalanche/Source/AvalancheMaterial/Private/AvaMaterialCacheHelper.cpp b/Engine/Plugins/VirtualProduction/Avalanche/Source/AvalancheMaterial/Private/AvaMaterialCacheHelper.cpp
--- a/Engine/Plugins/VirtualProduction/Avalanche/Source/AvalancheMaterial/Private/AvaMaterialCacheHelper.cpp
+++ b/Engine/Plugins/VirtualProduction/Avalanche/Source/AvalancheMaterial/Private/AvaMaterialCacheHelper.cpp
@@ -414,24 +414,24 @@ bool FMaterialCacheHelper::GatherMaterialData(const UObject* InObject, FName InS
 		};
 
 	// Gather all materials for object
-	MaterialBridge->VisitMaterials(FConstDataView(InObject),
-		[&AddIncompleteMaterial](const FMaterialBridge::FVisitedContext& InContext)
+	MaterialBridge->VisitMaterials(FMaterialBridgeVisitContext(InObject),
+		[&AddIncompleteMaterial](const FMaterialBridgeVisitContext& InContext, UMaterialInterface* InMaterial)
 		{
-			if (!InContext.Material)
+			if (!InMaterial)
 			{
 				return;
 			}
 
-			UMaterial* const Material = InContext.Material->GetMaterial();
+			UMaterial* const Material = InMaterial->GetMaterial();
 
 			// Add interface's material if it is different from the interface
-			if (Material && Material != InContext.Material && !Material->IsComplete())
+			if (Material && Material != InMaterial && !Material->IsComplete())
 			{
 				AddIncompleteMaterial(Material);
 			}
-			if (!InContext.Material->IsComplete())
+			if (!InMaterial->IsComplete())
 			{
-				AddIncompleteMaterial(InContext.Material);
+				AddIncompleteMaterial(InMaterial);
 			}
 		});
 
diff --git a/Engine/Plugins/VirtualProduction/Avalanche/Source/AvalancheMaterial/Private/AvaPrimitiveComponentMaterialBridge.cpp b/Engine/Plugins/VirtualProduction/Avalanche/Source/AvalancheMaterial/Private/AvaPrimitiveComponentMaterialBridge.cpp
--- a/Engine/Plugins/VirtualProduction/Avalanche/Source/AvalancheMaterial/Private/AvaPrimitiveComponentMaterialBridge.cpp
+++ b/Engine/Plugins/VirtualProduction/Avalanche/Source/AvalancheMaterial/Private/AvaPrimitiveComponentMaterialBridge.cpp
@@ -12,20 +12,16 @@ const UStruct* FPrimitiveComponentMaterialBridge::OnGetBridgedType() const
 	return UPrimitiveComponent::StaticClass();
 }
 
-void FPrimitiveComponentMaterialBridge::OnVisitMaterials(FConstDataView InDataView, TFunctionRef<void(const FVisitedContext&)> InFunc, const FVisitOptions& InOptions) const
+void FPrimitiveComponentMaterialBridge::OnVisitMaterials(const FMaterialBridgeVisitContext& InContext, TFunctionRef<void(const FMaterialBridgeVisitContext&, UMaterialInterface*)> InFunc, const FVisitOptions& InOptions) const
 {
-	const UPrimitiveComponent& Component = InDataView.Get<const UPrimitiveComponent>();
+	const UPrimitiveComponent& Component = InContext.Container.Get<const UPrimitiveComponent>();
 
 	TArray<UMaterialInterface*> Materials;
 	Component.GetUsedMaterials(Materials, InOptions.bVisitDebugMaterials);
 
-	FVisitedContext Context;
-	Context.Owner = InDataView;
-
 	for (UMaterialInterface* Material : Materials)
 	{
-		Context.Material = Material;
-		InFunc(Context);
+		InFunc(InContext, Material);
 	}
 }
 
diff --git a/Engine/Plugins/VirtualProduction/Avalanche/Source/AvalancheMaterial/Private/AvaStaticMeshMaterialBridge.cpp b/Engine/Plugins/VirtualProduction/Avalanche/Source/AvalancheMaterial/Private/AvaStaticMeshMaterialBridge.cpp
--- a/Engine/Plugins/VirtualProduction/Avalanche/Source/AvalancheMaterial/Private/AvaStaticMeshMaterialBridge.cpp
+++ b/Engine/Plugins/VirtualProduction/Avalanche/Source/AvalancheMaterial/Private/AvaStaticMeshMaterialBridge.cpp
@@ -12,21 +12,18 @@ const UStruct* FStaticMeshMaterialBridge::OnGetBridgedType() const
 	return UStaticMesh::StaticClass();
 }
 
-void FStaticMeshMaterialBridge::OnVisitMaterials(FConstDataView InDataView, TFunctionRef<void(const FVisitedContext&)> InFunc, const FVisitOptions& InOptions) const
+void FStaticMeshMaterialBridge::OnVisitMaterials(const FMaterialBridgeVisitContext& InContext, TFunctionRef<void(const FMaterialBridgeVisitContext&, UMaterialInterface*)> InFunc, const FVisitOptions& InOptions) const
 {
-	const UStaticMesh& StaticMesh = InDataView.Get<const UStaticMesh>();
-
-	FVisitedContext Context;
-	Context.Owner = InDataView;
+	const UStaticMesh& StaticMesh = InContext.Container.Get<const UStaticMesh>();
 
 	// If static mesh is compiling and visitor does not want to wait on materials, bail out.
 	// NOTE: Only LockedProperties needs to be checked for EStaticMeshAsyncProperties::StaticMaterials, but there's currently no option to check for it directly.
-	if (StaticMesh.IsCompiling() && !InOptions.bTryWaitOnOwnerCompletion)
+	if (StaticMesh.IsCompiling() && !InOptions.bTryWaitOnContainerCompletion)
 	{
-		if (InOptions.OnOwnerPendingCompletion.IsBound())
+		if (InOptions.OnContainerPendingCompletion.IsBound())
 		{
 			UE_LOG(LogAvaMaterialBridge, Verbose, TEXT("Material Bridge could not visit materials of static mesh '%s'. It is currently compiling."), *StaticMesh.GetName());
-			InOptions.OnOwnerPendingCompletion.Execute(Context);
+			InOptions.OnContainerPendingCompletion.Execute(InContext);
 		}
 		else
 		{
@@ -37,11 +34,8 @@ void FStaticMeshMaterialBridge::OnVisitMaterials(FConstDataView InDataView, TFun
 
 	for (const FStaticMaterial& StaticMaterial : StaticMesh.GetStaticMaterials())
 	{
-		Context.Material = StaticMaterial.MaterialInterface;
-		InFunc(Context);
-
-		Context.Material = StaticMaterial.OverlayMaterialInterface;
-		InFunc(Context);
+		InFunc(InContext, StaticMaterial.MaterialInterface);
+		InFunc(InContext, StaticMaterial.OverlayMaterialInterface);
 	}
 }
 
diff --git a/Engine/Plugins/VirtualProduction/Avalanche/Source/AvalancheMaterial/Public/AvaActorMaterialBridge.h b/Engine/Plugins/VirtualProduction/Avalanche/Source/AvalancheMaterial/Public/AvaActorMaterialBridge.h
--- a/Engine/Plugins/VirtualProduction/Avalanche/Source/AvalancheMaterial/Public/AvaActorMaterialBridge.h
+++ b/Engine/Plugins/VirtualProduction/Avalanche/Source/AvalancheMaterial/Public/AvaActorMaterialBridge.h
@@ -15,7 +15,7 @@ class FActorMaterialBridge : public FMaterialBridge
 protected:
 	//~ Begin FMaterialBridge
 	UE_API virtual const UStruct* OnGetBridgedType() const override;
-	UE_API virtual void OnVisitMaterials(FConstDataView InDataView, TFunctionRef<void(const FVisitedContext&)> InFunc, const FVisitOptions& InOptions) const override;
+	UE_API virtual void OnVisitMaterials(const FMaterialBridgeVisitContext& InContext, TFunctionRef<void(const FMaterialBridgeVisitContext&, UMaterialInterface*)> InFunc, const FVisitOptions& InOptions) const override;
 	//~ End FMaterialBridge
 };
 
diff --git a/Engine/Plugins/VirtualProduction/Avalanche/Source/AvalancheMaterial/Public/AvaDecalComponentMaterialBridge.h b/Engine/Plugins/VirtualProduction/Avalanche/Source/AvalancheMaterial/Public/AvaDecalComponentMaterialBridge.h
--- a/Engine/Plugins/VirtualProduction/Avalanche/Source/AvalancheMaterial/Public/AvaDecalComponentMaterialBridge.h
+++ b/Engine/Plugins/VirtualProduction/Avalanche/Source/AvalancheMaterial/Public/AvaDecalComponentMaterialBridge.h
@@ -15,7 +15,7 @@ class FDecalComponentMaterialBridge : public FMaterialBridge
 protected:
 	//~ Begin FMaterialBridge
 	UE_API virtual const UStruct* OnGetBridgedType() const override;
-	UE_API virtual void OnVisitMaterials(FConstDataView InDataView, TFunctionRef<void(const FVisitedContext&)> InFunc, const FVisitOptions& InOptions) const override;
+	UE_API virtual void OnVisitMaterials(const FMaterialBridgeVisitContext& InContext, TFunctionRef<void(const FMaterialBridgeVisitContext&, UMaterialInterface*)> InFunc, const FVisitOptions& InOptions) const override;
 	//~ End FMaterialBridge
 };
 
diff --git a/Engine/Plugins/VirtualProduction/Avalanche/Source/AvalancheMaterial/Public/AvaLevelMaterialBridge.h b/Engine/Plugins/VirtualProduction/Avalanche/Source/AvalancheMaterial/Public/AvaLevelMaterialBridge.h
--- a/Engine/Plugins/VirtualProduction/Avalanche/Source/AvalancheMaterial/Public/AvaLevelMaterialBridge.h
+++ b/Engine/Plugins/VirtualProduction/Avalanche/Source/AvalancheMaterial/Public/AvaLevelMaterialBridge.h
@@ -15,7 +15,7 @@ class FLevelMaterialBridge : public FMaterialBridge
 protected:
 	//~ Begin FMaterialBridge
 	UE_API virtual const UStruct* OnGetBridgedType() const override;
-	UE_API virtual void OnVisitMaterials(FConstDataView InDataView, TFunctionRef<void(const FVisitedContext&)> InFunc, const FVisitOptions& InOptions) const override;
+	UE_API virtual void OnVisitMaterials(const FMaterialBridgeVisitContext& InContext, TFunctionRef<void(const FMaterialBridgeVisitContext&, UMaterialInterface*)> InFunc, const FVisitOptions& InOptions) const override;
 	//~ End FMaterialBridge
 };
 
diff --git a/Engine/Plugins/VirtualProduction/Avalanche/Source/AvalancheMaterial/Public/AvaMaterialBridge.h b/Engine/Plugins/VirtualProduction/Avalanche/Source/AvalancheMaterial/Public/AvaMaterialBridge.h
--- a/Engine/Plugins/VirtualProduction/Avalanche/Source/AvalancheMaterial/Public/AvaMaterialBridge.h
+++ b/Engine/Plugins/VirtualProduction/Avalanche/Source/AvalancheMaterial/Public/AvaMaterialBridge.h
@@ -2,7 +2,7 @@
 
 #pragma once
 
-#include "AvaDataView.h"
+#include "AvaMaterialBridgeVisitContext.h"
 #include "Containers/ContainersFwd.h"
 #include "StructUtils/StructUtilsTypes.h"
 #include "UObject/ObjectPtr.h"
@@ -33,39 +33,32 @@ public:
 	/** Retrieves the type that this bridge supports */
 	UE_API const UStruct* GetBridgedType() const;
 
-	struct FVisitedContext
-	{
-		/** Owner of the visited material */
-		FConstDataView Owner;
-		/** The visited material */
-		TObjectPtr<UMaterialInterface> Material;
-	};
 	struct FVisitOptions
 	{
 		UE_API FVisitOptions();
 		/** The registry to use to visit any materials of sub objects. Defaults to the global registry. */
 		TNotNull<const FMaterialBridgeRegistry*> MaterialBridgeRegistry;
-		/** Callback when an owner is not yet complete and access its materials. Only called when bWaitOnOwnerCompletion is false. */
-		TDelegate<void(const FVisitedContext&)> OnOwnerPendingCompletion;
-		/** Whether to try to wait for owner completion to access its materials */
-		bool bTryWaitOnOwnerCompletion = true;
+		/** Callback when a container is not yet complete and cannot access its materials. Only called when bTryWaitOnContainerCompletion is false. */
+		TDelegate<void(const FMaterialBridgeVisitContext&)> OnContainerPendingCompletion;
+		/** Whether to try to wait for container completion to access its materials */
+		bool bTryWaitOnContainerCompletion = true;
 		/** Whether to include debug materials */
 		bool bVisitDebugMaterials = false;
 	};
 	/**
-	 * Visits all the used materials in the given data view.
-	 * @param InDataView the data view to find materials from.
+	 * Visits all the used materials in the given context.
+	 * @param InContext the current context to find materials from.
 	 * @param InFunc function to call for each material found.
-	 * @param InOptions additional visit options.
+	 * @param InOptions visit options. see FVisitOptions
 	 */
-	UE_API void VisitMaterials(FConstDataView InDataView, TFunctionRef<void(const FVisitedContext&)> InFunc, const FVisitOptions& InOptions = {}) const;
+	UE_API void VisitMaterials(const FMaterialBridgeVisitContext& InContext, TFunctionRef<void(const FMaterialBridgeVisitContext&, UMaterialInterface*)> InFunc, const FVisitOptions& InOptions = {}) const;
 
 protected:
 	/** Retrieves the type that this bridge supports. Only called once in Initialize() */
 	virtual const UStruct* OnGetBridgedType() const = 0;
 
-	/** Get all the used materials in the given data view. */
-	virtual void OnVisitMaterials(FConstDataView InDataView, TFunctionRef<void(const FVisitedContext&)> InFunc, const FVisitOptions& InOptions) const = 0;
+	/** Get all the used materials in the given context. */
+	virtual void OnVisitMaterials(const FMaterialBridgeVisitContext& InContext, TFunctionRef<void(const FMaterialBridgeVisitContext&, UMaterialInterface*)> InFunc, const FVisitOptions& InOptions) const = 0;
 
 private:
 	/** The type that this Material Bridge works on */
diff --git a/Engine/Plugins/VirtualProduction/Avalanche/Source/AvalancheMaterial/Public/AvaMaterialBridgeVisitContext.h b/Engine/Plugins/VirtualProduction/Avalanche/Source/AvalancheMaterial/Public/AvaMaterialBridgeVisitContext.h
new file mode 100644
--- /dev/null
+++ b/Engine/Plugins/VirtualProduction/Avalanche/Source/AvalancheMaterial/Public/AvaMaterialBridgeVisitContext.h
@@ -0,0 +1,30 @@
+// Copyright Epic Games, Inc. All Rights Reserved.
+
+#pragma once
+
+#include "AvaDataView.h"
+
+#define UE_API AVALANCHEMATERIAL_API
+
+namespace UE::Ava
+{
+
+/** Short-lived context information of a visited material */
+class FMaterialBridgeVisitContext
+{
+public:
+	FMaterialBridgeVisitContext() = default;
+
+	explicit UE_API FMaterialBridgeVisitContext(const UObject* InContainer);
+	explicit UE_API FMaterialBridgeVisitContext(const UObject* InContainer, TNotNull<const FMaterialBridgeVisitContext*> InParentContext);
+
+	/** Container of the visited material */
+	FConstDataView Container;
+
+	/** Parent of this context for additional information (e.g. container of the container) */
+	const FMaterialBridgeVisitContext* ParentContext = nullptr;
+};
+
+} // UE::Ava
+
+#undef UE_API
diff --git a/Engine/Plugins/VirtualProduction/Avalanche/Source/AvalancheMaterial/Public/AvaPrimitiveComponentMaterialBridge.h b/Engine/Plugins/VirtualProduction/Avalanche/Source/AvalancheMaterial/Public/AvaPrimitiveComponentMaterialBridge.h
--- a/Engine/Plugins/VirtualProduction/Avalanche/Source/AvalancheMaterial/Public/AvaPrimitiveComponentMaterialBridge.h
+++ b/Engine/Plugins/VirtualProduction/Avalanche/Source/AvalancheMaterial/Public/AvaPrimitiveComponentMaterialBridge.h
@@ -15,7 +15,7 @@ class FPrimitiveComponentMaterialBridge : public FMaterialBridge
 protected:
 	//~ Begin FMaterialBridge
 	UE_API virtual const UStruct* OnGetBridgedType() const override;
-	UE_API virtual void OnVisitMaterials(FConstDataView InDataView, TFunctionRef<void(const FVisitedContext&)> InFunc, const FVisitOptions& InOptions) const override;
+	UE_API virtual void OnVisitMaterials(const FMaterialBridgeVisitContext& InContext, TFunctionRef<void(const FMaterialBridgeVisitContext&, UMaterialInterface*)> InFunc, const FVisitOptions& InOptions) const override;
 	//~ End FMaterialBridge
 };
 
diff --git a/Engine/Plugins/VirtualProduction/Avalanche/Source/AvalancheMaterial/Public/AvaStaticMeshMaterialBridge.h b/Engine/Plugins/VirtualProduction/Avalanche/Source/AvalancheMaterial/Public/AvaStaticMeshMaterialBridge.h
--- a/Engine/Plugins/VirtualProduction/Avalanche/Source/AvalancheMaterial/Public/AvaStaticMeshMaterialBridge.h
+++ b/Engine/Plugins/VirtualProduction/Avalanche/Source/AvalancheMaterial/Public/AvaStaticMeshMaterialBridge.h
@@ -18,7 +18,7 @@ class FStaticMeshMaterialBridge : public FMaterialBridge
 protected:
 	//~ Begin FMaterialBridge
 	UE_API virtual const UStruct* OnGetBridgedType() const override;
-	UE_API virtual void OnVisitMaterials(FConstDataView InDataView, TFunctionRef<void(const FVisitedContext&)> InFunc, const FVisitOptions& InOptions) const override;
+	UE_API virtual void OnVisitMaterials(const FMaterialBridgeVisitContext& InContext, TFunctionRef<void(const FMaterialBridgeVisitContext&, UMaterialInterface*)> InFunc, const FVisitOptions& InOptions) const override;
 	//~ End FMaterialBridge
 };
 
diff --git a/Engine/Plugins/VirtualProduction/Avalanche/Source/AvalancheShapes/Private/AvaShapeComponentMaterialBridge.cpp b/Engine/Plugins/VirtualProduction/Avalanche/Source/AvalancheShapes/Private/AvaShapeComponentMaterialBridge.cpp
--- a/Engine/Plugins/VirtualProduction/Avalanche/Source/AvalancheShapes/Private/AvaShapeComponentMaterialBridge.cpp
+++ b/Engine/Plugins/VirtualProduction/Avalanche/Source/AvalancheShapes/Private/AvaShapeComponentMaterialBridge.cpp
@@ -11,17 +11,13 @@ const UStruct* FShapeComponentMaterialBridge::OnGetBridgedType() const
 	return UAvaShapeDynamicMeshBase::StaticClass();
 }
 
-void FShapeComponentMaterialBridge::OnVisitMaterials(FConstDataView InDataView, TFunctionRef<void(const FVisitedContext&)> InFunc, const FVisitOptions& InOptions) const
+void FShapeComponentMaterialBridge::OnVisitMaterials(const FMaterialBridgeVisitContext& InContext, TFunctionRef<void(const FMaterialBridgeVisitContext&, UMaterialInterface*)> InFunc, const FVisitOptions& InOptions) const
 {
-	const UAvaShapeDynamicMeshBase& Component = InDataView.Get<const UAvaShapeDynamicMeshBase>();
-
-	FVisitedContext Context;
-	Context.Owner = InDataView;
+	const UAvaShapeDynamicMeshBase& Component = InContext.Container.Get<const UAvaShapeDynamicMeshBase>();
 
 	for (const TPair<int32, FAvaShapeMeshData>& Pair : Component.GetMeshDatas())
 	{
-		Context.Material = Pair.Value.GetMaterial();
-		InFunc(Context);
+		InFunc(InContext, Pair.Value.GetMaterial());
 	}
 }
 
diff --git a/Engine/Plugins/VirtualProduction/Avalanche/Source/AvalancheShapes/Public/AvaShapeComponentMaterialBridge.h b/Engine/Plugins/VirtualProduction/Avalanche/Source/AvalancheShapes/Public/AvaShapeComponentMaterialBridge.h
--- a/Engine/Plugins/VirtualProduction/Avalanche/Source/AvalancheShapes/Public/AvaShapeComponentMaterialBridge.h
+++ b/Engine/Plugins/VirtualProduction/Avalanche/Source/AvalancheShapes/Public/AvaShapeComponentMaterialBridge.h
@@ -15,7 +15,7 @@ class FShapeComponentMaterialBridge : public FMaterialBridge
 protected:
 	//~ Begin FMaterialBridge
 	UE_API virtual const UStruct* OnGetBridgedType() const override;
-	UE_API virtual void OnVisitMaterials(FConstDataView InDataView, TFunctionRef<void(const FVisitedContext&)> InFunc, const FVisitOptions& InOptions) const override;
+	UE_API virtual void OnVisitMaterials(const FMaterialBridgeVisitContext& InContext, TFunctionRef<void(const FMaterialBridgeVisitContext&, UMaterialInterface*)> InFunc, const FVisitOptions& InOptions) const override;
 	//~ End FMaterialBridge
 };
 
diff --git a/Engine/Plugins/VirtualProduction/Avalanche/Source/AvalancheText/Private/AvaText3DComponentMaterialBridge.cpp b/Engine/Plugins/VirtualProduction/Avalanche/Source/AvalancheText/Private/AvaText3DComponentMaterialBridge.cpp
--- a/Engine/Plugins/VirtualProduction/Avalanche/Source/AvalancheText/Private/AvaText3DComponentMaterialBridge.cpp
+++ b/Engine/Plugins/VirtualProduction/Avalanche/Source/AvalancheText/Private/AvaText3DComponentMaterialBridge.cpp
@@ -13,21 +13,18 @@ const UStruct* FText3DComponentMaterialBridge::OnGetBridgedType() const
 	return UText3DComponent::StaticClass();
 }
 
-void FText3DComponentMaterialBridge::OnVisitMaterials(FConstDataView InDataView, TFunctionRef<void(const FVisitedContext&)> InFunc, const FVisitOptions& InOptions) const
+void FText3DComponentMaterialBridge::OnVisitMaterials(const FMaterialBridgeVisitContext& InContext, TFunctionRef<void(const FMaterialBridgeVisitContext&, UMaterialInterface*)> InFunc, const FVisitOptions& InOptions) const
 {
-	const UText3DComponent& Component = InDataView.Get<const UText3DComponent>();
+	const UText3DComponent& Component = InContext.Container.Get<const UText3DComponent>();
 
 	if (const UText3DMaterialExtensionBase* MaterialExtension = Component.GetMaterialExtension())
 	{
-		FVisitedContext VisitContext;
-		VisitContext.Owner = FConstDataView(MaterialExtension);
+		const FMaterialBridgeVisitContext VisitContext(MaterialExtension, &InContext);
 
 		MaterialExtension->ForEachMaterial(
 			[&InFunc, &VisitContext](const UE::Text3D::Material::FMaterialParameters&, UMaterialInterface* InMaterial)->bool
 			{
-				VisitContext.Material = InMaterial;
-
-				InFunc(VisitContext);
+				InFunc(VisitContext, InMaterial);
 				return true; // continue				
 			});
 	}
diff --git a/Engine/Plugins/VirtualProduction/Avalanche/Source/AvalancheText/Public/AvaText3DComponentMaterialBridge.h b/Engine/Plugins/VirtualProduction/Avalanche/Source/AvalancheText/Public/AvaText3DComponentMaterialBridge.h
--- a/Engine/Plugins/VirtualProduction/Avalanche/Source/AvalancheText/Public/AvaText3DComponentMaterialBridge.h
+++ b/Engine/Plugins/VirtualProduction/Avalanche/Source/AvalancheText/Public/AvaText3DComponentMaterialBridge.h
@@ -15,7 +15,7 @@ class FText3DComponentMaterialBridge : public FMaterialBridge
 protected:
 	//~ Begin FMaterialBridge
 	UE_API virtual const UStruct* OnGetBridgedType() const override;
-	UE_API virtual void OnVisitMaterials(FConstDataView InDataView, TFunctionRef<void(const FVisitedContext&)> InFunc, const FVisitOptions& InOptions) const override;
+	UE_API virtual void OnVisitMaterials(const FMaterialBridgeVisitContext& InContext, TFunctionRef<void(const FMaterialBridgeVisitContext&, UMaterialInterface*)> InFunc, const FVisitOptions& InOptions) const override;
 	//~ End FMaterialBridge
 };
 
