Motion Design: added support for waiting on containers (e.g. StaticMesh) that are not fully ready for material gather yet. Improvements on material bridge for request to be only called once instead of every frame.

#Synced-CL 48481255

diff --git a/Engine/Plugins/VirtualProduction/Avalanche/Source/AvalancheMaterial/Internal/AvaMaterialCacheHelper.h b/Engine/Plugins/VirtualProduction/Avalanche/Source/AvalancheMaterial/Internal/AvaMaterialCacheHelper.h
--- a/Engine/Plugins/VirtualProduction/Avalanche/Source/AvalancheMaterial/Internal/AvaMaterialCacheHelper.h
+++ b/Engine/Plugins/VirtualProduction/Avalanche/Source/AvalancheMaterial/Internal/AvaMaterialCacheHelper.h
@@ -23,14 +23,6 @@ namespace UE::Ava
 /** Helper to cache materials over multiple frames */
 class FMaterialCacheHelper : public TSharedFromThis<FMaterialCacheHelper>, public FGCObject
 {
-	enum class EMaterialCacheStatus : uint8
-	{
-		/** Material is waiting to call CacheShaders */
-		Waiting,
-		/** Material called CacheShaders */
-		Caching,
-	};
-
 	struct FMaterialKey
 	{
 		/** Object key to the Material being cached */
@@ -50,16 +42,20 @@ class FMaterialCacheHelper : public TSharedFromThis<FMaterialCacheHelper>, publi
 
 	struct FObjectData
 	{
+		/** Shader profile to use for caching the material shaders */
+		FName ShaderProfile;
 		/** All the keys to the materials belonging to an object that are still incomplete */
 		TSet<FMaterialKey> IncompleteMaterialKeys;
+		/** List of subobjects that were pending completion and could not gather materials */
+		TSet<TObjectPtr<const UObject>> PendingSubobjects;
+		/** Time elapsed before checking pending subobjects */
+		float PendingSubobjectElapsedTime = 0.f;
 	};
 
 	struct FMaterialData
 	{
 		/** Material that is pending completion */
 		TObjectPtr<UMaterialInterface> Material;
-		/** Gives the status of cache */
-		EMaterialCacheStatus CacheStatus = EMaterialCacheStatus::Waiting;
 	};
 
 public:
@@ -91,6 +87,9 @@ public:
 	//~ End FGCObject
 
 private:
+	/** Returns true if the material is still being cached */
+	bool IsCachingMaterial(const FMaterialKey& InMaterialKey) const;
+
 	/** Resets the object and material map */
 	void Reset();
 
@@ -100,20 +99,14 @@ private:
 	/** Unregisters the tick function from being called every frame */
 	void UnregisterTick();
 
-	/**
-	 * Caches as many queued materials as the time budget allows, and updates the status of each material to 'caching'
-	 * It also removes any material that is fully complete.
-	 */
-	void ProcessMaterials();
-
-	/** Finds the shader profile matching the given name */
-	const FResolvedShaderProfile* FindShaderProfile(FName InShaderProfile) const;
+	/** Processes subobjects that did not have its materials gathered */
+	void ProcessPendingObjects();
 
-	/** Caches the given material's shaders with the given shader profile. If not found, it will use the UMaterial::CacheShaders for compilation */
-	void CacheMaterial(TNotNull<UMaterialInterface*> InMaterial, FName InShaderProfile);
+	/** Caches as many queued materials as the time budget allows */
+	void ProcessMaterialsPendingCache();
 
-	/** Determines whether the given material is complete */
-	bool IsMaterialComplete(TNotNull<UMaterialInterface*> InMaterial, FName InShaderProfile) const;
+	/** Checks the status on materials that are being cached, and removes from the map once they're complete */
+	void ProcessMaterialsCaching();
 
 	/** Removes the complete materials from the objects' IncompleteMaterials set */
 	void RemoveCompletedMaterials();
@@ -122,16 +115,22 @@ private:
 	void ProcessShadersCompiling();
 
 	/**
-	 * Gathers the material data for an object, or skips it if the object's materials have been gathered before.
+	 * Gathers the material data for an object
 	 * returns true if a material bridge was found for the object and visited its materials
 	 */
-	bool GatherMaterialData(const UObject* InObject, FName InShaderProfile);
+	bool GatherMaterialData(const UObject* InObject, FObjectData& InObjectData, FName InShaderProfile);
 
 	/** Objects that have already been gathered, mapped to the incomplete materials that were gathered for the object */
 	TMap<FObjectKey, FObjectData> Objects;
 
-	/** All the materials that have already been gathered and are pending completion. Once materials are complete, they're removed from this map */
-	TMap<FMaterialKey, FMaterialData> Materials;
+	/** All the materials that have been gathered and are pending to start caching their shaders. Once materials start caching, they're removed from this map and added to MaterialsCaching */
+	TMap<FMaterialKey, FMaterialData> MaterialsPendingCache;
+
+	/** All the materials that are caching. Once materials are complete, they're removed from this map */
+	TMap<FMaterialKey, FMaterialData> MaterialsCaching;
+
+	/** The delta time at the start of the tick */
+	float PreTickFrameDuration = 0.0;
 
 	/** Time when BeginFrame was called */
 	double BeginFrameTime = 0.0;
diff --git a/Engine/Plugins/VirtualProduction/Avalanche/Source/AvalancheMaterial/Internal/AvaMaterialCacheSettings.h b/Engine/Plugins/VirtualProduction/Avalanche/Source/AvalancheMaterial/Internal/AvaMaterialCacheSettings.h
--- a/Engine/Plugins/VirtualProduction/Avalanche/Source/AvalancheMaterial/Internal/AvaMaterialCacheSettings.h
+++ b/Engine/Plugins/VirtualProduction/Avalanche/Source/AvalancheMaterial/Internal/AvaMaterialCacheSettings.h
@@ -11,14 +11,32 @@ class FShaderPipelineType;
 class FShaderType;
 class FVertexFactoryType;
 
+UENUM()
+enum class EAvaShaderProfileType : uint8
+{
+	/** CacheShaders will be used and will gather all the shaders it might need */
+	Default,
+	/**
+	 * Only the specified shaders will be cached. Should yield better performance 
+	 * NOTE: Specified shaders only will apply in editor builds. This will default to generic on non-editor builds
+	 */
+	Specific,
+	/** No shaders will be cached */
+	SkipCaching,
+};
+
 namespace UE::Ava
 {
-	struct FResolvedShaderProfile
-	{
-		TArray<const FVertexFactoryType*> VertexFactoryTypes;
-		TArray<const FShaderPipelineType*> PipelineTypes;
-		TArray<const FShaderType*> ShaderTypes;
-	};
+
+/** Shader Profile that has all the type pointers cached */
+struct FResolvedShaderProfile
+{
+	EAvaShaderProfileType ProfileType = EAvaShaderProfileType::Default;
+	TArray<const FVertexFactoryType*> VertexFactoryTypes;
+	TArray<const FShaderPipelineType*> PipelineTypes;
+	TArray<const FShaderType*> ShaderTypes;
+};
+
 } // UE::Ava
 
 USTRUCT()
@@ -54,9 +72,10 @@ struct FAvaShaderProfile
 	FName ProfileName;
 
 	UPROPERTY(EditAnywhere, Category="Material Cache")
-	bool bEnabled = true;
+	EAvaShaderProfileType ProfileType = EAvaShaderProfileType::Default;
 
-	UPROPERTY(EditAnywhere, Category="Material Cache", meta=(EditCondition="bEnabled"))
+	/** Specific Shader Combinations to use */
+	UPROPERTY(EditAnywhere, Category="Material Cache", meta=(EditCondition="ProfileType==EAvaShaderProfileType::Specific"))
 	TArray<FAvaShaderCombination> ShaderCombinations;
 };
 
@@ -71,6 +90,16 @@ public:
 
 	UAvaMaterialCacheSettings();
 
+	FName GetRealtimeProfile() const
+	{
+		return RealtimeProfile;
+	}
+
+	FName GetOfflineProfile() const
+	{
+		return OfflineProfile;
+	}
+
 	/** Finds the resolved shader profile for a given profile name */
 	const UE::Ava::FResolvedShaderProfile* FindResolvedShaderProfile(FName InProfileName) const;
 
@@ -78,6 +107,9 @@ protected:
 	//~ Begin UObject
 	virtual void PostInitProperties() override;
 	virtual void PostReloadConfig(FProperty* InPropertyThatWasLoaded) override;
+#if WITH_EDITOR
+	virtual void PostEditChangeProperty(FPropertyChangedEvent& InPropertyChangedEvent) override;
+#endif
 	//~ End UObject
 
 private:
diff --git a/Engine/Plugins/VirtualProduction/Avalanche/Source/AvalancheMaterial/Private/AvaMaterialCacheHelper.cpp b/Engine/Plugins/VirtualProduction/Avalanche/Source/AvalancheMaterial/Private/AvaMaterialCacheHelper.cpp
--- a/Engine/Plugins/VirtualProduction/Avalanche/Source/AvalancheMaterial/Private/AvaMaterialCacheHelper.cpp
+++ b/Engine/Plugins/VirtualProduction/Avalanche/Source/AvalancheMaterial/Private/AvaMaterialCacheHelper.cpp
@@ -27,7 +27,7 @@ FAutoConsoleCommand MaterialCacheDumpMaterials(TEXT("MotionDesign.MaterialCache.
 		}));
 
 TAutoConsoleVariable<float> CVarMaterialCacheFrameTimeStarveThreshold(TEXT("MotionDesign.MaterialCache.FrameTimeStarveThreshold")
-	, 1.f / 10.f
+	, 1.f / 30.f
 	, TEXT("The frame time threshold for starving tick. If the current duration of the frame surpasses this, FMaterialCacheHelper::Tick will be skipped and starve.")
 	);
 
@@ -63,6 +63,71 @@ void ForEachMaterialResource(TNotNull<UMaterialInterface*> InMaterial, TFunction
 }
 #endif
 
+/** Finds the shader profile matching the given name */
+const FResolvedShaderProfile* FindShaderProfile(FName InShaderProfile)
+{
+	if (const UAvaMaterialCacheSettings* MaterialCacheSettings = GetDefault<UAvaMaterialCacheSettings>())
+	{
+		return MaterialCacheSettings->FindResolvedShaderProfile(InShaderProfile);
+	}
+	UE_LOG(LogAvaMaterialCache, Verbose, TEXT("Shader Profile '%s' not found in Material Cache Settings"), *InShaderProfile.ToString());
+	return nullptr;
+}
+
+/** Performs shader caching of the given material with the given shader profile */
+void CacheMaterial(TNotNull<UMaterialInterface*> InMaterial, TNotNull<const FResolvedShaderProfile*> InShaderProfile)
+{
+	UE_LOG(LogAvaMaterialCache, VeryVerbose, TEXT("Caching Material '%s'"), *InMaterial->GetFullName());
+
+#if WITH_EDITOR
+	if (InShaderProfile->ProfileType == EAvaShaderProfileType::Specific)
+	{
+		InMaterial->CacheShaders(EMaterialShaderPrecompileMode::None);
+
+		ForEachMaterialResource(InMaterial,
+			[InShaderProfile](FMaterialResource* InMaterialResource)->bool
+			{
+				InMaterialResource->CacheGivenTypes(InShaderProfile->VertexFactoryTypes, InShaderProfile->PipelineTypes, InShaderProfile->ShaderTypes);
+				return true; // continue
+			});
+
+		return;
+	}
+#endif
+
+	InMaterial->CacheShaders(EMaterialShaderPrecompileMode::Background);
+}
+
+/** Checks whether the material is complete */
+bool IsMaterialComplete(TNotNull<UMaterialInterface*> InMaterial, TNotNull<const FResolvedShaderProfile*> InShaderProfile)
+{
+#if WITH_EDITOR
+	if (InShaderProfile->ProfileType == EAvaShaderProfileType::Specific)
+	{
+		bool bIsComplete = true;
+
+		ForEachMaterialResource(InMaterial, 
+			[&bIsComplete](FMaterialResource* InMaterialResource)->bool
+			{
+				if (InMaterialResource->IsCachingShaders())
+				{
+					bIsComplete = false;
+					return false; // break
+				}
+				return true; // continue
+			});
+
+		return bIsComplete;
+	}
+#endif
+	return InMaterial->IsComplete();
+}
+
+inline double GetCurrentTime()
+{
+	return FPlatformTime::Seconds();
+}
+
 } // UE::Ava::Private
 
 FMaterialCacheHelper& FMaterialCacheHelper::Get()
@@ -73,16 +138,18 @@ FMaterialCacheHelper& FMaterialCacheHelper::Get()
 
 bool FMaterialCacheHelper::RequestCacheMaterials(const UObject* InObject, FName InShaderProfile)
 {
-	if (!InObject)
+	if (!InObject || Objects.Contains(InObject))
 	{
 		return false;
 	}
 
-	if (!GatherMaterialData(InObject, InShaderProfile))
+	FObjectData ObjectData;
+	if (!GatherMaterialData(InObject, ObjectData, InShaderProfile))
 	{
 		return false;
 	}
 
+	Objects.Add(InObject, MoveTemp(ObjectData));
 	RegisterTick();
 	return true;
 }
@@ -98,11 +165,11 @@ void FMaterialCacheHelper::Tick()
 
 	TRACE_CPUPROFILER_EVENT_SCOPE(FMaterialCacheHelper::Tick);
 
-	// Duration of current frame
-	const float CurrentFrameTime = (FPlatformTime::Seconds() - BeginFrameTime);
+	// Duration of the frame so far prior to ticking
+	PreTickFrameDuration = (Private::GetCurrentTime() - BeginFrameTime);
 
 	// If the current frame time is already past the threshold, tick should be skipped unless it has been starving for a while already.
-	if (CurrentFrameTime >= Private::CVarMaterialCacheFrameTimeStarveThreshold.GetValueOnGameThread())
+	if (PreTickFrameDuration >= Private::CVarMaterialCacheFrameTimeStarveThreshold.GetValueOnGameThread())
 	{
 		if (FramesStarved < Private::CVarMaterialCacheTickMaxStarveFrames.GetValueOnGameThread())
 		{
@@ -115,7 +182,9 @@ void FMaterialCacheHelper::Tick()
 	FramesStarved = 0;
 	RemainingTimeBudget = Private::CVarMaterialCacheTickTimeBudget.GetValueOnGameThread();
 
-	ProcessMaterials();
+	ProcessPendingObjects();
+	ProcessMaterialsPendingCache();
+	ProcessMaterialsCaching();
 
 	if (IsCaching())
 	{
@@ -130,25 +199,40 @@ void FMaterialCacheHelper::Tick()
 
 bool FMaterialCacheHelper::IsCaching() const
 {
-	return !Materials.IsEmpty();
+	// If there is at least one element in either map, return true
+	if (!MaterialsPendingCache.IsEmpty() || !MaterialsCaching.IsEmpty())
+	{
+		return true;
+	}
+
+	// If there is one pending subobject still loading to cache, return true
+	for (const TPair<FObjectKey, FObjectData>& Object : Objects)
+	{
+		if (!Object.Value.PendingSubobjects.IsEmpty())
+		{
+			return true;
+		}
+	}
+	return false;
 }
 
 bool FMaterialCacheHelper::IsCaching(const UObject* InObject) const
 {
-	if (!IsCaching())
+	const FObjectData* ObjectData = Objects.Find(InObject);
+	if (!ObjectData)
 	{
 		return false;
 	}
 
-	const FObjectData* ObjectData = Objects.Find(InObject);
-	if (!ObjectData)
+	// There are pending subobjects waiting on completion for material gathering
+	if (!ObjectData->PendingSubobjects.IsEmpty())
 	{
-		return false;
+		return true;
 	}
 
 	for (const FMaterialKey& MaterialKey : ObjectData->IncompleteMaterialKeys)
 	{
-		if (Materials.Contains(MaterialKey))
+		if (IsCachingMaterial(MaterialKey))
 		{
 			return true;
 		}
@@ -158,27 +242,31 @@ bool FMaterialCacheHelper::IsCaching(const UObject* InObject) const
 
 void FMaterialCacheHelper::DumpMaterials()
 {
-	constexpr const TCHAR* CacheStatusDisplays[] =
+	auto PrintLog = [](const TPair<FMaterialKey, FMaterialData>& InPair)
 		{
-			TEXT("Waiting"),
-			TEXT("Caching"),
+			const FMaterialData& MaterialData = InPair.Value;
+			UE_LOG(LogAvaMaterialCache, Log, TEXT("\t%s"), MaterialData.Material ? *MaterialData.Material->GetFullName() : TEXT("(invalid)"));
 		};
 
-	UE_LOG(LogAvaMaterialCache, Log, TEXT("----- DUMP MATERIALS--------------------------------------------"));
-	for (const TPair<FMaterialKey, FMaterialData>& Pair : Materials)
+	UE_LOG(LogAvaMaterialCache, Log, TEXT("--- Materials Pending Cache ------------------------------------"));
+	for (const TPair<FMaterialKey, FMaterialData>& Pair : MaterialsPendingCache)
 	{
-		const FMaterialData& MaterialData = Pair.Value;
-
-		UE_LOG(LogAvaMaterialCache, Log, TEXT("\t'%s' --> [%s]")
-			, MaterialData.Material ? *MaterialData.Material->GetFullName() : TEXT("(invalid)")
-			, CacheStatusDisplays[static_cast<uint8>(MaterialData.CacheStatus)]);
+		PrintLog(Pair);
+	}
+	UE_LOG(LogAvaMaterialCache, Log, TEXT("--- Materials Caching ------------------------------------------"));
+	for (const TPair<FMaterialKey, FMaterialData>& Pair : MaterialsCaching)
+	{
+		PrintLog(Pair);
 	}
-	UE_LOG(LogAvaMaterialCache, Log, TEXT("----------------------------------------------------------------"));
 }
 
 void FMaterialCacheHelper::AddReferencedObjects(FReferenceCollector& InCollector)
 {
-	for (TPair<FMaterialKey, FMaterialData>& Pair : Materials)
+	for (TPair<FMaterialKey, FMaterialData>& Pair : MaterialsPendingCache)
+	{
+		InCollector.AddReferencedObject(Pair.Value.Material);
+	}
+	for (TPair<FMaterialKey, FMaterialData>& Pair : MaterialsCaching)
 	{
 		InCollector.AddReferencedObject(Pair.Value.Material);
 	}
@@ -189,20 +277,31 @@ FString FMaterialCacheHelper::GetReferencerName() const
 	return TEXT("UE::Ava::FMaterialCacheHelper");
 }
 
+bool FMaterialCacheHelper::IsCachingMaterial(const FMaterialKey& InMaterialKey) const
+{
+	return MaterialsPendingCache.Contains(InMaterialKey) || MaterialsCaching.Contains(InMaterialKey);
+}
+
 void FMaterialCacheHelper::Reset()
 {
-	Materials.Empty();
+	MaterialsPendingCache.Empty();
+	MaterialsCaching.Empty();
 	Objects.Empty();
 	UnregisterTick();
 }
 
 void FMaterialCacheHelper::RegisterTick()
 {
+	bResetRequested = false;
+
 	if (!BeginFrameHandle.IsValid())
 	{
+		// It is likely RegisterTick is called between OnBeginFrame and OnEndFrame delegate broadcasts.
+		// This will cause BeginFrameTime to have its last value in the first tick (e.g. 0.0, or an old begin frame's time) and likely starve tick for that one frame.
+		// And this is ok. An alternative could be to use a timer that is set at BeginFrame always.  
 		BeginFrameHandle = FCoreDelegates::OnBeginFrame.AddSPLambda(this, [this]
 			{
-				BeginFrameTime = FPlatformTime::Seconds();
+				BeginFrameTime = Private::GetCurrentTime();
 			});
 	}
 
@@ -227,21 +326,76 @@ void FMaterialCacheHelper::UnregisterTick()
 	}
 }
 
-void FMaterialCacheHelper::ProcessMaterials()
+void FMaterialCacheHelper::ProcessPendingObjects()
 {
-	if (RemainingTimeBudget <= 0.f || Materials.IsEmpty())
+	if (RemainingTimeBudget <= 0.f || Objects.IsEmpty())
 	{
 		return;
 	}
 
-	TRACE_CPUPROFILER_EVENT_SCOPE(FMaterialCacheHelper::ProcessMaterials);
+	// Time when budget runs out
+	const double EndBudgetTime = Private::GetCurrentTime() + RemainingTimeBudget;
+	bool bTimeBudgetExceeded = false;
+
+	// Check the pending status every 100ms (0.1s) to give it some time to load
+	constexpr float PendingSubobjectCheckTime = 0.1f;
+
+	for (TPair<FObjectKey, FObjectData>& Object : Objects)
+	{
+		Object.Value.PendingSubobjectElapsedTime += PreTickFrameDuration;
+		if (Object.Value.PendingSubobjectElapsedTime < PendingSubobjectCheckTime)
+		{
+			continue;
+		}
+		Object.Value.PendingSubobjectElapsedTime = 0.f;
+
+		// All pending subobjects might be ready so move them into a temp set.
+		// This also allows adding elements to the existing set.
+		// Those that are still not gathered pending load will be re-added to the set.
+		TSet<TObjectPtr<const UObject>> PendingSubobjects = MoveTemp(Object.Value.PendingSubobjects);
+		Object.Value.PendingSubobjects.Empty();
+
+		for (TSet<TObjectPtr<const UObject>>::TIterator Iter(PendingSubobjects); Iter; ++Iter)
+		{
+			GatherMaterialData(*Iter, Object.Value, Object.Value.ShaderProfile);
+			Iter.RemoveCurrent();
+
+			// Check if time is past end budget time
+			bTimeBudgetExceeded = Private::GetCurrentTime() >= EndBudgetTime;
+			if (bTimeBudgetExceeded)
+			{
+				UE_LOG(LogAvaMaterialCache, Verbose, TEXT("FMaterialCacheHelper::Tick ran out of budget in ProcessPendingObjects. Time Budget: %f"), RemainingTimeBudget);
+				break;
+			}
+		}
+
+		// Add pending subobjects that could not be processed due to time budget back
+		Object.Value.PendingSubobjects.Append(MoveTemp(PendingSubobjects));
+		if (bTimeBudgetExceeded)
+		{
+			break;
+		}
+	}
+
+	RemainingTimeBudget = EndBudgetTime - Private::GetCurrentTime();
+}
+
+void FMaterialCacheHelper::ProcessMaterialsPendingCache()
+{
+	if (RemainingTimeBudget <= 0.f)
+	{
+		return;
+	}
+
+	TRACE_CPUPROFILER_EVENT_SCOPE(FMaterialCacheHelper::ProcessMaterialsPendingCache);
 
 	// Time when budget runs out
-	const double EndBudgetTime = FPlatformTime::Seconds() + RemainingTimeBudget;
+	const double EndBudgetTime = Private::GetCurrentTime() + RemainingTimeBudget;
 
+	// Material update context. Only allocated when needed
 	TOptional<FMaterialUpdateContext> UpdateContext;
 
-	for (TMap<FMaterialKey, FMaterialData>::TIterator Iter(Materials); Iter; ++Iter)
+	for (TMap<FMaterialKey, FMaterialData>::TIterator Iter(MaterialsPendingCache); Iter; ++Iter)
 	{
 		const FMaterialKey& MaterialKey = Iter.Key();
 		FMaterialData& MaterialData = Iter.Value();
@@ -252,90 +406,72 @@ void FMaterialCacheHelper::ProcessMaterials()
 			continue;
 		}
 
-		switch (MaterialData.CacheStatus)
+		const FResolvedShaderProfile* ShaderProfile = Private::FindShaderProfile(MaterialKey.ShaderProfile);
+		if (!ShaderProfile || ShaderProfile->ProfileType == EAvaShaderProfileType::SkipCaching)
 		{
-		case EMaterialCacheStatus::Waiting:
-			if (!UpdateContext.IsSet())
-			{
-				UpdateContext.Emplace(0);
-			}
-			UpdateContext->AddMaterialInterface(MaterialData.Material);
-			CacheMaterial(MaterialData.Material, MaterialKey.ShaderProfile);
-			MaterialData.CacheStatus = EMaterialCacheStatus::Caching;
-			break;
+			Iter.RemoveCurrent();
+			continue;
+		}
 
-		case EMaterialCacheStatus::Caching:
-			if (IsMaterialComplete(MaterialData.Material, MaterialKey.ShaderProfile))
-			{
-				Iter.RemoveCurrent();
-			}
-			break;
+		if (!UpdateContext.IsSet())
+		{
+			UpdateContext.Emplace(0);
 		}
+		UpdateContext->AddMaterialInterface(MaterialData.Material);
+		Private::CacheMaterial(MaterialData.Material, ShaderProfile);
+
+		MaterialsCaching.Add(MaterialKey, MoveTemp(MaterialData));
+		Iter.RemoveCurrent();
 
 		// Check if time is past end budget time
-		if (FPlatformTime::Seconds() >= EndBudgetTime)
+		// This is only checked for materials waiting, since the other operations are fast.
+		if (Private::GetCurrentTime() >= EndBudgetTime)
 		{
+			UE_LOG(LogAvaMaterialCache, Verbose, TEXT("FMaterialCacheHelper::Tick ran out of budget in ProcessMaterials. Time Budget: %f"), RemainingTimeBudget);
 			break;
 		}
 	}
 
-	RemainingTimeBudget = EndBudgetTime - FPlatformTime::Seconds();
+	RemainingTimeBudget = EndBudgetTime - Private::GetCurrentTime();
 }
 
-const FResolvedShaderProfile* FMaterialCacheHelper::FindShaderProfile(FName InShaderProfile) const
+void FMaterialCacheHelper::ProcessMaterialsCaching()
 {
-	if (const UAvaMaterialCacheSettings* MaterialCacheSettings = GetDefault<UAvaMaterialCacheSettings>())
+	if (RemainingTimeBudget <= 0.f)
 	{
-		return MaterialCacheSettings->FindResolvedShaderProfile(InShaderProfile);
+		return;
 	}
-	return nullptr;
-}
 
-void FMaterialCacheHelper::CacheMaterial(TNotNull<UMaterialInterface*> InMaterial, FName InShaderProfile)
-{
-#if WITH_EDITOR
-	// CacheGivenTypes is EditorOnly
-	if (const FResolvedShaderProfile* const ShaderProfile = FindShaderProfile(InShaderProfile))
-	{
-		InMaterial->CacheShaders(EMaterialShaderPrecompileMode::None);
+	TRACE_CPUPROFILER_EVENT_SCOPE(FMaterialCacheHelper::ProcessCachingMaterials);
 
-		Private::ForEachMaterialResource(InMaterial,
-			[ShaderProfile](FMaterialResource* InMaterialResource)->bool
-			{
-				
-				InMaterialResource->CacheGivenTypes(ShaderProfile->VertexFactoryTypes, ShaderProfile->PipelineTypes, ShaderProfile->ShaderTypes);
-				return true; // continue
-			});
-	}
-	else
-#endif
-	{
-		InMaterial->CacheShaders(EMaterialShaderPrecompileMode::Background);
-	}
-}
+	// Time when budget runs out
+	const double StartTime = Private::GetCurrentTime();
 
-bool FMaterialCacheHelper::IsMaterialComplete(TNotNull<UMaterialInterface*> InMaterial, FName InShaderProfile) const
-{
-#if WITH_EDITOR
-	if (const FResolvedShaderProfile* const ShaderProfile = FindShaderProfile(InShaderProfile))
+	for (TMap<FMaterialKey, FMaterialData>::TIterator Iter(MaterialsCaching); Iter; ++Iter)
 	{
-		bool bIsComplete = true;
+		const FMaterialKey& MaterialKey = Iter.Key();
+		FMaterialData& MaterialData = Iter.Value();
 
-		Private::ForEachMaterialResource(InMaterial, 
-			[&bIsComplete](FMaterialResource* InMaterialResource)->bool
-			{
-				if (InMaterialResource->IsCachingShaders())
-				{
-					bIsComplete = false;
-					return false; // break
-				}
-				return true; // continue
-			});
+		if (!MaterialData.Material)
+		{
+			Iter.RemoveCurrent();
+			continue;
+		}
 
-		return bIsComplete;
+		const FResolvedShaderProfile* ShaderProfile = Private::FindShaderProfile(MaterialKey.ShaderProfile);
+		if (!ShaderProfile)
+		{
+			Iter.RemoveCurrent();
+			continue;
+		}
+
+		if (Private::IsMaterialComplete(MaterialData.Material, ShaderProfile))
+		{
+			Iter.RemoveCurrent();
+		}
 	}
-#endif
-	return InMaterial->IsComplete();
+
+	RemainingTimeBudget -= (Private::GetCurrentTime() - StartTime);
 }
 
 void FMaterialCacheHelper::RemoveCompletedMaterials()
@@ -345,20 +481,20 @@ void FMaterialCacheHelper::RemoveCompletedMaterials()
 		return;
 	}
 
-	const double StartTime = FPlatformTime::Seconds();
+	const double StartTime = Private::GetCurrentTime();
 
 	for (TPair<FObjectKey, FObjectData>& Object : Objects)
 	{
 		for (TSet<FMaterialKey>::TIterator MaterialKeyIter(Object.Value.IncompleteMaterialKeys); MaterialKeyIter; ++MaterialKeyIter)
 		{
-			if (!Materials.Contains(*MaterialKeyIter))
+			if (!IsCachingMaterial(*MaterialKeyIter))
 			{
 				MaterialKeyIter.RemoveCurrent();
 			}
 		}
 	}
 
-	RemainingTimeBudget -= (FPlatformTime::Seconds() - StartTime);
+	RemainingTimeBudget -= (Private::GetCurrentTime() - StartTime);
 }
 
 void FMaterialCacheHelper::ProcessShadersCompiling()
@@ -370,17 +506,17 @@ void FMaterialCacheHelper::ProcessShadersCompiling()
 
 	TRACE_CPUPROFILER_EVENT_SCOPE(FMaterialCacheHelper::ProcessShadersCompiling);
 
-	const double StartTime = FPlatformTime::Seconds();
+	const double StartTime = Private::GetCurrentTime();
 
 	constexpr bool bBlockOnGlobalShaderCompletion = false;
 	GShaderCompilingManager->ProcessAsyncResults(RemainingTimeBudget, bBlockOnGlobalShaderCompletion);
 
-	RemainingTimeBudget -= (FPlatformTime::Seconds() - StartTime);
+	RemainingTimeBudget -= (Private::GetCurrentTime() - StartTime);
 }
 
-bool FMaterialCacheHelper::GatherMaterialData(const UObject* InObject, FName InShaderProfile)
+bool FMaterialCacheHelper::GatherMaterialData(const UObject* InObject, FObjectData& InObjectData, FName InShaderProfile)
 {
-	if (!InObject || Objects.Contains(InObject))
+	if (!InObject)
 	{
 		return false;
 	}
@@ -390,28 +526,40 @@ bool FMaterialCacheHelper::GatherMaterialData(const UObject* InObject, FName InS
 	const FMaterialBridge* MaterialBridge = FMaterialBridgeRegistry::Get().GetMaterialBridgeForObject(InObject);
 	if (!MaterialBridge)
 	{
-		UE_LOG(LogAvaMaterialCache, Verbose, TEXT("Object '%s' does not have a valid bridge to cache materials!"), *GetNameSafe(InObject));
+		UE_LOG(LogAvaMaterialCache, VeryVerbose, TEXT("Object '%s' (Class: %s) does not have a valid bridge to cache materials!")
+			, *InObject->GetName()
+			, *GetNameSafe(InObject->GetClass()));
 		return false;
 	}
 
-	FObjectData& ObjectData = Objects.Add(InObject);
+	InObjectData.ShaderProfile = InShaderProfile;
 
 	auto AddIncompleteMaterial =
-		[&Materials = Materials, &ObjectData, InShaderProfile](TNotNull<UMaterialInterface*> InMaterial)
+		[&MaterialsPendingCache = MaterialsPendingCache, &InObjectData, InShaderProfile](TNotNull<UMaterialInterface*> InMaterial)
 		{
 			FMaterialKey MaterialKey;
 			MaterialKey.Material = InMaterial;
 			MaterialKey.ShaderProfile = InShaderProfile;
 
-			FMaterialData* MaterialData = Materials.Find(MaterialKey);
-			if (!MaterialData)
+			FMaterialData& MaterialData = MaterialsPendingCache.FindOrAdd(MaterialKey);
+			MaterialData.Material = InMaterial;
+			InObjectData.IncompleteMaterialKeys.Add(MaterialKey);
+		};
+
+	// Options to not wait for the container to have materials available and instead get a callback
+	FMaterialBridge::FVisitOptions VisitOptions;
+	VisitOptions.bTryWaitOnContainerCompletion = false;
+	VisitOptions.OnContainerPendingCompletion.BindLambda([&InObjectData](const FMaterialBridgeVisitContext& InContext)
+		{
+			if (const UObject* Container = InContext.Container.GetPtr<const UObject>())
 			{
-				MaterialData = &Materials.Add(MaterialKey);
-				MaterialData->Material = InMaterial;
-				MaterialData->CacheStatus = EMaterialCacheStatus::Waiting;
+				InObjectData.PendingSubobjects.Add(Container);
 			}
-			ObjectData.IncompleteMaterialKeys.Add(MaterialKey);
-		};
+			else
+			{
+				ensureMsgf(0, TEXT("Container of type '%s' not supported for deferring material gather"), *GetNameSafe(InContext.Container.GetStruct()));
+			}
+		});
 
 	// Gather all materials for object
 	MaterialBridge->VisitMaterials(FMaterialBridgeVisitContext(InObject),
@@ -433,9 +581,10 @@ bool FMaterialCacheHelper::GatherMaterialData(const UObject* InObject, FName InS
 			{
 				AddIncompleteMaterial(InMaterial);
 			}
-		});
+		}
+		, VisitOptions);
 
 	return true;
 }
 
-} // UE::Ava
\ No newline at end of file
+} // UE::Ava
diff --git a/Engine/Plugins/VirtualProduction/Avalanche/Source/AvalancheMaterial/Private/AvaMaterialCacheSettings.cpp b/Engine/Plugins/VirtualProduction/Avalanche/Source/AvalancheMaterial/Private/AvaMaterialCacheSettings.cpp
--- a/Engine/Plugins/VirtualProduction/Avalanche/Source/AvalancheMaterial/Private/AvaMaterialCacheSettings.cpp
+++ b/Engine/Plugins/VirtualProduction/Avalanche/Source/AvalancheMaterial/Private/AvaMaterialCacheSettings.cpp
@@ -103,27 +103,41 @@ void UAvaMaterialCacheSettings::PostReloadConfig(FProperty* InPropertyThatWasLoa
 	ResolveShaderProfiles();
 }
 
+#if WITH_EDITOR
+void UAvaMaterialCacheSettings::PostEditChangeProperty(FPropertyChangedEvent& InPropertyChangedEvent)
+{
+	Super::PostEditChangeProperty(InPropertyChangedEvent);
+
+	if (InPropertyChangedEvent.GetMemberPropertyName() == GET_MEMBER_NAME_CHECKED(UAvaMaterialCacheSettings, ShaderProfiles))
+	{
+		ResolveShaderProfiles();
+	}
+}
+#endif
+
 void UAvaMaterialCacheSettings::RegisterDefaultShaderProfile()
 {
 	const TArray<FAvaShaderCombination> DefaultShaderCombinations = UE::Ava::Private::MakeDefaultShaderCombinations();
 
 	ShaderProfiles.Add(
-		{ 
+		{
 			.ProfileName = UAvaMaterialCacheSettings::DefaultRealtimeProfile,
-			.bEnabled = true,
+			.ProfileType = EAvaShaderProfileType::Specific,
 			.ShaderCombinations = DefaultShaderCombinations,
 		});
 
 	ShaderProfiles.Add(
-		{ 
+		{
 			.ProfileName = UAvaMaterialCacheSettings::DefaultOfflineProfile,
-			.bEnabled = true,
+			.ProfileType = EAvaShaderProfileType::Specific,
 			.ShaderCombinations = DefaultShaderCombinations,
 		});
 }
 
 void UAvaMaterialCacheSettings::ResolveShaderProfiles()
 {
+	UE_LOG(LogAvaMaterialCache, Verbose, TEXT("Resolving Shader Profiles"));
+
 	ResolvedShaderProfiles.Empty(ShaderProfiles.Num());
 
 	TMap<FName, FHashedName> HashedNames;
@@ -135,12 +149,6 @@ void UAvaMaterialCacheSettings::ResolveShaderProfiles()
 
 void UAvaMaterialCacheSettings::ResolveShaderProfile(const FAvaShaderProfile& InShaderProfile, TMap<FName, FHashedName>& InOutHashedNames)
 {
-	if (!InShaderProfile.bEnabled)
-	{
-		UE_LOG(LogAvaMaterialCache, Verbose, TEXT("Shader Profile '%s' not resolved as it's not enabled."), *InShaderProfile.ProfileName.ToString());
-		return;
-	}
-
 	if (ResolvedShaderProfiles.Contains(InShaderProfile.ProfileName))
 	{
 		UE_LOG(LogAvaMaterialCache, Warning, TEXT("Shader Profile '%s' already exists. New entry will be ignored."), *InShaderProfile.ProfileName.ToString());
@@ -148,6 +156,7 @@ void UAvaMaterialCacheSettings::ResolveShaderProfile(const FAvaShaderProfile& In
 	}
 
 	UE::Ava::FResolvedShaderProfile& ResolvedShaderProfile = ResolvedShaderProfiles.Add(InShaderProfile.ProfileName);
+	ResolvedShaderProfile.ProfileType = InShaderProfile.ProfileType;
 
 	// Keep track of FHashedName instances to avoid recomputation
 	auto GetHashedName = [&InOutHashedNames](FName InName)
diff --git a/Engine/Plugins/VirtualProduction/Avalanche/Source/AvalancheMedia/Internal/Playback/AvaPlaybackUtils.h b/Engine/Plugins/VirtualProduction/Avalanche/Source/AvalancheMedia/Internal/Playback/AvaPlaybackUtils.h
--- a/Engine/Plugins/VirtualProduction/Avalanche/Source/AvalancheMedia/Internal/Playback/AvaPlaybackUtils.h
+++ b/Engine/Plugins/VirtualProduction/Avalanche/Source/AvalancheMedia/Internal/Playback/AvaPlaybackUtils.h
@@ -89,7 +89,7 @@ namespace UE::AvaPlayback::Utils
 	class FAsyncAssetLoader : public FGCObject, public TSharedFromThis<FAsyncAssetLoader>
 	{
 	public:
-		void ProcessLoadedObject(UObject* InObject);
+		void ProcessLoadedObject(UObject* InObject, FName InShaderProfile);
 
 		/**
 		 * @brief Issue the async load command. Returns immediately.
diff --git a/Engine/Plugins/VirtualProduction/Avalanche/Source/AvalancheMedia/Private/Playable/Playables/AvaPlayableLevelStreaming.cpp b/Engine/Plugins/VirtualProduction/Avalanche/Source/AvalancheMedia/Private/Playable/Playables/AvaPlayableLevelStreaming.cpp
--- a/Engine/Plugins/VirtualProduction/Avalanche/Source/AvalancheMedia/Private/Playable/Playables/AvaPlayableLevelStreaming.cpp
+++ b/Engine/Plugins/VirtualProduction/Avalanche/Source/AvalancheMedia/Private/Playable/Playables/AvaPlayableLevelStreaming.cpp
@@ -31,7 +31,7 @@ namespace UE::AvaMedia::LevelStreamingPlayable::Private
 {
 	static TAutoConsoleVariable<bool> CVarRuntimePrecacheShaders(
 		TEXT("MotionDesignPlayable.RuntimePrecacheShaders"),
-		true,
+		false, // False by default to avoid the extra cost of having to try cache materials at runtime/playtime.
 		TEXT("Caches any remaining shaders from materials pending completion found in the playable while preparing to start play.")
 	);
 
@@ -392,10 +392,10 @@ void UAvaPlayableLevelStreaming::UpdatePlayableStatus(ELevelStreamingState InNew
 
 bool UAvaPlayableLevelStreaming::IsReadyToPlay() const
 {
-	TRACE_CPUPROFILER_EVENT_SCOPE(UAvaPlayableLevelStreaming::IsReadyToPlay)
-
-	using namespace UE::AvaMedia::LevelStreamingPlayable;
-
+	if (bCachingShaders)
+	{
+		return false;
+	}
 	for (const UAvaPlayableLevelStreaming* SubPlayable : SubPlayables)
 	{
 		if (SubPlayable && !SubPlayable->IsReadyToPlay())
@@ -403,27 +403,7 @@ bool UAvaPlayableLevelStreaming::IsReadyToPlay() const
 			return false;
 		}
 	}
-
-	if (!Private::CVarRuntimePrecacheShaders.GetValueOnGameThread())
-	{
-		// Not allowed to cache. Ready to play
-		return true;
-	}
-
-	if (!LevelStreaming)
-	{
-		return false;
-	}
-
-	const ULevel* const Level = LevelStreaming->GetLoadedLevel();
-	if (!Level)
-	{
-		return false;
-	}
-
-	// Check whether the compilation has finished
-	const bool bIsCachingShaders = CacheMaterials(Level, UAvaMaterialCacheSettings::DefaultRealtimeProfile);
-	return !bIsCachingShaders;
+	return true;
 }
 
 IAvaSceneInterface* UAvaPlayableLevelStreaming::GetSceneInterface() const
@@ -524,16 +504,6 @@ bool UAvaPlayableLevelStreaming::LoadLevel(const TSoftObjectPtr<UWorld>& InSourc
 	return true;
 }
 
-bool UAvaPlayableLevelStreaming::CacheMaterials(TNotNull<const ULevel*> InLevel, FName InShaderProfile) const
-{
-	TRACE_CPUPROFILER_EVENT_SCOPE(UAvaPlayableLevelStreaming::CacheMaterials)
-
-	UE::Ava::FMaterialCacheHelper& MaterialCacheHelper = UE::Ava::FMaterialCacheHelper::Get();
-
-	MaterialCacheHelper.RequestCacheMaterials(InLevel, InShaderProfile);
-	return MaterialCacheHelper.IsCaching(InLevel);
-}
-
 void UAvaPlayableLevelStreaming::HandleTransitionEvent(UAvaPlayable* InPlayable, UAvaPlayableTransition* InTransition, EAvaPlayableTransitionEventFlags InTransitionFlags)
 {
 	if (InPlayable == this && EnumHasAnyFlags(InTransitionFlags, EAvaPlayableTransitionEventFlags::ShowPlayable))
@@ -643,31 +613,48 @@ void UAvaPlayableLevelStreaming::OnLevelStreamingStateChanged_Synchronized(ELeve
 		// Resolve the ava scene for the other operations.
 		ResolveScene(Level);
 
-		if (bCacheShadersOnLoad)
+		FName ShaderProfile;
+		bool bShouldCacheShaders = false;
+
+		// Caching happens in two scenarios:
+		// - When the user has requested shaders to be cached on load and the level has just been loaded (not visible yet).
+		// - When the RuntimePrecacheShaders cvar is set to true and the level just became visible, prior to any transition logic starting, etc.
+		// Both scenarios can happen for one playable, and it's expected since these two states can be at completely different times.
+		// However, the shader profile will be different in these two scenarios
+		if (InNewState == ELevelStreamingState::LoadedNotVisible && bCacheShadersOnLoad)
 		{
-			// If there was more than one material that required caching,
-			// kick off a ticker that incrementally caches materials,
-			// until there's no materials to cache and all shaders have completed compilation
-			// NOTE: This uses the default offline shader profile. The real time profile is used in UAvaPlayableLevelStreaming::IsReadyToPlay
-			if (CacheMaterials(Level, UAvaMaterialCacheSettings::DefaultOfflineProfile))
-			{
-				bCachingShadersOnLoad = true;
+			ShaderProfile = GetDefault<UAvaMaterialCacheSettings>()->GetOfflineProfile();
+			bShouldCacheShaders = true;
+		}
+		else if (InNewState == ELevelStreamingState::LoadedVisible && UE::AvaMedia::LevelStreamingPlayable::Private::CVarRuntimePrecacheShaders.GetValueOnGameThread())
+		{
+			ShaderProfile = GetDefault<UAvaMaterialCacheSettings>()->GetRealtimeProfile();
+			bShouldCacheShaders = true;
+		}
 
+		// Cache materials for the level incrementally, only if not already caching shaders
+		if (bShouldCacheShaders && !bCachingShaders)
+		{
+			UE::Ava::FMaterialCacheHelper& MaterialCacheHelper = UE::Ava::FMaterialCacheHelper::Get();
+			MaterialCacheHelper.RequestCacheMaterials(Level, ShaderProfile);
+			bCachingShaders = MaterialCacheHelper.IsCaching(Level);
+
+			if (bCachingShaders)
+			{
 				FTSTicker::GetCoreTicker().AddTicker(FTickerDelegate::CreateWeakLambda(this,
 					[this](float)
 					{
 						const ULevel* const Level = LevelStreaming ? LevelStreaming->GetLoadedLevel() : nullptr;
 						if (Level)
 						{
-							bCachingShadersOnLoad = CacheMaterials(Level, UAvaMaterialCacheSettings::DefaultOfflineProfile);
+							bCachingShaders = UE::Ava::FMaterialCacheHelper::Get().IsCaching(Level);
 						}
 						else
 						{
-							bCachingShadersOnLoad = false;
+							bCachingShaders = false;
 						}
-
 						// Continue ticking until caching shaders stops
-						return bCachingShadersOnLoad;
+						return bCachingShaders;
 					}));
 			}
 		}
diff --git a/Engine/Plugins/VirtualProduction/Avalanche/Source/AvalancheMedia/Private/Playable/Playables/AvaPlayableLevelStreaming.h b/Engine/Plugins/VirtualProduction/Avalanche/Source/AvalancheMedia/Private/Playable/Playables/AvaPlayableLevelStreaming.h
--- a/Engine/Plugins/VirtualProduction/Avalanche/Source/AvalancheMedia/Private/Playable/Playables/AvaPlayableLevelStreaming.h
+++ b/Engine/Plugins/VirtualProduction/Avalanche/Source/AvalancheMedia/Private/Playable/Playables/AvaPlayableLevelStreaming.h
@@ -49,14 +49,6 @@ public:
 protected:
 	bool LoadLevel(const TSoftObjectPtr<UWorld>& InSourceLevel, const FTransform& InTransform, bool bInInitiallyVisible);
 
-	/**
-	 * Caches the shaders for the given level.
-	 * @param InLevel the level to cache shaders for
-	 * @param InShaderProfile the shader profile to use for caching
-	 * @return true if caching is taking place. false otherwise. 
-	 */
-	bool CacheMaterials(TNotNull<const ULevel*> InLevel, FName InShaderProfile) const;
-
 	void HandleTransitionEvent(UAvaPlayable* InPlayable, UAvaPlayableTransition* InTransition, EAvaPlayableTransitionEventFlags InTransitionFlags);
 
 	void OnLevelStreamingStateChanged(UWorld* InWorld
@@ -137,7 +129,7 @@ protected:
 	bool bCacheShadersOnLoad = false;
 
 	/** True if shaders are currently being compiled */
-	bool bCachingShadersOnLoad = false;
+	bool bCachingShaders = false;
 
 	/** True if the load command specified a transform. */
 	bool bHasTransform = false;
diff --git a/Engine/Plugins/VirtualProduction/Avalanche/Source/AvalancheMedia/Private/Playback/AvaPlaybackUtils.cpp b/Engine/Plugins/VirtualProduction/Avalanche/Source/AvalancheMedia/Private/Playback/AvaPlaybackUtils.cpp
--- a/Engine/Plugins/VirtualProduction/Avalanche/Source/AvalancheMedia/Private/Playback/AvaPlaybackUtils.cpp
+++ b/Engine/Plugins/VirtualProduction/Avalanche/Source/AvalancheMedia/Private/Playback/AvaPlaybackUtils.cpp
@@ -42,7 +42,7 @@ namespace UE::AvaPlayback::Utils::Private
 
 	static TAutoConsoleVariable<bool> CVarRuntimePrecacheRCReferencedAssetsMaterials(
 		TEXT("MotionDesignPlayback.RuntimePrecacheRemoteControlObjectMaterials"),
-		true,
+		false, // False by default to avoid the extra cost of having to try cache materials at runtime/playtime.
 		TEXT("Manually caches and waits for materials used by Remote Control referenced objects to be complete while preparing to start play.")
 	);
 }
@@ -266,7 +266,7 @@ namespace UE::AvaPlayback::Utils
 		return UE::Ava::FMaterialCacheHelper::Get().RequestCacheMaterials(InObject, InShaderProfile);
 	}
 
-	void FAsyncAssetLoader::ProcessLoadedObject(UObject* InObject)
+	void FAsyncAssetLoader::ProcessLoadedObject(UObject* InObject, FName InShaderProfile)
 	{
 		if (Private::CVarRuntimePrestreamRCReferencedAssets.GetValueOnGameThread())
 		{
@@ -278,7 +278,7 @@ namespace UE::AvaPlayback::Utils
 
 		if (Private::CVarRuntimePrecacheRCReferencedAssetsMaterials.GetValueOnGameThread())
 		{
-			if (CacheObjectMaterials(InObject, UAvaMaterialCacheSettings::DefaultRealtimeProfile))
+			if (CacheObjectMaterials(InObject, InShaderProfile))
 			{
 				PendingMaterialCompletionObjects.Add(InObject);
 			}
@@ -287,6 +287,8 @@ namespace UE::AvaPlayback::Utils
 
 	void FAsyncAssetLoader::BeginLoadingAssets(const TArray<FSoftObjectPath>& InAssetsToLoad)
 	{
+		const FName ShaderProfile = GetDefault<UAvaMaterialCacheSettings>()->GetRealtimeProfile();
+
 		for (const FSoftObjectPath& AssetToLoad : InAssetsToLoad)
 		{
 			if (PendingAssets.Contains(AssetToLoad))
@@ -299,20 +301,20 @@ namespace UE::AvaPlayback::Utils
 			if (ResolvedObject)
 			{
 				// Already loaded... but streamable objects could still require streaming in, or it could have materials to cache.
-				ProcessLoadedObject(ResolvedObject);
+				ProcessLoadedObject(ResolvedObject, ShaderProfile);
 				continue;
 			}
 
 			PendingAssets.Add(AssetToLoad);
 
-			AssetToLoad.LoadAsync(FLoadSoftObjectPathAsyncDelegate::CreateSPLambda(this, [this](const FSoftObjectPath& InObjectPath, UObject* InLoadedObject)
+			AssetToLoad.LoadAsync(FLoadSoftObjectPathAsyncDelegate::CreateSPLambda(this, [this, ShaderProfile](const FSoftObjectPath& InObjectPath, UObject* InLoadedObject)
 			{
 				PendingAssets.Remove(InObjectPath);
 				if (PendingAssets.IsEmpty())
 				{
 					OnLoadingCompleted.Broadcast();
 				}
-				ProcessLoadedObject(InLoadedObject);
+				ProcessLoadedObject(InLoadedObject, ShaderProfile);
 			}));
 		}
 	}
diff --git a/Engine/Plugins/VirtualProduction/Avalanche/Source/AvalancheMedia/Private/Rundown/AvaRundown.cpp b/Engine/Plugins/VirtualProduction/Avalanche/Source/AvalancheMedia/Private/Rundown/AvaRundown.cpp
--- a/Engine/Plugins/VirtualProduction/Avalanche/Source/AvalancheMedia/Private/Rundown/AvaRundown.cpp
+++ b/Engine/Plugins/VirtualProduction/Avalanche/Source/AvalancheMedia/Private/Rundown/AvaRundown.cpp
@@ -1216,22 +1216,24 @@ TArray<UAvaRundown::FLoadedInstanceInfo> UAvaRundown::LoadPage(int32 InPageId,
 		FAvaPlayableRemoteControlValues::CollectReferencedAssetPaths(RemoteControlValues.ControllerValues, RemoteControlAssets);
 		FAvaPlayableRemoteControlValues::CollectReferencedAssetPaths(RemoteControlValues.EntityValues, RemoteControlAssets);
 
+		const FName ShaderProfile = GetDefault<UAvaMaterialCacheSettings>()->GetOfflineProfile();
+
 		for (const FSoftObjectPath& Asset : RemoteControlAssets)
 		{
 			if (UObject* Object = Asset.ResolveObject())
 			{
 				UE::AvaPlayback::Utils::StreamRenderAsset(Object);
-				UE::AvaPlayback::Utils::CacheObjectMaterials(Object, UAvaMaterialCacheSettings::DefaultOfflineProfile);
+				UE::AvaPlayback::Utils::CacheObjectMaterials(Object, ShaderProfile);
 			}
 			else
 			{
 				Asset.LoadAsync(FLoadSoftObjectPathAsyncDelegate::CreateLambda(
-					[](const FSoftObjectPath&, UObject* InObject)
+					[ShaderProfile](const FSoftObjectPath&, UObject* InObject)
 					{
 						if (InObject)
 						{
 							UE::AvaPlayback::Utils::StreamRenderAsset(InObject);
-							UE::AvaPlayback::Utils::CacheObjectMaterials(InObject, UAvaMaterialCacheSettings::DefaultOfflineProfile);
+							UE::AvaPlayback::Utils::CacheObjectMaterials(InObject, ShaderProfile);
 						}
 					}));
 			}
