Motion Design: added option to specify how shaders should be cached on level load in level streaming playables

#Synced-CL 48481255

diff --git a/Engine/Plugins/VirtualProduction/Avalanche/Source/AvalancheMedia/Private/Playable/Playables/AvaPlayableLevelStreaming.cpp b/Engine/Plugins/VirtualProduction/Avalanche/Source/AvalancheMedia/Private/Playable/Playables/AvaPlayableLevelStreaming.cpp
--- a/Engine/Plugins/VirtualProduction/Avalanche/Source/AvalancheMedia/Private/Playable/Playables/AvaPlayableLevelStreaming.cpp
+++ b/Engine/Plugins/VirtualProduction/Avalanche/Source/AvalancheMedia/Private/Playable/Playables/AvaPlayableLevelStreaming.cpp
@@ -35,6 +35,15 @@ namespace UE::AvaMedia::LevelStreamingPlayable::Private
 		TEXT("Caches any remaining shaders from materials pending completion found in the playable while preparing to start play.")
 	);
 
+	static TAutoConsoleVariable<int32> CVarCacheShadersOnLevelLoadMethod(
+		TEXT("MotionDesignPlayable.CacheShadersOnLevelLoadMethod"),
+		1,
+		TEXT("Determines how shader should be cached on level load.\n")
+		TEXT(" 0: No Caching. Shaders will not be cached on level load. They could still be cached via MotionDesignPlayable.RuntimePrecacheShaders.\n")
+		TEXT(" 1: FirstLoad. Shaders will be cached only on the first level load. Subsequent level loads will not cache shaders\n")
+		TEXT(" 2: EveryLoad. Shaders will be cached on every level load.")
+	);
+
 	UAvaPlayableAssetUserData* FindOrAddPlayableAssetUserData(ULevel* InLevel)
 	{
 		if (!InLevel)
@@ -108,12 +117,25 @@ namespace UE::AvaMedia::LevelStreamingPlayable::Private
 	}
 
 	/**
-	 * Utility function to determine whether to cache shaders from load options.
-	 * returns true if shaders should be cached on load.
+	 * Utility function to determine how to cache shaders from load options.
+	 * returns how shaders should be cached on load.
 	 */
-	bool ShouldCacheShadersOnLoad(const FString& InOptions)
+	UAvaPlayableLevelStreaming::ECacheShadersOnLevelLoadMethod GetCacheShadersOnLevelLoadMethod(const FString& InOptions)
 	{
-		return FParse::Param(*InOptions, TEXT("CacheShaders"));
+		if (FParse::Param(*InOptions, TEXT("CacheShaders")))
+		{
+			constexpr UAvaPlayableLevelStreaming::ECacheShadersOnLevelLoadMethod CacheShadersOnLevelLoadMethods[] = 
+			{
+				UAvaPlayableLevelStreaming::ECacheShadersOnLevelLoadMethod::None,
+				UAvaPlayableLevelStreaming::ECacheShadersOnLevelLoadMethod::FirstLoad,
+				UAvaPlayableLevelStreaming::ECacheShadersOnLevelLoadMethod::EveryLoad,
+			};
+			constexpr int32 LastElementIndex = static_cast<int32>(UE_ARRAY_COUNT(CacheShadersOnLevelLoadMethods)) - 1;
+
+			const int32 Index = FMath::Clamp(CVarCacheShadersOnLevelLoadMethod.GetValueOnGameThread(), 0, LastElementIndex);
+			return CacheShadersOnLevelLoadMethods[Index];
+		}
+		return UAvaPlayableLevelStreaming::ECacheShadersOnLevelLoadMethod::None;
 	}
 
 	/**
@@ -299,8 +321,8 @@ bool UAvaPlayableLevelStreaming::LoadAsset(const FAvaSoftAssetPtr& InSourceAsset
 	// Setup the level transform from the load options.
 	bHasTransform = UE::AvaMedia::LevelStreamingPlayable::Private::GetTransformFromOptions(LevelTransform, InLoadOptions, PlayableGroup, InSourceAsset);
 
-	// Determine if shaders should be cached on load
-	bCacheShadersOnLoad = UE::AvaMedia::LevelStreamingPlayable::Private::ShouldCacheShadersOnLoad(InLoadOptions);
+	// Determine how shaders should be cached on load
+	CacheShadersOnLevelLoadMethod = UE::AvaMedia::LevelStreamingPlayable::Private::GetCacheShadersOnLevelLoadMethod(InLoadOptions);
 
 	const FAvaInstanceSettings& PlaybackInstanceSettings = IAvaMediaModule::Get().GetAvaInstanceSettings();
 	bLoadSubPlayables = PlaybackInstanceSettings.bEnableLoadSubPlayables;
@@ -499,6 +521,10 @@ bool UAvaPlayableLevelStreaming::LoadLevel(const TSoftObjectPtr<UWorld>& InSourc
 	}
 
 	SourceLevel = InSourceLevel;
+
+	// Reset has ever cached shaders if the level asset has changed
+	bHasEverCachedShaders = false;
+
 	LevelStreaming->SetShouldBeLoaded(true);
 	LevelStreaming->SetShouldBeVisible(bInInitiallyVisible);
 	return true;
@@ -621,10 +647,14 @@ void UAvaPlayableLevelStreaming::OnLevelStreamingStateChanged_Synchronized(ELeve
 		// - When the RuntimePrecacheShaders cvar is set to true and the level just became visible, prior to any transition logic starting, etc.
 		// Both scenarios can happen for one playable, and it's expected since these two states can be at completely different times.
 		// However, the shader profile will be different in these two scenarios
-		if (InNewState == ELevelStreamingState::LoadedNotVisible && bCacheShadersOnLoad)
+		if (InNewState == ELevelStreamingState::LoadedNotVisible && CacheShadersOnLevelLoadMethod != ECacheShadersOnLevelLoadMethod::None)
 		{
-			ShaderProfile = GetDefault<UAvaMaterialCacheSettings>()->GetOfflineProfile();
-			bShouldCacheShaders = true;
+			// Only cache shaders if it has never done so, or the method is to cache every load.
+			if (!bHasEverCachedShaders || CacheShadersOnLevelLoadMethod == ECacheShadersOnLevelLoadMethod::EveryLoad)
+			{
+				ShaderProfile = GetDefault<UAvaMaterialCacheSettings>()->GetOfflineProfile();
+				bShouldCacheShaders = true;
+			}
 		}
 		else if (InNewState == ELevelStreamingState::LoadedVisible && UE::AvaMedia::LevelStreamingPlayable::Private::CVarRuntimePrecacheShaders.GetValueOnGameThread())
 		{
@@ -638,6 +668,7 @@ void UAvaPlayableLevelStreaming::OnLevelStreamingStateChanged_Synchronized(ELeve
 			UE::Ava::FMaterialCacheHelper& MaterialCacheHelper = UE::Ava::FMaterialCacheHelper::Get();
 			MaterialCacheHelper.RequestCacheMaterials(Level, ShaderProfile);
 			bCachingShaders = MaterialCacheHelper.IsCaching(Level);
+			bHasEverCachedShaders = true;
 
 			if (bCachingShaders)
 			{
diff --git a/Engine/Plugins/VirtualProduction/Avalanche/Source/AvalancheMedia/Private/Playable/Playables/AvaPlayableLevelStreaming.h b/Engine/Plugins/VirtualProduction/Avalanche/Source/AvalancheMedia/Private/Playable/Playables/AvaPlayableLevelStreaming.h
--- a/Engine/Plugins/VirtualProduction/Avalanche/Source/AvalancheMedia/Private/Playable/Playables/AvaPlayableLevelStreaming.h
+++ b/Engine/Plugins/VirtualProduction/Avalanche/Source/AvalancheMedia/Private/Playable/Playables/AvaPlayableLevelStreaming.h
@@ -18,6 +18,16 @@ class UAvaPlayableLevelStreaming : public UAvaPlayable
 {
 	GENERATED_BODY()
 public:
+	enum class ECacheShadersOnLevelLoadMethod : uint8
+	{
+		/** No shaders will be cached on load */
+		None,
+		/** Shaders will be cached only on the first level load. Subsequent level loads will not cache shaders */
+		FirstLoad,
+		/** Shaders will be cached every load */
+		EveryLoad,
+	};
+
 	//~ Begin UAvaPlayable
 	virtual bool LoadAsset(const FAvaSoftAssetPtr& InSourceAsset, bool bInInitiallyVisible, const FString& InLoadOptions) override;
 	virtual bool UnloadAsset() override;
@@ -117,6 +127,9 @@ protected:
 	 */
 	TSet<TObjectKey<UAvaPlayableLevelStreaming>> ParentPlayables;
 
+	/** The method of caching shaders when loading the level */
+	ECacheShadersOnLevelLoadMethod CacheShadersOnLevelLoadMethod = ECacheShadersOnLevelLoadMethod::None;
+
 	bool bOnPlayQueued = false;
 	
 	/** Enter Playables should be hidden until the transition has started. */
@@ -125,12 +138,12 @@ protected:
 	/** If true, all primitives from the playable will hidden. */
 	bool bShouldBeHidden = false;
 
-	/** True if shaders should be cached when loading the level */
-	bool bCacheShadersOnLoad = false;
-
 	/** True if shaders are currently being compiled */
 	bool bCachingShaders = false;
 
+	/** True if shaders have been cached for the level at some point in the past */
+	bool bHasEverCachedShaders = false;
+
 	/** True if the load command specified a transform. */
 	bool bHasTransform = false;
 
