Motion Design: changed visit context to accept data view instead of const data view.

#Synced-CL 48481255

diff --git a/Engine/Plugins/VirtualProduction/Avalanche/Source/AvalancheMaterial/Private/AvaLevelMaterialBridge.cpp b/Engine/Plugins/VirtualProduction/Avalanche/Source/AvalancheMaterial/Private/AvaLevelMaterialBridge.cpp
--- a/Engine/Plugins/VirtualProduction/Avalanche/Source/AvalancheMaterial/Private/AvaLevelMaterialBridge.cpp
+++ b/Engine/Plugins/VirtualProduction/Avalanche/Source/AvalancheMaterial/Private/AvaLevelMaterialBridge.cpp
@@ -18,7 +18,7 @@ void FLevelMaterialBridge::OnVisitMaterials(const FMaterialBridgeVisitContext& I
 {
 	const ULevel& Level = InContext.Container.Get<const ULevel>();
 
-	for (const AActor* Actor : Level.Actors)
+	for (AActor* Actor : Level.Actors)
 	{
 		if (const FMaterialBridge* MaterialBridge = InOptions.MaterialBridgeRegistry->GetMaterialBridgeForObject(Actor))
 		{
diff --git a/Engine/Plugins/VirtualProduction/Avalanche/Source/AvalancheMaterial/Private/AvaMaterialBridgeVisitContext.cpp b/Engine/Plugins/VirtualProduction/Avalanche/Source/AvalancheMaterial/Private/AvaMaterialBridgeVisitContext.cpp
--- a/Engine/Plugins/VirtualProduction/Avalanche/Source/AvalancheMaterial/Private/AvaMaterialBridgeVisitContext.cpp
+++ b/Engine/Plugins/VirtualProduction/Avalanche/Source/AvalancheMaterial/Private/AvaMaterialBridgeVisitContext.cpp
@@ -5,12 +5,12 @@
 namespace UE::Ava
 {
 
-FMaterialBridgeVisitContext::FMaterialBridgeVisitContext(const UObject* InContainer)
+FMaterialBridgeVisitContext::FMaterialBridgeVisitContext(UObject* InContainer)
 	: Container(InContainer)
 {
 }
 
-FMaterialBridgeVisitContext::FMaterialBridgeVisitContext(const UObject* InContainer, TNotNull<const FMaterialBridgeVisitContext*> InParentContext)
+FMaterialBridgeVisitContext::FMaterialBridgeVisitContext(UObject* InContainer, TNotNull<const FMaterialBridgeVisitContext*> InParentContext)
 	: Container(InContainer)
 	, ParentContext(InParentContext)
 {
diff --git a/Engine/Plugins/VirtualProduction/Avalanche/Source/AvalancheMaterial/Private/AvaMaterialCacheHelper.cpp b/Engine/Plugins/VirtualProduction/Avalanche/Source/AvalancheMaterial/Private/AvaMaterialCacheHelper.cpp
--- a/Engine/Plugins/VirtualProduction/Avalanche/Source/AvalancheMaterial/Private/AvaMaterialCacheHelper.cpp
+++ b/Engine/Plugins/VirtualProduction/Avalanche/Source/AvalancheMaterial/Private/AvaMaterialCacheHelper.cpp
@@ -562,7 +562,7 @@ bool FMaterialCacheHelper::GatherMaterialData(const UObject* InObject, FObjectDa
 		});
 
 	// Gather all materials for object
-	MaterialBridge->VisitMaterials(FMaterialBridgeVisitContext(InObject),
+	MaterialBridge->VisitMaterials(FMaterialBridgeVisitContext(const_cast<UObject*>(InObject)),
 		[&AddIncompleteMaterial](const FMaterialBridgeVisitContext& InContext, UMaterialInterface* InMaterial)
 		{
 			if (!InMaterial)
diff --git a/Engine/Plugins/VirtualProduction/Avalanche/Source/AvalancheMaterial/Public/AvaMaterialBridgeVisitContext.h b/Engine/Plugins/VirtualProduction/Avalanche/Source/AvalancheMaterial/Public/AvaMaterialBridgeVisitContext.h
--- a/Engine/Plugins/VirtualProduction/Avalanche/Source/AvalancheMaterial/Public/AvaMaterialBridgeVisitContext.h
+++ b/Engine/Plugins/VirtualProduction/Avalanche/Source/AvalancheMaterial/Public/AvaMaterialBridgeVisitContext.h
@@ -15,11 +15,11 @@ class FMaterialBridgeVisitContext
 public:
 	FMaterialBridgeVisitContext() = default;
 
-	explicit UE_API FMaterialBridgeVisitContext(const UObject* InContainer);
-	explicit UE_API FMaterialBridgeVisitContext(const UObject* InContainer, TNotNull<const FMaterialBridgeVisitContext*> InParentContext);
+	explicit UE_API FMaterialBridgeVisitContext(UObject* InContainer);
+	explicit UE_API FMaterialBridgeVisitContext(UObject* InContainer, TNotNull<const FMaterialBridgeVisitContext*> InParentContext);
 
 	/** Container of the visited material */
-	FConstDataView Container;
+	FDataView Container;
 
 	/** Parent of this context for additional information (e.g. container of the container) */
 	const FMaterialBridgeVisitContext* ParentContext = nullptr;
diff --git a/Engine/Plugins/VirtualProduction/Avalanche/Source/AvalancheText/Private/AvaText3DComponentMaterialBridge.cpp b/Engine/Plugins/VirtualProduction/Avalanche/Source/AvalancheText/Private/AvaText3DComponentMaterialBridge.cpp
--- a/Engine/Plugins/VirtualProduction/Avalanche/Source/AvalancheText/Private/AvaText3DComponentMaterialBridge.cpp
+++ b/Engine/Plugins/VirtualProduction/Avalanche/Source/AvalancheText/Private/AvaText3DComponentMaterialBridge.cpp
@@ -19,12 +19,10 @@ void FText3DComponentMaterialBridge::OnVisitMaterials(const FMaterialBridgeVisit
 
 	if (const UText3DMaterialExtensionBase* MaterialExtension = Component.GetMaterialExtension())
 	{
-		const FMaterialBridgeVisitContext VisitContext(MaterialExtension, &InContext);
-
 		MaterialExtension->ForEachMaterial(
-			[&InFunc, &VisitContext](const UE::Text3D::Material::FMaterialParameters&, UMaterialInterface* InMaterial)->bool
+			[&InFunc, &InContext](const UE::Text3D::Material::FMaterialParameters&, UMaterialInterface* InMaterial)->bool
 			{
-				InFunc(VisitContext, InMaterial);
+				InFunc(InContext, InMaterial);
 				return true; // continue				
 			});
 	}
