Motion Design:  moved material bridge and material cache files to their respective folders.
+ fixed missing gc reference from material bridge
+ added visit flow as a return type (to stop visitmaterials)

#Synced-CL 48481255

diff --git a/Engine/Plugins/VirtualProduction/Avalanche/Source/AvalancheMaterial/Internal/AvaMaterialCacheHelper.h b/Engine/Plugins/VirtualProduction/Avalanche/Source/AvalancheMaterial/Internal/AvaMaterialCacheHelper.h
deleted file mode 100644
diff --git a/Engine/Plugins/VirtualProduction/Avalanche/Source/AvalancheMaterial/Internal/AvaMaterialCacheLog.h b/Engine/Plugins/VirtualProduction/Avalanche/Source/AvalancheMaterial/Internal/AvaMaterialCacheLog.h
deleted file mode 100644
diff --git a/Engine/Plugins/VirtualProduction/Avalanche/Source/AvalancheMaterial/Internal/AvaMaterialCacheSettings.h b/Engine/Plugins/VirtualProduction/Avalanche/Source/AvalancheMaterial/Internal/AvaMaterialCacheSettings.h
deleted file mode 100644
diff --git a/Engine/Plugins/VirtualProduction/Avalanche/Source/AvalancheMaterial/Internal/MaterialCache/AvaMaterialCacheHelper.h b/Engine/Plugins/VirtualProduction/Avalanche/Source/AvalancheMaterial/Internal/MaterialCache/AvaMaterialCacheHelper.h
new file mode 100644
--- /dev/null
+++ b/Engine/Plugins/VirtualProduction/Avalanche/Source/AvalancheMaterial/Internal/MaterialCache/AvaMaterialCacheHelper.h
@@ -0,0 +1,156 @@
+// Copyright Epic Games, Inc. All Rights Reserved.
+
+#pragma once
+
+#include "Templates/SharedPointer.h"
+#include "UObject/GCObject.h"
+#include "UObject/ObjectKey.h"
+
+#define UE_API AVALANCHEMATERIAL_API
+
+class FMaterialUpdateContext;
+class UAvaMaterialCacheSettings;
+class UMaterialInterface;
+
+namespace UE::Ava
+{
+	struct FResolvedShaderProfile;
+}
+
+namespace UE::Ava
+{
+
+/** Helper to cache materials over multiple frames */
+class FMaterialCacheHelper : public TSharedFromThis<FMaterialCacheHelper>, public FGCObject
+{
+	struct FMaterialKey
+	{
+		/** Object key to the Material being cached */
+		TObjectKey<UMaterialInterface> Material;
+		/** Shader profile to use for caching the material shaders */
+		FName ShaderProfile;
+
+		bool operator==(const FMaterialKey& InOther) const
+		{
+			return Material == InOther.Material && ShaderProfile == InOther.ShaderProfile;
+		}
+		friend uint32 GetTypeHash(const FMaterialKey& InKey)
+		{
+			return HashCombineFast(GetTypeHash(InKey.Material), GetTypeHash(InKey.ShaderProfile));
+		}
+	};
+
+	struct FObjectData
+	{
+		/** Shader profile to use for caching the material shaders */
+		FName ShaderProfile;
+		/** All the keys to the materials belonging to an object that are still incomplete */
+		TSet<FMaterialKey> IncompleteMaterialKeys;
+		/** List of subobjects that were pending completion and could not gather materials */
+		TSet<TObjectPtr<const UObject>> PendingSubobjects;
+		/** Time elapsed before checking pending subobjects */
+		float PendingSubobjectElapsedTime = 0.f;
+	};
+
+	struct FMaterialData
+	{
+		/** Material that is pending completion */
+		TObjectPtr<UMaterialInterface> Material;
+	};
+
+public:
+	UE_API static FMaterialCacheHelper& Get();
+
+	/**
+	 * Caches the materials in the given object.
+	 * @param InObject the object to gather materials for via the material bridge api
+	 * @param InShaderProfile the shader profile to use for caching shaders. If none, it will cache via UMaterial::CacheShaders
+	 * @return true if materials were gathered for the object, false otherwise. (e.g. if object's materials were already gathered, or no material bridge found)
+	 */
+	UE_API bool RequestCacheMaterials(const UObject* InObject, FName InShaderProfile = NAME_None);
+
+	/** Processes materials pending completion with a limited time budget */
+	UE_API void Tick();
+
+	/** Returns true if there are still any materials being cached */
+	UE_API bool IsCaching() const;
+
+	/** Returns true if there are materials from the given object that are still being cached */
+	UE_API bool IsCaching(const UObject* InObject) const;
+
+	/** Dumps all the materials that are still being processed (i.e. still incomplete) to the output log */
+	UE_API void DumpMaterials();
+
+	//~ Begin FGCObject
+	virtual void AddReferencedObjects(FReferenceCollector& InCollector) override;
+	virtual FString GetReferencerName() const override;
+	//~ End FGCObject
+
+private:
+	/** Returns true if the material is still being cached */
+	bool IsCachingMaterial(const FMaterialKey& InMaterialKey) const;
+
+	/** Resets the object and material map */
+	void Reset();
+
+	/** Registers the tick function to get called every frame */
+	void RegisterTick();
+
+	/** Unregisters the tick function from being called every frame */
+	void UnregisterTick();
+
+	/** Processes subobjects that did not have its materials gathered */
+	void ProcessPendingObjects();
+
+	/** Caches as many queued materials as the time budget allows */
+	void ProcessMaterialsPendingCache();
+
+	/** Checks the status on materials that are being cached, and removes from the map once they're complete */
+	void ProcessMaterialsCaching();
+
+	/** Removes the complete materials from the objects' IncompleteMaterials set */
+	void RemoveCompletedMaterials();
+
+	/** Process async results if shader is compiling */
+	void ProcessShadersCompiling();
+
+	/**
+	 * Gathers the material data for an object
+	 * returns true if a material bridge was found for the object and visited its materials
+	 */
+	bool GatherMaterialData(const UObject* InObject, FObjectData& InObjectData, FName InShaderProfile);
+
+	/** Objects that have already been gathered, mapped to the incomplete materials that were gathered for the object */
+	TMap<FObjectKey, FObjectData> Objects;
+
+	/** All the materials that have been gathered and are pending to start caching their shaders. Once materials start caching, they're removed from this map and added to MaterialsCaching */
+	TMap<FMaterialKey, FMaterialData> MaterialsPendingCache;
+
+	/** All the materials that are caching. Once materials are complete, they're removed from this map */
+	TMap<FMaterialKey, FMaterialData> MaterialsCaching;
+
+	/** The delta time at the start of the tick */
+	float PreTickFrameDuration = 0.0;
+
+	/** Time when BeginFrame was called */
+	double BeginFrameTime = 0.0;
+
+	/** Current time budget left for the tick. Resets every tick */
+	float RemainingTimeBudget = 0.f;
+
+	/** Number of frames where Tick was skipped because the frame time was past max frame time */
+	int32 FramesStarved = 0;
+
+	/** Handle to the BeginFrame delegate to keep track of BeginFrameTime */
+	FDelegateHandle BeginFrameHandle;
+
+	/** Handle to the delegate for ticking */
+	FDelegateHandle EndFrameHandle;
+
+	/** Reset material cache */
+	bool bResetRequested = false;
+};
+
+} // UE::Ava
+
+#undef UE_API
diff --git a/Engine/Plugins/VirtualProduction/Avalanche/Source/AvalancheMaterial/Internal/MaterialCache/AvaMaterialCacheLog.h b/Engine/Plugins/VirtualProduction/Avalanche/Source/AvalancheMaterial/Internal/MaterialCache/AvaMaterialCacheLog.h
new file mode 100644
--- /dev/null
+++ b/Engine/Plugins/VirtualProduction/Avalanche/Source/AvalancheMaterial/Internal/MaterialCache/AvaMaterialCacheLog.h
@@ -0,0 +1,7 @@
+// Copyright Epic Games, Inc. All Rights Reserved.
+
+#pragma once
+
+#include "Logging/LogCategory.h"
+
+DECLARE_LOG_CATEGORY_EXTERN(LogAvaMaterialCache, Log, All);
diff --git a/Engine/Plugins/VirtualProduction/Avalanche/Source/AvalancheMaterial/Internal/MaterialCache/AvaMaterialCacheSettings.h b/Engine/Plugins/VirtualProduction/Avalanche/Source/AvalancheMaterial/Internal/MaterialCache/AvaMaterialCacheSettings.h
new file mode 100644
--- /dev/null
+++ b/Engine/Plugins/VirtualProduction/Avalanche/Source/AvalancheMaterial/Internal/MaterialCache/AvaMaterialCacheSettings.h
@@ -0,0 +1,141 @@
+// Copyright Epic Games, Inc. All Rights Reserved.
+
+#pragma once
+
+#include "Engine/DeveloperSettings.h"
+#include "AvaMaterialCacheSettings.generated.h"
+
+#define UE_API AVALANCHEMATERIAL_API
+
+class FShaderPipelineType;
+class FShaderType;
+class FVertexFactoryType;
+
+UENUM()
+enum class EAvaShaderProfileType : uint8
+{
+	/** CacheShaders will be used and will gather all the shaders it might need */
+	Default,
+	/**
+	 * Only the specified shaders will be cached. Should yield better performance 
+	 * NOTE: Specified shaders only will apply in editor builds. This will default to generic on non-editor builds
+	 */
+	Specific,
+	/** No shaders will be cached */
+	SkipCaching,
+};
+
+namespace UE::Ava
+{
+
+/** Shader Profile that has all the type pointers cached */
+struct FResolvedShaderProfile
+{
+	EAvaShaderProfileType ProfileType = EAvaShaderProfileType::Default;
+	TArray<const FVertexFactoryType*> VertexFactoryTypes;
+	TArray<const FShaderPipelineType*> PipelineTypes;
+	TArray<const FShaderType*> ShaderTypes;
+};
+
+} // UE::Ava
+
+USTRUCT()
+struct FAvaShaderCombination
+{
+	GENERATED_BODY()
+
+	FAvaShaderCombination() = default;
+
+	explicit FAvaShaderCombination(FName InVertexFactorType, TConstArrayView<FName> InPipelineTypes, TConstArrayView<FName> InShaderTypes)
+		: VertexFactoryType(InVertexFactorType)
+		, PipelineTypes(InPipelineTypes)
+		, ShaderTypes(InShaderTypes)
+	{
+	}
+
+	UPROPERTY(EditAnywhere, Category="Material Cache")
+	FName VertexFactoryType;
+
+	UPROPERTY(EditAnywhere, Category="Material Cache")
+	TArray<FName> PipelineTypes;
+
+	UPROPERTY(EditAnywhere, Category="Material Cache")
+	TArray<FName> ShaderTypes;
+};
+
+USTRUCT()
+struct FAvaShaderProfile
+{
+	GENERATED_BODY()
+
+	UPROPERTY(EditAnywhere, Category="Material Cache")
+	FName ProfileName;
+
+	UPROPERTY(EditAnywhere, Category="Material Cache")
+	EAvaShaderProfileType ProfileType = EAvaShaderProfileType::Default;
+
+	/** Specific Shader Combinations to use */
+	UPROPERTY(EditAnywhere, Category="Material Cache", meta=(EditCondition="ProfileType==EAvaShaderProfileType::Specific"))
+	TArray<FAvaShaderCombination> ShaderCombinations;
+};
+
+UCLASS(MinimalAPI, config=Engine, meta=(DisplayName="Material Cache Settings"))
+class UAvaMaterialCacheSettings : public UDeveloperSettings
+{
+	GENERATED_BODY()
+
+public:
+	UE_API static const FLazyName DefaultRealtimeProfile;
+	UE_API static const FLazyName DefaultOfflineProfile;
+
+	UAvaMaterialCacheSettings();
+
+	FName GetRealtimeProfile() const
+	{
+		return RealtimeProfile;
+	}
+
+	FName GetOfflineProfile() const
+	{
+		return OfflineProfile;
+	}
+
+	/** Finds the resolved shader profile for a given profile name */
+	const UE::Ava::FResolvedShaderProfile* FindResolvedShaderProfile(FName InProfileName) const;
+
+protected:
+	//~ Begin UObject
+	virtual void PostInitProperties() override;
+	virtual void PostReloadConfig(FProperty* InPropertyThatWasLoaded) override;
+#if WITH_EDITOR
+	virtual void PostEditChangeProperty(FPropertyChangedEvent& InPropertyChangedEvent) override;
+#endif
+	//~ End UObject
+
+private:
+	/** Add the built-in shader profiles */
+	void RegisterDefaultShaderProfile();
+
+	/** Resolves all the registered shader profiles to its type pointers */
+	void ResolveShaderProfiles();
+
+	/** Resolves a single shader profile to its type pointers */
+	void ResolveShaderProfile(const FAvaShaderProfile& InShaderProfile, TMap<FName, FHashedName>& InOutHashedNames);
+
+	/** Profile to use in realtime (e.g. when playing levels) */
+	UPROPERTY(Config, EditAnywhere, Category="Material Cache")
+	FName RealtimeProfile;
+
+	/** Profile to use offline (e.g. when preloading levels) */
+	UPROPERTY(Config, EditAnywhere, Category="Material Cache")
+	FName OfflineProfile;
+
+	/** Shader profiles that can be saved to disk */
+	UPROPERTY(Config, EditAnywhere, Category="Material Cache")
+	TArray<FAvaShaderProfile> ShaderProfiles;
+
+	/** The saved shader profiles turned into resolved pointers */
+	TMap<FName, UE::Ava::FResolvedShaderProfile> ResolvedShaderProfiles;
+};
+
+#undef UE_API
diff --git a/Engine/Plugins/VirtualProduction/Avalanche/Source/AvalancheMaterial/Private/AvaActorMaterialBridge.cpp b/Engine/Plugins/VirtualProduction/Avalanche/Source/AvalancheMaterial/Private/AvaActorMaterialBridge.cpp
deleted file mode 100644
diff --git a/Engine/Plugins/VirtualProduction/Avalanche/Source/AvalancheMaterial/Private/AvaDecalComponentMaterialBridge.cpp b/Engine/Plugins/VirtualProduction/Avalanche/Source/AvalancheMaterial/Private/AvaDecalComponentMaterialBridge.cpp
deleted file mode 100644
diff --git a/Engine/Plugins/VirtualProduction/Avalanche/Source/AvalancheMaterial/Private/AvaLevelMaterialBridge.cpp b/Engine/Plugins/VirtualProduction/Avalanche/Source/AvalancheMaterial/Private/AvaLevelMaterialBridge.cpp
deleted file mode 100644
diff --git a/Engine/Plugins/VirtualProduction/Avalanche/Source/AvalancheMaterial/Private/AvaMaterialBridge.cpp b/Engine/Plugins/VirtualProduction/Avalanche/Source/AvalancheMaterial/Private/AvaMaterialBridge.cpp
deleted file mode 100644
diff --git a/Engine/Plugins/VirtualProduction/Avalanche/Source/AvalancheMaterial/Private/AvaMaterialBridgeRegistry.cpp b/Engine/Plugins/VirtualProduction/Avalanche/Source/AvalancheMaterial/Private/AvaMaterialBridgeRegistry.cpp
deleted file mode 100644
diff --git a/Engine/Plugins/VirtualProduction/Avalanche/Source/AvalancheMaterial/Private/AvaMaterialBridgeVisitContext.cpp b/Engine/Plugins/VirtualProduction/Avalanche/Source/AvalancheMaterial/Private/AvaMaterialBridgeVisitContext.cpp
deleted file mode 100644
diff --git a/Engine/Plugins/VirtualProduction/Avalanche/Source/AvalancheMaterial/Private/AvaMaterialCacheHelper.cpp b/Engine/Plugins/VirtualProduction/Avalanche/Source/AvalancheMaterial/Private/AvaMaterialCacheHelper.cpp
deleted file mode 100644
diff --git a/Engine/Plugins/VirtualProduction/Avalanche/Source/AvalancheMaterial/Private/AvaMaterialCacheSettings.cpp b/Engine/Plugins/VirtualProduction/Avalanche/Source/AvalancheMaterial/Private/AvaMaterialCacheSettings.cpp
deleted file mode 100644
diff --git a/Engine/Plugins/VirtualProduction/Avalanche/Source/AvalancheMaterial/Private/AvaMaterialModule.cpp b/Engine/Plugins/VirtualProduction/Avalanche/Source/AvalancheMaterial/Private/AvaMaterialModule.cpp
--- a/Engine/Plugins/VirtualProduction/Avalanche/Source/AvalancheMaterial/Private/AvaMaterialModule.cpp
+++ b/Engine/Plugins/VirtualProduction/Avalanche/Source/AvalancheMaterial/Private/AvaMaterialModule.cpp
@@ -1,14 +1,14 @@
 // Copyright Epic Games, Inc. All Rights Reserved.
 
 #include "AvaMaterialModule.h"
-#include "AvaActorMaterialBridge.h"
-#include "AvaDecalComponentMaterialBridge.h"
-#include "AvaLevelMaterialBridge.h"
-#include "AvaMaterialBridge.h"
-#include "AvaMaterialBridgeLog.h"
-#include "AvaMaterialBridgeRegistry.h"
-#include "AvaPrimitiveComponentMaterialBridge.h"
-#include "AvaStaticMeshMaterialBridge.h"
+#include "MaterialBridge/AvaActorMaterialBridge.h"
+#include "MaterialBridge/AvaDecalComponentMaterialBridge.h"
+#include "MaterialBridge/AvaLevelMaterialBridge.h"
+#include "MaterialBridge/AvaMaterialBridge.h"
+#include "MaterialBridge/AvaMaterialBridgeLog.h"
+#include "MaterialBridge/AvaMaterialBridgeRegistry.h"
+#include "MaterialBridge/AvaPrimitiveComponentMaterialBridge.h"
+#include "MaterialBridge/AvaStaticMeshMaterialBridge.h"
 #include "Modules/ModuleManager.h"
 
 DEFINE_LOG_CATEGORY(LogAvaMaterialBridge);
diff --git a/Engine/Plugins/VirtualProduction/Avalanche/Source/AvalancheMaterial/Private/AvaPrimitiveComponentMaterialBridge.cpp b/Engine/Plugins/VirtualProduction/Avalanche/Source/AvalancheMaterial/Private/AvaPrimitiveComponentMaterialBridge.cpp
deleted file mode 100644
diff --git a/Engine/Plugins/VirtualProduction/Avalanche/Source/AvalancheMaterial/Private/AvaStaticMeshMaterialBridge.cpp b/Engine/Plugins/VirtualProduction/Avalanche/Source/AvalancheMaterial/Private/AvaStaticMeshMaterialBridge.cpp
deleted file mode 100644
diff --git a/Engine/Plugins/VirtualProduction/Avalanche/Source/AvalancheMaterial/Private/MaterialBridge/AvaActorMaterialBridge.cpp b/Engine/Plugins/VirtualProduction/Avalanche/Source/AvalancheMaterial/Private/MaterialBridge/AvaActorMaterialBridge.cpp
new file mode 100644
--- /dev/null
+++ b/Engine/Plugins/VirtualProduction/Avalanche/Source/AvalancheMaterial/Private/MaterialBridge/AvaActorMaterialBridge.cpp
@@ -0,0 +1,39 @@
+// Copyright Epic Games, Inc. All Rights Reserved.
+
+#include "MaterialBridge/AvaActorMaterialBridge.h"
+#include "GameFramework/Actor.h"
+#include "MaterialBridge/AvaMaterialBridgeRegistry.h"
+
+namespace UE::Ava
+{
+
+const UStruct* FActorMaterialBridge::OnGetBridgedType() const
+{
+	return AActor::StaticClass();
+}
+
+EVisitFlow FActorMaterialBridge::OnVisitMaterials(const FMaterialBridgeVisitContext& InContext, TFunctionRef<EVisitFlow(const FMaterialBridgeVisitContext&, UMaterialInterface*)> InFunc, const FVisitOptions& InOptions) const
+{
+	const AActor& Actor = InContext.Container.Get<const AActor>();
+
+	constexpr bool bIncludeFromChildActors = true;
+
+	EVisitFlow VisitFlow = EVisitFlow::Continue;
+
+	Actor.ForEachComponent(bIncludeFromChildActors,
+		[&InFunc, &InOptions, &InContext, &VisitFlow](UActorComponent* InComponent)
+		{
+			if (VisitFlow != EVisitFlow::Continue)
+			{
+				return;
+			}
+			if (const FMaterialBridge* MaterialBridge = InOptions.MaterialBridgeRegistry->GetMaterialBridgeForObject(InComponent))
+			{
+				VisitFlow = MaterialBridge->VisitMaterials(FMaterialBridgeVisitContext(InComponent, &InContext), InFunc, InOptions);
+			}
+		});
+
+	return VisitFlow;
+}
+
+} // UE::Ava
diff --git a/Engine/Plugins/VirtualProduction/Avalanche/Source/AvalancheMaterial/Private/MaterialBridge/AvaDecalComponentMaterialBridge.cpp b/Engine/Plugins/VirtualProduction/Avalanche/Source/AvalancheMaterial/Private/MaterialBridge/AvaDecalComponentMaterialBridge.cpp
new file mode 100644
--- /dev/null
+++ b/Engine/Plugins/VirtualProduction/Avalanche/Source/AvalancheMaterial/Private/MaterialBridge/AvaDecalComponentMaterialBridge.cpp
@@ -0,0 +1,21 @@
+// Copyright Epic Games, Inc. All Rights Reserved.
+
+#include "MaterialBridge/AvaDecalComponentMaterialBridge.h"
+#include "AvaDataView.h"
+#include "Components/DecalComponent.h"
+
+namespace UE::Ava
+{
+
+const UStruct* FDecalComponentMaterialBridge::OnGetBridgedType() const
+{
+	return UDecalComponent::StaticClass();
+}
+
+EVisitFlow FDecalComponentMaterialBridge::OnVisitMaterials(const FMaterialBridgeVisitContext& InContext, TFunctionRef<EVisitFlow(const FMaterialBridgeVisitContext&, UMaterialInterface*)> InFunc, const FVisitOptions& InOptions) const
+{
+	const UDecalComponent& Component = InContext.Container.Get<const UDecalComponent>();
+	return InFunc(InContext, Component.GetDecalMaterial());
+}
+
+} // UE::Ava
diff --git a/Engine/Plugins/VirtualProduction/Avalanche/Source/AvalancheMaterial/Private/MaterialBridge/AvaLevelMaterialBridge.cpp b/Engine/Plugins/VirtualProduction/Avalanche/Source/AvalancheMaterial/Private/MaterialBridge/AvaLevelMaterialBridge.cpp
new file mode 100644
--- /dev/null
+++ b/Engine/Plugins/VirtualProduction/Avalanche/Source/AvalancheMaterial/Private/MaterialBridge/AvaLevelMaterialBridge.cpp
@@ -0,0 +1,35 @@
+// Copyright Epic Games, Inc. All Rights Reserved.
+
+#include "MaterialBridge/AvaLevelMaterialBridge.h"
+#include "Engine/Level.h"
+#include "GameFramework/Actor.h"
+#include "MaterialBridge/AvaMaterialBridgeRegistry.h"
+
+namespace UE::Ava
+{
+
+const UStruct* FLevelMaterialBridge::OnGetBridgedType() const
+{
+	return ULevel::StaticClass();
+}
+
+EVisitFlow FLevelMaterialBridge::OnVisitMaterials(const FMaterialBridgeVisitContext& InContext, TFunctionRef<EVisitFlow(const FMaterialBridgeVisitContext&, UMaterialInterface*)> InFunc, const FVisitOptions& InOptions) const
+{
+	const ULevel& Level = InContext.Container.Get<const ULevel>();
+
+	for (AActor* Actor : Level.Actors)
+	{
+		if (const FMaterialBridge* MaterialBridge = InOptions.MaterialBridgeRegistry->GetMaterialBridgeForObject(Actor))
+		{
+			const EVisitFlow VisitFlow = MaterialBridge->VisitMaterials(FMaterialBridgeVisitContext(Actor, &InContext), InFunc, InOptions);
+			if (VisitFlow == EVisitFlow::Break)
+			{
+				return EVisitFlow::Break;
+			}
+		}
+	}
+
+	return EVisitFlow::Continue;
+}
+
+} // UE::Ava
diff --git a/Engine/Plugins/VirtualProduction/Avalanche/Source/AvalancheMaterial/Private/MaterialBridge/AvaMaterialBridge.cpp b/Engine/Plugins/VirtualProduction/Avalanche/Source/AvalancheMaterial/Private/MaterialBridge/AvaMaterialBridge.cpp
new file mode 100644
--- /dev/null
+++ b/Engine/Plugins/VirtualProduction/Avalanche/Source/AvalancheMaterial/Private/MaterialBridge/AvaMaterialBridge.cpp
@@ -0,0 +1,32 @@
+// Copyright Epic Games, Inc. All Rights Reserved.
+
+#include "MaterialBridge/AvaMaterialBridge.h"
+#include "MaterialBridge/AvaMaterialBridgeRegistry.h"
+
+namespace UE::Ava
+{
+FMaterialBridge::FVisitOptions::FVisitOptions()
+	: MaterialBridgeRegistry(&FMaterialBridgeRegistry::Get())
+{
+}
+
+void FMaterialBridge::Initialize()
+{
+	BridgedType = OnGetBridgedType();
+}
+
+const UStruct* FMaterialBridge::GetBridgedType() const
+{
+	return BridgedType;
+}
+
+EVisitFlow FMaterialBridge::VisitMaterials(const FMaterialBridgeVisitContext& InContext, TFunctionRef<EVisitFlow(const FMaterialBridgeVisitContext&, UMaterialInterface*)> InFunc, const FVisitOptions& InOptions) const
+{
+	if (InContext.Container.IsValidFor(BridgedType))
+	{
+		return OnVisitMaterials(InContext, InFunc, InOptions);
+	}
+	return EVisitFlow::Continue;
+}
+
+} // UE::Ava
diff --git a/Engine/Plugins/VirtualProduction/Avalanche/Source/AvalancheMaterial/Private/MaterialBridge/AvaMaterialBridgeRegistry.cpp b/Engine/Plugins/VirtualProduction/Avalanche/Source/AvalancheMaterial/Private/MaterialBridge/AvaMaterialBridgeRegistry.cpp
new file mode 100644
--- /dev/null
+++ b/Engine/Plugins/VirtualProduction/Avalanche/Source/AvalancheMaterial/Private/MaterialBridge/AvaMaterialBridgeRegistry.cpp
@@ -0,0 +1,62 @@
+// Copyright Epic Games, Inc. All Rights Reserved.
+
+#include "MaterialBridge/AvaMaterialBridgeRegistry.h"
+#include "MaterialBridge/AvaMaterialBridge.h"
+
+namespace UE::Ava
+{
+
+const FMaterialBridgeRegistry& FMaterialBridgeRegistry::Get()
+{
+	return GetMutable();
+}
+
+FMaterialBridgeRegistry& FMaterialBridgeRegistry::GetMutable()
+{
+	static FMaterialBridgeRegistry MaterialBridgeRegistry;
+	return MaterialBridgeRegistry;
+}
+
+const FMaterialBridge* FMaterialBridgeRegistry::GetMaterialBridgeForObject(const UObject* InObject) const
+{
+	if (InObject)
+	{
+		return GetMaterialBridgeForType(InObject->GetClass());
+	}
+	return nullptr;
+}
+
+const FMaterialBridge* FMaterialBridgeRegistry::GetMaterialBridgeForType(const UStruct* InStruct) const
+{
+	if (!InStruct)
+	{
+		return nullptr;
+	}
+	for (const UStruct* Struct : InStruct->GetSuperStructIterator())
+	{
+		if (const TUniquePtr<FMaterialBridge>* MaterialBridge = RegisteredMaterialBridges.Find(Struct))
+		{
+			return MaterialBridge->Get();
+		}
+	}
+	return nullptr;
+}
+
+
+void FMaterialBridgeRegistry::RegisterInternal(TUniquePtr<UE::Ava::FMaterialBridge>&& InMaterialBridge)
+{
+	if (!ensureMsgf(UObjectInitialized(), TEXT("Cannot register Material Bridges. UObject system not initialized!")))
+	{
+		return;
+	}
+
+	InMaterialBridge->Initialize();
+	const UStruct* const BridgedType = InMaterialBridge->GetBridgedType();
+
+	if (ensureMsgf(BridgedType, TEXT("Cannot register Material Bridge with a null bridged type!")))
+	{
+		RegisteredMaterialBridges.Add(BridgedType, MoveTemp(InMaterialBridge));
+	}
+}
+
+} // UE::Ava
diff --git a/Engine/Plugins/VirtualProduction/Avalanche/Source/AvalancheMaterial/Private/MaterialBridge/AvaMaterialBridgeVisitContext.cpp b/Engine/Plugins/VirtualProduction/Avalanche/Source/AvalancheMaterial/Private/MaterialBridge/AvaMaterialBridgeVisitContext.cpp
new file mode 100644
--- /dev/null
+++ b/Engine/Plugins/VirtualProduction/Avalanche/Source/AvalancheMaterial/Private/MaterialBridge/AvaMaterialBridgeVisitContext.cpp
@@ -0,0 +1,19 @@
+// Copyright Epic Games, Inc. All Rights Reserved.
+
+#include "MaterialBridge/AvaMaterialBridgeVisitContext.h"
+
+namespace UE::Ava
+{
+
+FMaterialBridgeVisitContext::FMaterialBridgeVisitContext(UObject* InContainer)
+	: Container(InContainer)
+{
+}
+
+FMaterialBridgeVisitContext::FMaterialBridgeVisitContext(UObject* InContainer, TNotNull<const FMaterialBridgeVisitContext*> InParentContext)
+	: Container(InContainer)
+	, ParentContext(InParentContext)
+{
+}
+
+} // UE::Ava
diff --git a/Engine/Plugins/VirtualProduction/Avalanche/Source/AvalancheMaterial/Private/MaterialBridge/AvaPrimitiveComponentMaterialBridge.cpp b/Engine/Plugins/VirtualProduction/Avalanche/Source/AvalancheMaterial/Private/MaterialBridge/AvaPrimitiveComponentMaterialBridge.cpp
new file mode 100644
--- /dev/null
+++ b/Engine/Plugins/VirtualProduction/Avalanche/Source/AvalancheMaterial/Private/MaterialBridge/AvaPrimitiveComponentMaterialBridge.cpp
@@ -0,0 +1,33 @@
+// Copyright Epic Games, Inc. All Rights Reserved.
+
+#include "MaterialBridge/AvaPrimitiveComponentMaterialBridge.h"
+#include "AvaDataView.h"
+#include "Components/PrimitiveComponent.h"
+
+namespace UE::Ava
+{
+
+const UStruct* FPrimitiveComponentMaterialBridge::OnGetBridgedType() const
+{
+	return UPrimitiveComponent::StaticClass();
+}
+
+EVisitFlow FPrimitiveComponentMaterialBridge::OnVisitMaterials(const FMaterialBridgeVisitContext& InContext, TFunctionRef<EVisitFlow(const FMaterialBridgeVisitContext&, UMaterialInterface*)> InFunc, const FVisitOptions& InOptions) const
+{
+	const UPrimitiveComponent& Component = InContext.Container.Get<const UPrimitiveComponent>();
+
+	TArray<UMaterialInterface*> Materials;
+	Component.GetUsedMaterials(Materials, InOptions.bVisitDebugMaterials);
+
+	for (UMaterialInterface* Material : Materials)
+	{
+		if (InFunc(InContext, Material) == EVisitFlow::Break)
+		{
+			return EVisitFlow::Break;
+		}
+	}
+
+	return EVisitFlow::Continue;
+}
+
+} // UE::Ava
diff --git a/Engine/Plugins/VirtualProduction/Avalanche/Source/AvalancheMaterial/Private/MaterialBridge/AvaStaticMeshMaterialBridge.cpp b/Engine/Plugins/VirtualProduction/Avalanche/Source/AvalancheMaterial/Private/MaterialBridge/AvaStaticMeshMaterialBridge.cpp
new file mode 100644
--- /dev/null
+++ b/Engine/Plugins/VirtualProduction/Avalanche/Source/AvalancheMaterial/Private/MaterialBridge/AvaStaticMeshMaterialBridge.cpp
@@ -0,0 +1,47 @@
+// Copyright Epic Games, Inc. All Rights Reserved.
+
+#include "MaterialBridge/AvaStaticMeshMaterialBridge.h"
+#include "MaterialBridge/AvaMaterialBridgeLog.h"
+#include "Engine/StaticMesh.h"
+
+namespace UE::Ava
+{
+
+const UStruct* FStaticMeshMaterialBridge::OnGetBridgedType() const
+{
+	return UStaticMesh::StaticClass();
+}
+
+EVisitFlow FStaticMeshMaterialBridge::OnVisitMaterials(const FMaterialBridgeVisitContext& InContext, TFunctionRef<EVisitFlow(const FMaterialBridgeVisitContext&, UMaterialInterface*)> InFunc, const FVisitOptions& InOptions) const
+{
+	const UStaticMesh& StaticMesh = InContext.Container.Get<const UStaticMesh>();
+
+	// If static mesh is compiling and visitor does not want to wait on materials, bail out.
+	// NOTE: Only LockedProperties needs to be checked for EStaticMeshAsyncProperties::StaticMaterials, but there's currently no option to check for it directly.
+	if (StaticMesh.IsCompiling() && !InOptions.bTryWaitOnContainerCompletion)
+	{
+		if (InOptions.OnContainerPendingCompletion.IsBound())
+		{
+			UE_LOG(LogAvaMaterialBridge, Verbose, TEXT("Material Bridge could not visit materials of static mesh '%s'. It is currently compiling."), *StaticMesh.GetName());
+			return InOptions.OnContainerPendingCompletion.Execute(InContext);
+		}
+
+		UE_LOG(LogAvaMaterialBridge, Warning, TEXT("Material Bridge could not visit materials of static mesh '%s'. It is currently compiling."), *StaticMesh.GetName());
+		return EVisitFlow::Continue;
+	}
+
+	for (const FStaticMaterial& StaticMaterial : StaticMesh.GetStaticMaterials())
+	{
+		if (InFunc(InContext, StaticMaterial.MaterialInterface) == EVisitFlow::Break)
+		{
+			return EVisitFlow::Break;
+		}
+		if (InFunc(InContext, StaticMaterial.OverlayMaterialInterface) == EVisitFlow::Break)
+		{
+			return EVisitFlow::Break;
+		}
+	}
+	return EVisitFlow::Continue;
+}
+
+} // UE::Ava
diff --git a/Engine/Plugins/VirtualProduction/Avalanche/Source/AvalancheMaterial/Private/MaterialCache/AvaMaterialCacheHelper.cpp b/Engine/Plugins/VirtualProduction/Avalanche/Source/AvalancheMaterial/Private/MaterialCache/AvaMaterialCacheHelper.cpp
new file mode 100644
--- /dev/null
+++ b/Engine/Plugins/VirtualProduction/Avalanche/Source/AvalancheMaterial/Private/MaterialCache/AvaMaterialCacheHelper.cpp
@@ -0,0 +1,596 @@
+// Copyright Epic Games, Inc. All Rights Reserved.
+
+#include "MaterialCache/AvaMaterialCacheHelper.h"
+#include "AvaDataView.h"
+#include "MaterialBridge/AvaMaterialBridge.h"
+#include "MaterialBridge/AvaMaterialBridgeRegistry.h"
+#include "MaterialCache/AvaMaterialCacheLog.h"
+#include "MaterialCache/AvaMaterialCacheSettings.h"
+#include "MaterialShared.h"
+#include "Materials/Material.h"
+#include "ShaderCompiler.h"
+#include "UnrealEngine.h"
+
+DEFINE_LOG_CATEGORY(LogAvaMaterialCache);
+
+namespace UE::Ava
+{
+
+namespace Private
+{
+
+FAutoConsoleCommand MaterialCacheDumpMaterials(TEXT("MotionDesign.MaterialCache.DumpMaterials")
+	, TEXT("Dumps all the materials that are currently being processed in the Material Cache Helper")
+	, FConsoleCommandDelegate::CreateLambda([]
+		{
+			FMaterialCacheHelper::Get().DumpMaterials();
+		}));
+
+TAutoConsoleVariable<float> CVarMaterialCacheFrameTimeStarveThreshold(TEXT("MotionDesign.MaterialCache.FrameTimeStarveThreshold")
+	, 1.f / 30.f
+	, TEXT("The frame time threshold for starving tick. If the current duration of the frame surpasses this, FMaterialCacheHelper::Tick will be skipped and starve.")
+	);
+
+TAutoConsoleVariable<int32> CVarMaterialCacheTickMaxStarveFrames(TEXT("MotionDesign.MaterialCache.TickMaxStarveFrames")
+	, 8
+	, TEXT("The maximum number of consecutive frames that FMaterialCacheHelper::Tick can starve for. After this, tick will be forced to execute.")
+	);
+
+TAutoConsoleVariable<float> CVarMaterialCacheTickTimeBudget(TEXT("MotionDesign.MaterialCache.TickTimeBudget")
+	, 0.002f
+	, TEXT("The time budget FMaterialCacheHelper::Tick has to process its pending elements.")
+	);
+
+#if WITH_EDITOR
+void ForEachMaterialResource(TNotNull<UMaterialInterface*> InMaterial, TFunctionRef<bool(FMaterialResource*)> InFunc)
+{
+	const EMaterialQualityLevel::Type ActiveQualityLevel = GetCachedScalabilityCVars().MaterialQualityLevel;
+	uint32 FeatureLevelsToCompile = UMaterial::GetFeatureLevelsToCompileForAllMaterials();
+
+	while (FeatureLevelsToCompile != 0)
+	{
+		const ERHIFeatureLevel::Type FeatureLevel = static_cast<ERHIFeatureLevel::Type>(FBitSet::GetAndClearNextBit(FeatureLevelsToCompile));
+		const EShaderPlatform ShaderPlatform = GShaderPlatformForFeatureLevel[FeatureLevel];
+
+		// Only cache shaders for the quality level that will actually be used to render
+		// In cooked build, there is no shader compilation but this is still needed to register the loaded shadermap
+		FMaterialResource* CurrentResource = InMaterial->GetMaterialResource(ShaderPlatform, ActiveQualityLevel);
+		if (ensure(CurrentResource) && !InFunc(CurrentResource))
+		{
+			break;
+		}
+	}
+}
+#endif
+
+/** Finds the shader profile matching the given name */
+const FResolvedShaderProfile* FindShaderProfile(FName InShaderProfile)
+{
+	if (const UAvaMaterialCacheSettings* MaterialCacheSettings = GetDefault<UAvaMaterialCacheSettings>())
+	{
+		return MaterialCacheSettings->FindResolvedShaderProfile(InShaderProfile);
+	}
+	UE_LOG(LogAvaMaterialCache, Verbose, TEXT("Shader Profile '%s' not found in Material Cache Settings"), *InShaderProfile.ToString());
+	return nullptr;
+}
+
+/** Performs shader caching of the given material with the given shader profile */
+void CacheMaterial(TNotNull<UMaterialInterface*> InMaterial, TNotNull<const FResolvedShaderProfile*> InShaderProfile)
+{
+	UE_LOG(LogAvaMaterialCache, VeryVerbose, TEXT("Caching Material '%s'"), *InMaterial->GetFullName());
+
+#if WITH_EDITOR
+	if (InShaderProfile->ProfileType == EAvaShaderProfileType::Specific)
+	{
+		InMaterial->CacheShaders(EMaterialShaderPrecompileMode::None);
+
+		ForEachMaterialResource(InMaterial,
+			[InShaderProfile](FMaterialResource* InMaterialResource)->bool
+			{
+				InMaterialResource->CacheGivenTypes(InShaderProfile->VertexFactoryTypes, InShaderProfile->PipelineTypes, InShaderProfile->ShaderTypes);
+				return true; // continue
+			});
+
+		return;
+	}
+#endif
+
+	InMaterial->CacheShaders(EMaterialShaderPrecompileMode::Background);
+}
+
+/** Checks whether the material is complete */
+bool IsMaterialComplete(TNotNull<UMaterialInterface*> InMaterial, TNotNull<const FResolvedShaderProfile*> InShaderProfile)
+{
+#if WITH_EDITOR
+	if (InShaderProfile->ProfileType == EAvaShaderProfileType::Specific)
+	{
+		bool bIsComplete = true;
+
+		ForEachMaterialResource(InMaterial, 
+			[&bIsComplete](FMaterialResource* InMaterialResource)->bool
+			{
+				if (InMaterialResource->IsCachingShaders())
+				{
+					bIsComplete = false;
+					return false; // break
+				}
+				return true; // continue
+			});
+
+		return bIsComplete;
+	}
+#endif
+	return InMaterial->IsComplete();
+}
+
+inline double GetCurrentTime()
+{
+	return FPlatformTime::Seconds();
+}
+
+} // UE::Ava::Private
+
+FMaterialCacheHelper& FMaterialCacheHelper::Get()
+{
+	static TSharedRef<FMaterialCacheHelper> MaterialCacheHelper = MakeShared<FMaterialCacheHelper>();
+	return *MaterialCacheHelper;
+}
+
+bool FMaterialCacheHelper::RequestCacheMaterials(const UObject* InObject, FName InShaderProfile)
+{
+	if (!InObject || Objects.Contains(InObject))
+	{
+		return false;
+	}
+
+	FObjectData ObjectData;
+	if (!GatherMaterialData(InObject, ObjectData, InShaderProfile))
+	{
+		return false;
+	}
+
+	Objects.Add(InObject, MoveTemp(ObjectData));
+	RegisterTick();
+	return true;
+}
+
+void FMaterialCacheHelper::Tick()
+{
+	if (bResetRequested)
+	{
+		bResetRequested = false;
+		Reset();
+		return;
+	}
+
+	TRACE_CPUPROFILER_EVENT_SCOPE(FMaterialCacheHelper::Tick);
+
+	// Duration of the frame so far prior to ticking
+	PreTickFrameDuration = (Private::GetCurrentTime() - BeginFrameTime);
+
+	// If the current frame time is already past the threshold, tick should be skipped unless it has been starving for a while already.
+	if (PreTickFrameDuration >= Private::CVarMaterialCacheFrameTimeStarveThreshold.GetValueOnGameThread())
+	{
+		if (FramesStarved < Private::CVarMaterialCacheTickMaxStarveFrames.GetValueOnGameThread())
+		{
+			++FramesStarved;
+			return;
+		}
+		UE_LOG(LogAvaMaterialCache, Verbose, TEXT("FMaterialCacheHelper::Tick was starved for %d frames. Forcing tick..."), FramesStarved);
+	}
+
+	FramesStarved = 0;
+	RemainingTimeBudget = Private::CVarMaterialCacheTickTimeBudget.GetValueOnGameThread();
+
+	ProcessPendingObjects();
+	ProcessMaterialsPendingCache();
+	ProcessMaterialsCaching();
+
+	if (IsCaching())
+	{
+		ProcessShadersCompiling();
+		RemoveCompletedMaterials();
+	}
+	else
+	{
+		bResetRequested = true;
+	}
+}
+
+bool FMaterialCacheHelper::IsCaching() const
+{
+	// If there is at least one element in either map, return true
+	if (!MaterialsPendingCache.IsEmpty() || !MaterialsCaching.IsEmpty())
+	{
+		return true;
+	}
+
+	// If there is one pending subobject still loading to cache, return true
+	for (const TPair<FObjectKey, FObjectData>& Object : Objects)
+	{
+		if (!Object.Value.PendingSubobjects.IsEmpty())
+		{
+			return true;
+		}
+	}
+	return false;
+}
+
+bool FMaterialCacheHelper::IsCaching(const UObject* InObject) const
+{
+	const FObjectData* ObjectData = Objects.Find(InObject);
+	if (!ObjectData)
+	{
+		return false;
+	}
+
+	// There are pending subobjects waiting on completion for material gathering
+	if (!ObjectData->PendingSubobjects.IsEmpty())
+	{
+		return true;
+	}
+
+	for (const FMaterialKey& MaterialKey : ObjectData->IncompleteMaterialKeys)
+	{
+		if (IsCachingMaterial(MaterialKey))
+		{
+			return true;
+		}
+	}
+	return false;
+}
+
+void FMaterialCacheHelper::DumpMaterials()
+{
+	auto PrintLog = [](const TPair<FMaterialKey, FMaterialData>& InPair)
+		{
+			const FMaterialData& MaterialData = InPair.Value;
+			UE_LOG(LogAvaMaterialCache, Log, TEXT("\t%s"), MaterialData.Material ? *MaterialData.Material->GetFullName() : TEXT("(invalid)"));
+		};
+
+	UE_LOG(LogAvaMaterialCache, Log, TEXT("--- Materials Pending Cache ------------------------------------"));
+	for (const TPair<FMaterialKey, FMaterialData>& Pair : MaterialsPendingCache)
+	{
+		PrintLog(Pair);
+	}
+	UE_LOG(LogAvaMaterialCache, Log, TEXT("--- Materials Caching ------------------------------------------"));
+	for (const TPair<FMaterialKey, FMaterialData>& Pair : MaterialsCaching)
+	{
+		PrintLog(Pair);
+	}
+}
+
+void FMaterialCacheHelper::AddReferencedObjects(FReferenceCollector& InCollector)
+{
+	for (TPair<FObjectKey, FObjectData>& Pair : Objects)
+	{
+		InCollector.AddReferencedObjects(Pair.Value.PendingSubobjects);
+	}
+	for (TPair<FMaterialKey, FMaterialData>& Pair : MaterialsPendingCache)
+	{
+		InCollector.AddReferencedObject(Pair.Value.Material);
+	}
+	for (TPair<FMaterialKey, FMaterialData>& Pair : MaterialsCaching)
+	{
+		InCollector.AddReferencedObject(Pair.Value.Material);
+	}
+}
+
+FString FMaterialCacheHelper::GetReferencerName() const
+{
+	return TEXT("UE::Ava::FMaterialCacheHelper");
+}
+
+bool FMaterialCacheHelper::IsCachingMaterial(const FMaterialKey& InMaterialKey) const
+{
+	return MaterialsPendingCache.Contains(InMaterialKey) || MaterialsCaching.Contains(InMaterialKey);
+}
+
+void FMaterialCacheHelper::Reset()
+{
+	MaterialsPendingCache.Empty();
+	MaterialsCaching.Empty();
+	Objects.Empty();
+	UnregisterTick();
+}
+
+void FMaterialCacheHelper::RegisterTick()
+{
+	bResetRequested = false;
+
+	if (!BeginFrameHandle.IsValid())
+	{
+		// It is likely RegisterTick is called between OnBeginFrame and OnEndFrame delegate broadcasts.
+		// This will cause BeginFrameTime to have its last value in the first tick (e.g. 0.0, or an old begin frame's time) and likely starve tick for that one frame.
+		// And this is ok. An alternative could be to use a timer that is set at BeginFrame always.  
+		BeginFrameHandle = FCoreDelegates::OnBeginFrame.AddSPLambda(this, [this]
+			{
+				BeginFrameTime = Private::GetCurrentTime();
+			});
+	}
+
+	if (!EndFrameHandle.IsValid())
+	{
+		EndFrameHandle = FCoreDelegates::OnEndFrame.AddSP(this, &FMaterialCacheHelper::Tick);
+	}
+}
+
+void FMaterialCacheHelper::UnregisterTick()
+{
+	if (BeginFrameHandle.IsValid())
+	{
+		FCoreDelegates::OnBeginFrame.Remove(BeginFrameHandle);
+		BeginFrameHandle.Reset();
+	}
+
+	if (EndFrameHandle.IsValid())
+	{
+		FCoreDelegates::OnEndFrame.Remove(EndFrameHandle);
+		EndFrameHandle.Reset();
+	}
+}
+
+void FMaterialCacheHelper::ProcessPendingObjects()
+{
+	if (RemainingTimeBudget <= 0.f || Objects.IsEmpty())
+	{
+		return;
+	}
+
+	// Time when budget runs out
+	const double EndBudgetTime = Private::GetCurrentTime() + RemainingTimeBudget;
+	bool bTimeBudgetExceeded = false;
+
+	// Check the pending status every 100ms (0.1s) to give it some time to load
+	constexpr float PendingSubobjectCheckTime = 0.1f;
+
+	for (TPair<FObjectKey, FObjectData>& Object : Objects)
+	{
+		Object.Value.PendingSubobjectElapsedTime += PreTickFrameDuration;
+		if (Object.Value.PendingSubobjectElapsedTime < PendingSubobjectCheckTime)
+		{
+			continue;
+		}
+		Object.Value.PendingSubobjectElapsedTime = 0.f;
+
+		// All pending subobjects might be ready so move them into a temp set.
+		// This also allows adding elements to the existing set.
+		// Those that are still not gathered pending load will be re-added to the set.
+		TSet<TObjectPtr<const UObject>> PendingSubobjects = MoveTemp(Object.Value.PendingSubobjects);
+		Object.Value.PendingSubobjects.Empty();
+
+		for (TSet<TObjectPtr<const UObject>>::TIterator Iter(PendingSubobjects); Iter; ++Iter)
+		{
+			GatherMaterialData(*Iter, Object.Value, Object.Value.ShaderProfile);
+			Iter.RemoveCurrent();
+
+			// Check if time is past end budget time
+			bTimeBudgetExceeded = Private::GetCurrentTime() >= EndBudgetTime;
+			if (bTimeBudgetExceeded)
+			{
+				UE_LOG(LogAvaMaterialCache, Verbose, TEXT("FMaterialCacheHelper::Tick ran out of budget in ProcessPendingObjects. Time Budget: %f"), RemainingTimeBudget);
+				break;
+			}
+		}
+
+		// Add pending subobjects that could not be processed due to time budget back
+		Object.Value.PendingSubobjects.Append(MoveTemp(PendingSubobjects));
+		if (bTimeBudgetExceeded)
+		{
+			break;
+		}
+	}
+
+	RemainingTimeBudget = EndBudgetTime - Private::GetCurrentTime();
+}
+
+void FMaterialCacheHelper::ProcessMaterialsPendingCache()
+{
+	if (RemainingTimeBudget <= 0.f)
+	{
+		return;
+	}
+
+	TRACE_CPUPROFILER_EVENT_SCOPE(FMaterialCacheHelper::ProcessMaterialsPendingCache);
+
+	// Time when budget runs out
+	const double EndBudgetTime = Private::GetCurrentTime() + RemainingTimeBudget;
+
+	// Material update context. Only allocated when needed
+	TOptional<FMaterialUpdateContext> UpdateContext;
+
+	for (TMap<FMaterialKey, FMaterialData>::TIterator Iter(MaterialsPendingCache); Iter; ++Iter)
+	{
+		const FMaterialKey& MaterialKey = Iter.Key();
+		FMaterialData& MaterialData = Iter.Value();
+
+		if (!MaterialData.Material)
+		{
+			Iter.RemoveCurrent();
+			continue;
+		}
+
+		const FResolvedShaderProfile* ShaderProfile = Private::FindShaderProfile(MaterialKey.ShaderProfile);
+		if (!ShaderProfile || ShaderProfile->ProfileType == EAvaShaderProfileType::SkipCaching)
+		{
+			Iter.RemoveCurrent();
+			continue;
+		}
+
+		if (!UpdateContext.IsSet())
+		{
+			UpdateContext.Emplace(0);
+		}
+		UpdateContext->AddMaterialInterface(MaterialData.Material);
+		Private::CacheMaterial(MaterialData.Material, ShaderProfile);
+
+		MaterialsCaching.Add(MaterialKey, MoveTemp(MaterialData));
+		Iter.RemoveCurrent();
+
+		// Check if time is past end budget time
+		// This is only checked for materials waiting, since the other operations are fast.
+		if (Private::GetCurrentTime() >= EndBudgetTime)
+		{
+			UE_LOG(LogAvaMaterialCache, Verbose, TEXT("FMaterialCacheHelper::Tick ran out of budget in ProcessMaterials. Time Budget: %f"), RemainingTimeBudget);
+			break;
+		}
+	}
+
+	RemainingTimeBudget = EndBudgetTime - Private::GetCurrentTime();
+}
+
+void FMaterialCacheHelper::ProcessMaterialsCaching()
+{
+	if (RemainingTimeBudget <= 0.f)
+	{
+		return;
+	}
+
+	TRACE_CPUPROFILER_EVENT_SCOPE(FMaterialCacheHelper::ProcessCachingMaterials);
+
+	// Time when budget runs out
+	const double StartTime = Private::GetCurrentTime();
+
+	for (TMap<FMaterialKey, FMaterialData>::TIterator Iter(MaterialsCaching); Iter; ++Iter)
+	{
+		const FMaterialKey& MaterialKey = Iter.Key();
+		FMaterialData& MaterialData = Iter.Value();
+
+		if (!MaterialData.Material)
+		{
+			Iter.RemoveCurrent();
+			continue;
+		}
+
+		const FResolvedShaderProfile* ShaderProfile = Private::FindShaderProfile(MaterialKey.ShaderProfile);
+		if (!ShaderProfile)
+		{
+			Iter.RemoveCurrent();
+			continue;
+		}
+
+		if (Private::IsMaterialComplete(MaterialData.Material, ShaderProfile))
+		{
+			Iter.RemoveCurrent();
+		}
+	}
+
+	RemainingTimeBudget -= (Private::GetCurrentTime() - StartTime);
+}
+
+void FMaterialCacheHelper::RemoveCompletedMaterials()
+{
+	if (RemainingTimeBudget <= 0.f)
+	{
+		return;
+	}
+
+	const double StartTime = Private::GetCurrentTime();
+
+	for (TPair<FObjectKey, FObjectData>& Object : Objects)
+	{
+		for (TSet<FMaterialKey>::TIterator MaterialKeyIter(Object.Value.IncompleteMaterialKeys); MaterialKeyIter; ++MaterialKeyIter)
+		{
+			if (!IsCachingMaterial(*MaterialKeyIter))
+			{
+				MaterialKeyIter.RemoveCurrent();
+			}
+		}
+	}
+
+	RemainingTimeBudget -= (Private::GetCurrentTime() - StartTime);
+}
+
+void FMaterialCacheHelper::ProcessShadersCompiling()
+{
+	if (RemainingTimeBudget <= 0.f || !GShaderCompilingManager || !GShaderCompilingManager->IsCompiling())
+	{
+		return;
+	}
+
+	TRACE_CPUPROFILER_EVENT_SCOPE(FMaterialCacheHelper::ProcessShadersCompiling);
+
+	const double StartTime = Private::GetCurrentTime();
+
+	constexpr bool bBlockOnGlobalShaderCompletion = false;
+	GShaderCompilingManager->ProcessAsyncResults(RemainingTimeBudget, bBlockOnGlobalShaderCompletion);
+
+	RemainingTimeBudget -= (Private::GetCurrentTime() - StartTime);
+}
+
+bool FMaterialCacheHelper::GatherMaterialData(const UObject* InObject, FObjectData& InObjectData, FName InShaderProfile)
+{
+	if (!InObject)
+	{
+		return false;
+	}
+
+	TRACE_CPUPROFILER_EVENT_SCOPE(FMaterialCacheHelper::GatherMaterialData);
+
+	const FMaterialBridge* MaterialBridge = FMaterialBridgeRegistry::Get().GetMaterialBridgeForObject(InObject);
+	if (!MaterialBridge)
+	{
+		UE_LOG(LogAvaMaterialCache, VeryVerbose, TEXT("Object '%s' (Class: %s) does not have a valid bridge to cache materials!")
+			, *InObject->GetName()
+			, *GetNameSafe(InObject->GetClass()));
+		return false;
+	}
+
+	InObjectData.ShaderProfile = InShaderProfile;
+
+	auto AddIncompleteMaterial =
+		[&MaterialsPendingCache = MaterialsPendingCache, &InObjectData, InShaderProfile](TNotNull<UMaterialInterface*> InMaterial)
+		{
+			FMaterialKey MaterialKey;
+			MaterialKey.Material = InMaterial;
+			MaterialKey.ShaderProfile = InShaderProfile;
+
+			FMaterialData& MaterialData = MaterialsPendingCache.FindOrAdd(MaterialKey);
+			MaterialData.Material = InMaterial;
+			InObjectData.IncompleteMaterialKeys.Add(MaterialKey);
+		};
+
+	// Options to not wait for the container to have materials available and instead get a callback
+	FMaterialBridge::FVisitOptions VisitOptions;
+	VisitOptions.bTryWaitOnContainerCompletion = false;
+	VisitOptions.OnContainerPendingCompletion.BindLambda([&InObjectData](const FMaterialBridgeVisitContext& InContext)->EVisitFlow
+		{
+			if (const UObject* Container = InContext.Container.GetPtr<const UObject>())
+			{
+				InObjectData.PendingSubobjects.Add(Container);
+			}
+			else
+			{
+				ensureMsgf(0, TEXT("Container of type '%s' not supported for deferring material gather"), *GetNameSafe(InContext.Container.GetStruct()));
+			}
+			return EVisitFlow::Continue;
+		});
+
+	// Gather all materials for object
+	MaterialBridge->VisitMaterials(FMaterialBridgeVisitContext(const_cast<UObject*>(InObject)),
+		[&AddIncompleteMaterial](const FMaterialBridgeVisitContext& InContext, UMaterialInterface* InMaterial)->EVisitFlow
+		{
+			if (!InMaterial)
+			{
+				return EVisitFlow::Continue;
+			}
+
+			UMaterial* const Material = InMaterial->GetMaterial();
+
+			// Add interface's material if it is different from the interface
+			if (Material && Material != InMaterial && !Material->IsComplete())
+			{
+				AddIncompleteMaterial(Material);
+			}
+			if (!InMaterial->IsComplete())
+			{
+				AddIncompleteMaterial(InMaterial);
+			}
+			return EVisitFlow::Continue;
+		}
+		, VisitOptions);
+
+	return true;
+}
+
+} // UE::Ava
diff --git a/Engine/Plugins/VirtualProduction/Avalanche/Source/AvalancheMaterial/Private/MaterialCache/AvaMaterialCacheSettings.cpp b/Engine/Plugins/VirtualProduction/Avalanche/Source/AvalancheMaterial/Private/MaterialCache/AvaMaterialCacheSettings.cpp
new file mode 100644
--- /dev/null
+++ b/Engine/Plugins/VirtualProduction/Avalanche/Source/AvalancheMaterial/Private/MaterialCache/AvaMaterialCacheSettings.cpp
@@ -0,0 +1,225 @@
+// Copyright Epic Games, Inc. All Rights Reserved.
+
+#include "MaterialCache/AvaMaterialCacheSettings.h"
+#include "MaterialCache/AvaMaterialCacheLog.h"
+#include "VertexFactory.h"
+
+const FLazyName UAvaMaterialCacheSettings::DefaultRealtimeProfile(TEXT("BasicSceneRealtime"));
+const FLazyName UAvaMaterialCacheSettings::DefaultOfflineProfile(TEXT("BasicSceneOffline"));
+
+namespace UE::Ava::Private
+{
+
+TArray<FAvaShaderCombination> MakeDefaultShaderCombinations()
+{
+	// Vertex Factories
+	const FName LocalVertexFactory       = TEXT("FLocalVertexFactory");
+	const FName NiagaraMeshVertexFactory = TEXT("FNiagaraMeshVertexFactory");
+
+	// Pipelines
+	const FName StandardVelocityPipeline = TEXT("StandardVelocityPipeline");
+	const FName DepthPipeline            = TEXT("DepthPipeline");
+
+	// Shaders
+	const FName BasePassVS               = TEXT("TBasePassVSFNoLightMapPolicy");
+	const FName BasePassPS               = TEXT("TBasePassPSFNoLightMapPolicy");
+	const FName BasePassPSSL             = TEXT("TBasePassPSFNoLightMapPolicySkylight");
+	const FName NiagaraShader            = TEXT("FNiagaraShader");
+	const FName MaterialCHS              = TEXT("TMaterialCHSFNoLightMapPolicyFAnyHitShader");
+	const FName LumenCardVS              = TEXT("FLumenCardVS");
+	const FName LumenCardPS              = TEXT("FLumenCardPS");
+	const FName VelocityVS               = TEXT("TVelocityVS<EVelocityPassMode::Velocity_Standard>");
+	const FName VelocityPS               = TEXT("TVelocityPS<EVelocityPassMode::Velocity_Standard>");
+	const FName DepthOnlyVS              = TEXT("TDepthOnlyVS<false>");
+	const FName DepthOnlyPS              = TEXT("FDepthOnlyPS");
+	const FName ShadowDepthVS            = TEXT("TShadowDepthVSCombined");
+	const FName ShadowDepthPS            = TEXT("TShadowDepthPSCombined");
+
+	return 
+	{
+		FAvaShaderCombination(LocalVertexFactory,
+			{
+				StandardVelocityPipeline,
+				DepthPipeline,
+			},
+			{
+				BasePassVS,
+				BasePassPS,
+				BasePassPSSL,
+				NiagaraShader,
+				MaterialCHS,
+				LumenCardVS,
+				LumenCardPS,
+				VelocityVS,
+				VelocityPS,
+				DepthOnlyVS,
+				DepthOnlyPS,
+				ShadowDepthVS,
+				ShadowDepthPS,
+			}
+		),
+		FAvaShaderCombination(NiagaraMeshVertexFactory,
+			{
+				StandardVelocityPipeline
+			},
+			{
+				BasePassVS,
+				BasePassPS,
+				ShadowDepthVS,
+				ShadowDepthPS,
+				VelocityVS,
+				VelocityPS,
+			}
+		)
+	};
+}
+
+} // UE::Ava::Private
+
+UAvaMaterialCacheSettings::UAvaMaterialCacheSettings()
+{
+	RealtimeProfile = UAvaMaterialCacheSettings::DefaultRealtimeProfile;
+	OfflineProfile = UAvaMaterialCacheSettings::DefaultOfflineProfile;
+
+	CategoryName = TEXT("Motion Design");
+	SectionName  = TEXT("Material Cache");
+	RegisterDefaultShaderProfile();
+}
+
+const UE::Ava::FResolvedShaderProfile* UAvaMaterialCacheSettings::FindResolvedShaderProfile(FName InProfileName) const
+{
+	return ResolvedShaderProfiles.Find(InProfileName);
+}
+
+void UAvaMaterialCacheSettings::PostInitProperties()
+{
+	Super::PostInitProperties();
+	ResolveShaderProfiles();
+}
+
+void UAvaMaterialCacheSettings::PostReloadConfig(FProperty* InPropertyThatWasLoaded)
+{
+	Super::PostReloadConfig(InPropertyThatWasLoaded);
+	ResolveShaderProfiles();
+}
+
+#if WITH_EDITOR
+void UAvaMaterialCacheSettings::PostEditChangeProperty(FPropertyChangedEvent& InPropertyChangedEvent)
+{
+	Super::PostEditChangeProperty(InPropertyChangedEvent);
+
+	if (InPropertyChangedEvent.GetMemberPropertyName() == GET_MEMBER_NAME_CHECKED(UAvaMaterialCacheSettings, ShaderProfiles))
+	{
+		ResolveShaderProfiles();
+	}
+}
+#endif
+
+void UAvaMaterialCacheSettings::RegisterDefaultShaderProfile()
+{
+	const TArray<FAvaShaderCombination> DefaultShaderCombinations = UE::Ava::Private::MakeDefaultShaderCombinations();
+
+	ShaderProfiles.Add(
+		{
+			.ProfileName = UAvaMaterialCacheSettings::DefaultRealtimeProfile,
+			.ProfileType = EAvaShaderProfileType::Specific,
+			.ShaderCombinations = DefaultShaderCombinations,
+		});
+
+	ShaderProfiles.Add(
+		{
+			.ProfileName = UAvaMaterialCacheSettings::DefaultOfflineProfile,
+			.ProfileType = EAvaShaderProfileType::Specific,
+			.ShaderCombinations = DefaultShaderCombinations,
+		});
+}
+
+void UAvaMaterialCacheSettings::ResolveShaderProfiles()
+{
+	UE_LOG(LogAvaMaterialCache, Verbose, TEXT("Resolving Shader Profiles"));
+
+	ResolvedShaderProfiles.Empty(ShaderProfiles.Num());
+
+	TMap<FName, FHashedName> HashedNames;
+	for (const FAvaShaderProfile& ShaderProfile : ShaderProfiles)
+	{
+		ResolveShaderProfile(ShaderProfile, HashedNames);
+	}
+}
+
+void UAvaMaterialCacheSettings::ResolveShaderProfile(const FAvaShaderProfile& InShaderProfile, TMap<FName, FHashedName>& InOutHashedNames)
+{
+	if (ResolvedShaderProfiles.Contains(InShaderProfile.ProfileName))
+	{
+		UE_LOG(LogAvaMaterialCache, Warning, TEXT("Shader Profile '%s' already exists. New entry will be ignored."), *InShaderProfile.ProfileName.ToString());
+		return;
+	}
+
+	UE::Ava::FResolvedShaderProfile& ResolvedShaderProfile = ResolvedShaderProfiles.Add(InShaderProfile.ProfileName);
+	ResolvedShaderProfile.ProfileType = InShaderProfile.ProfileType;
+
+	// Keep track of FHashedName instances to avoid recomputation
+	auto GetHashedName = [&InOutHashedNames](FName InName)
+		{
+			if (const FHashedName* HashedName = InOutHashedNames.Find(InName))
+			{
+				return *HashedName;
+			}
+			return InOutHashedNames.Add(InName, InName);
+		};
+
+	int32 ReserveCount = 0;
+	for (const FAvaShaderCombination& ShaderCombination : InShaderProfile.ShaderCombinations)
+	{
+		ReserveCount += ShaderCombination.PipelineTypes.Num();
+		ReserveCount += ShaderCombination.ShaderTypes.Num();
+	}
+
+	// The types must all have the same count at the end
+	ResolvedShaderProfile.VertexFactoryTypes.Reserve(ReserveCount);
+	ResolvedShaderProfile.PipelineTypes.Reserve(ReserveCount);
+	ResolvedShaderProfile.ShaderTypes.Reserve(ReserveCount);
+
+	for (const FAvaShaderCombination& ShaderCombination : InShaderProfile.ShaderCombinations)
+	{
+		const FVertexFactoryType* const VertexFactoryType = FVertexFactoryType::GetVFByName(GetHashedName(ShaderCombination.VertexFactoryType));
+		if (!VertexFactoryType)
+		{
+			UE_LOG(LogAvaMaterialCache, Warning, TEXT("Invalid Shader Config. Vertex Factory Type '%s' was not found"), *ShaderCombination.VertexFactoryType.ToString());
+			continue;
+		}
+
+		for (FName PipelineTypeName : ShaderCombination.PipelineTypes)
+		{
+			const FShaderPipelineType* PipelineType = FShaderPipelineType::GetShaderPipelineTypeByName(GetHashedName(PipelineTypeName));
+			if (!PipelineType)
+			{
+				UE_LOG(LogAvaMaterialCache, Warning, TEXT("Invalid Shader Config. Pipeline Type '%s' was found"), *PipelineTypeName.ToString());
+				continue;
+			}
+
+			ResolvedShaderProfile.VertexFactoryTypes.Add(VertexFactoryType);
+			ResolvedShaderProfile.PipelineTypes.Add(PipelineType);
+			ResolvedShaderProfile.ShaderTypes.Add(nullptr);
+		}
+
+		for (FName ShaderTypeName : ShaderCombination.ShaderTypes)
+		{
+			const FShaderType* ShaderType = FindShaderTypeByName(GetHashedName(ShaderTypeName));
+			if (!ShaderType)
+			{
+				UE_LOG(LogAvaMaterialCache, Warning, TEXT("Invalid Shader Config. Shader Type '%s' was not found"), *ShaderTypeName.ToString());
+				continue;
+			}
+
+			ResolvedShaderProfile.VertexFactoryTypes.Add(VertexFactoryType);
+			ResolvedShaderProfile.PipelineTypes.Add(nullptr);
+			ResolvedShaderProfile.ShaderTypes.Add(ShaderType);
+		}
+	}
+
+	ensureMsgf(ResolvedShaderProfile.VertexFactoryTypes.Num() == ResolvedShaderProfile.PipelineTypes.Num()
+		&& ResolvedShaderProfile.VertexFactoryTypes.Num() == ResolvedShaderProfile.ShaderTypes.Num()
+		, TEXT("Resolved Shader Profile '%s' unexpectedly has mismatched number of types.")
+		, *InShaderProfile.ProfileName.ToString());
+}
diff --git a/Engine/Plugins/VirtualProduction/Avalanche/Source/AvalancheMaterial/Public/AvaActorMaterialBridge.h b/Engine/Plugins/VirtualProduction/Avalanche/Source/AvalancheMaterial/Public/AvaActorMaterialBridge.h
deleted file mode 100644
diff --git a/Engine/Plugins/VirtualProduction/Avalanche/Source/AvalancheMaterial/Public/AvaDecalComponentMaterialBridge.h b/Engine/Plugins/VirtualProduction/Avalanche/Source/AvalancheMaterial/Public/AvaDecalComponentMaterialBridge.h
deleted file mode 100644
diff --git a/Engine/Plugins/VirtualProduction/Avalanche/Source/AvalancheMaterial/Public/AvaLevelMaterialBridge.h b/Engine/Plugins/VirtualProduction/Avalanche/Source/AvalancheMaterial/Public/AvaLevelMaterialBridge.h
deleted file mode 100644
diff --git a/Engine/Plugins/VirtualProduction/Avalanche/Source/AvalancheMaterial/Public/AvaMaterialBridge.h b/Engine/Plugins/VirtualProduction/Avalanche/Source/AvalancheMaterial/Public/AvaMaterialBridge.h
deleted file mode 100644
diff --git a/Engine/Plugins/VirtualProduction/Avalanche/Source/AvalancheMaterial/Public/AvaMaterialBridgeLog.h b/Engine/Plugins/VirtualProduction/Avalanche/Source/AvalancheMaterial/Public/AvaMaterialBridgeLog.h
deleted file mode 100644
diff --git a/Engine/Plugins/VirtualProduction/Avalanche/Source/AvalancheMaterial/Public/AvaMaterialBridgeRegistry.h b/Engine/Plugins/VirtualProduction/Avalanche/Source/AvalancheMaterial/Public/AvaMaterialBridgeRegistry.h
deleted file mode 100644
diff --git a/Engine/Plugins/VirtualProduction/Avalanche/Source/AvalancheMaterial/Public/AvaMaterialBridgeVisitContext.h b/Engine/Plugins/VirtualProduction/Avalanche/Source/AvalancheMaterial/Public/AvaMaterialBridgeVisitContext.h
deleted file mode 100644
diff --git a/Engine/Plugins/VirtualProduction/Avalanche/Source/AvalancheMaterial/Public/AvaPrimitiveComponentMaterialBridge.h b/Engine/Plugins/VirtualProduction/Avalanche/Source/AvalancheMaterial/Public/AvaPrimitiveComponentMaterialBridge.h
deleted file mode 100644
diff --git a/Engine/Plugins/VirtualProduction/Avalanche/Source/AvalancheMaterial/Public/AvaStaticMeshMaterialBridge.h b/Engine/Plugins/VirtualProduction/Avalanche/Source/AvalancheMaterial/Public/AvaStaticMeshMaterialBridge.h
deleted file mode 100644
diff --git a/Engine/Plugins/VirtualProduction/Avalanche/Source/AvalancheMaterial/Public/MaterialBridge/AvaActorMaterialBridge.h b/Engine/Plugins/VirtualProduction/Avalanche/Source/AvalancheMaterial/Public/MaterialBridge/AvaActorMaterialBridge.h
new file mode 100644
--- /dev/null
+++ b/Engine/Plugins/VirtualProduction/Avalanche/Source/AvalancheMaterial/Public/MaterialBridge/AvaActorMaterialBridge.h
@@ -0,0 +1,24 @@
+// Copyright Epic Games, Inc. All Rights Reserved.
+
+#pragma once
+
+#include "AvaMaterialBridge.h"
+
+#define UE_API AVALANCHEMATERIAL_API
+
+namespace UE::Ava
+{
+
+/** Material Bridge for Actors */
+class FActorMaterialBridge : public FMaterialBridge
+{
+protected:
+	//~ Begin FMaterialBridge
+	UE_API virtual const UStruct* OnGetBridgedType() const override;
+	UE_API virtual EVisitFlow OnVisitMaterials(const FMaterialBridgeVisitContext& InContext, TFunctionRef<EVisitFlow(const FMaterialBridgeVisitContext&, UMaterialInterface*)> InFunc, const FVisitOptions& InOptions) const override;
+	//~ End FMaterialBridge
+};
+
+} // UE::Ava
+
+#undef UE_API
diff --git a/Engine/Plugins/VirtualProduction/Avalanche/Source/AvalancheMaterial/Public/MaterialBridge/AvaDecalComponentMaterialBridge.h b/Engine/Plugins/VirtualProduction/Avalanche/Source/AvalancheMaterial/Public/MaterialBridge/AvaDecalComponentMaterialBridge.h
new file mode 100644
--- /dev/null
+++ b/Engine/Plugins/VirtualProduction/Avalanche/Source/AvalancheMaterial/Public/MaterialBridge/AvaDecalComponentMaterialBridge.h
@@ -0,0 +1,24 @@
+// Copyright Epic Games, Inc. All Rights Reserved.
+
+#pragma once
+
+#include "AvaMaterialBridge.h"
+
+#define UE_API AVALANCHEMATERIAL_API
+
+namespace UE::Ava
+{
+
+/** Material Bridge for Decal Components */
+class FDecalComponentMaterialBridge : public FMaterialBridge
+{
+protected:
+	//~ Begin FMaterialBridge
+	UE_API virtual const UStruct* OnGetBridgedType() const override;
+	UE_API virtual EVisitFlow OnVisitMaterials(const FMaterialBridgeVisitContext& InContext, TFunctionRef<EVisitFlow(const FMaterialBridgeVisitContext&, UMaterialInterface*)> InFunc, const FVisitOptions& InOptions) const override;
+	//~ End FMaterialBridge
+};
+
+} // UE::Ava
+
+#undef UE_API
diff --git a/Engine/Plugins/VirtualProduction/Avalanche/Source/AvalancheMaterial/Public/MaterialBridge/AvaLevelMaterialBridge.h b/Engine/Plugins/VirtualProduction/Avalanche/Source/AvalancheMaterial/Public/MaterialBridge/AvaLevelMaterialBridge.h
new file mode 100644
--- /dev/null
+++ b/Engine/Plugins/VirtualProduction/Avalanche/Source/AvalancheMaterial/Public/MaterialBridge/AvaLevelMaterialBridge.h
@@ -0,0 +1,24 @@
+// Copyright Epic Games, Inc. All Rights Reserved.
+
+#pragma once
+
+#include "AvaMaterialBridge.h"
+
+#define UE_API AVALANCHEMATERIAL_API
+
+namespace UE::Ava
+{
+
+/** Material Bridge for Levels */
+class FLevelMaterialBridge : public FMaterialBridge
+{
+protected:
+	//~ Begin FMaterialBridge
+	UE_API virtual const UStruct* OnGetBridgedType() const override;
+	UE_API virtual EVisitFlow OnVisitMaterials(const FMaterialBridgeVisitContext& InContext, TFunctionRef<EVisitFlow(const FMaterialBridgeVisitContext&, UMaterialInterface*)> InFunc, const FVisitOptions& InOptions) const override;
+	//~ End FMaterialBridge
+};
+
+} // UE::Ava
+
+#undef UE_API
diff --git a/Engine/Plugins/VirtualProduction/Avalanche/Source/AvalancheMaterial/Public/MaterialBridge/AvaMaterialBridge.h b/Engine/Plugins/VirtualProduction/Avalanche/Source/AvalancheMaterial/Public/MaterialBridge/AvaMaterialBridge.h
new file mode 100644
--- /dev/null
+++ b/Engine/Plugins/VirtualProduction/Avalanche/Source/AvalancheMaterial/Public/MaterialBridge/AvaMaterialBridge.h
@@ -0,0 +1,81 @@
+// Copyright Epic Games, Inc. All Rights Reserved.
+
+#pragma once
+
+#include "AvaMaterialBridgeVisitContext.h"
+#include "Containers/ContainersFwd.h"
+#include "StructUtils/StructUtilsTypes.h"
+#include "UObject/ObjectPtr.h"
+
+#define UE_API AVALANCHEMATERIAL_API
+
+class UMaterialInterface;
+
+namespace UE::Ava
+{
+	class FMaterialBridgeRegistry;
+
+} // UE::Ava
+
+namespace UE::Ava
+{
+
+/** Indicates how iteration should be handled */
+enum class EVisitFlow : uint8
+{
+	/** Continue iteration */
+	Continue,
+	/** Break iteration */
+	Break,
+};
+
+/**
+ * Base class that allows interacting with materials from any type of data view (Actors, ActorComponents, etc).
+ */
+class FMaterialBridge
+{
+public:
+	virtual ~FMaterialBridge() = default;
+
+	/** Initializes the bridge for use */
+	void Initialize();
+
+	/** Retrieves the type that this bridge supports */
+	UE_API const UStruct* GetBridgedType() const;
+
+	struct FVisitOptions
+	{
+		UE_API FVisitOptions();
+		/** The registry to use to visit any materials of sub objects. Defaults to the global registry. */
+		TNotNull<const FMaterialBridgeRegistry*> MaterialBridgeRegistry;
+		/** Callback when a container is not yet complete and cannot access its materials. Only called when bTryWaitOnContainerCompletion is false. */
+		TDelegate<EVisitFlow(const FMaterialBridgeVisitContext&)> OnContainerPendingCompletion;
+		/** Whether to try to wait for container completion to access its materials */
+		bool bTryWaitOnContainerCompletion = true;
+		/** Whether to include debug materials */
+		bool bVisitDebugMaterials = false;
+	};
+	/**
+	 * Visits all the used materials in the given context.
+	 * @param InContext the current context to find materials from.
+	 * @param InFunc function to call for each material found.
+	 * @param InOptions visit options. see FVisitOptions
+	 * @return the iteration result of VisitMaterials.
+	 */
+	UE_API EVisitFlow VisitMaterials(const FMaterialBridgeVisitContext& InContext, TFunctionRef<EVisitFlow(const FMaterialBridgeVisitContext&, UMaterialInterface*)> InFunc, const FVisitOptions& InOptions = {}) const;
+
+protected:
+	/** Retrieves the type that this bridge supports. Only called once in Initialize() */
+	virtual const UStruct* OnGetBridgedType() const = 0;
+
+	/** Get all the used materials in the given context. */
+	virtual EVisitFlow OnVisitMaterials(const FMaterialBridgeVisitContext& InContext, TFunctionRef<EVisitFlow(const FMaterialBridgeVisitContext&, UMaterialInterface*)> InFunc, const FVisitOptions& InOptions) const = 0;
+
+private:
+	/** The type that this Material Bridge works on */
+	TObjectPtr<const UStruct> BridgedType;
+};
+
+} // UE::Ava
+
+#undef UE_API
diff --git a/Engine/Plugins/VirtualProduction/Avalanche/Source/AvalancheMaterial/Public/MaterialBridge/AvaMaterialBridgeLog.h b/Engine/Plugins/VirtualProduction/Avalanche/Source/AvalancheMaterial/Public/MaterialBridge/AvaMaterialBridgeLog.h
new file mode 100644
--- /dev/null
+++ b/Engine/Plugins/VirtualProduction/Avalanche/Source/AvalancheMaterial/Public/MaterialBridge/AvaMaterialBridgeLog.h
@@ -0,0 +1,7 @@
+// Copyright Epic Games, Inc. All Rights Reserved.
+
+#pragma once
+
+#include "Logging/LogCategory.h"
+
+DECLARE_LOG_CATEGORY_EXTERN(LogAvaMaterialBridge, Log, All);
diff --git a/Engine/Plugins/VirtualProduction/Avalanche/Source/AvalancheMaterial/Public/MaterialBridge/AvaMaterialBridgeRegistry.h b/Engine/Plugins/VirtualProduction/Avalanche/Source/AvalancheMaterial/Public/MaterialBridge/AvaMaterialBridgeRegistry.h
new file mode 100644
--- /dev/null
+++ b/Engine/Plugins/VirtualProduction/Avalanche/Source/AvalancheMaterial/Public/MaterialBridge/AvaMaterialBridgeRegistry.h
@@ -0,0 +1,49 @@
+// Copyright Epic Games, Inc. All Rights Reserved.
+
+#pragma once
+
+#include "Concepts/DerivedFrom.h"
+#include "Containers/Map.h"
+#include "Templates/SharedPointerFwd.h"
+#include "UObject/ObjectKey.h"
+
+#define UE_API AVALANCHEMATERIAL_API
+
+namespace UE::Ava
+{
+	class FMaterialBridge;
+}
+
+namespace UE::Ava
+{
+
+/** Holds all the Material Bridges to use */
+class FMaterialBridgeRegistry
+{
+public:
+	UE_API static const FMaterialBridgeRegistry& Get();
+	UE_API static FMaterialBridgeRegistry& GetMutable();
+
+	template<typename InMaterialBridgeType, typename... InArgTypes> requires (UE::CDerivedFrom<InMaterialBridgeType, UE::Ava::FMaterialBridge>)
+	void Register(InArgTypes&&... InArgs)
+	{
+		this->RegisterInternal(MakeUnique<InMaterialBridgeType>(InArgs...));
+	}
+
+	/** Tries to find the most relevant material bridge for a given object, returning a null if object is invalid, or if a relevant material bridge was not found. */
+	UE_API const FMaterialBridge* GetMaterialBridgeForObject(const UObject* InObject) const;
+
+	/** Tries to find the most relevant material bridge for a given type, returning null if a relevant material bridge was not found. */
+	UE_API const FMaterialBridge* GetMaterialBridgeForType(const UStruct* InType) const;
+
+protected:
+	UE_API void RegisterInternal(TUniquePtr<UE::Ava::FMaterialBridge>&& InMaterialBridge);
+
+private:
+	/** Registered material bridges by their supported type */
+	TMap<TObjectKey<UStruct>, TUniquePtr<FMaterialBridge>> RegisteredMaterialBridges;
+};
+
+} // UE::Ava
+
+#undef UE_API
diff --git a/Engine/Plugins/VirtualProduction/Avalanche/Source/AvalancheMaterial/Public/MaterialBridge/AvaMaterialBridgeVisitContext.h b/Engine/Plugins/VirtualProduction/Avalanche/Source/AvalancheMaterial/Public/MaterialBridge/AvaMaterialBridgeVisitContext.h
new file mode 100644
--- /dev/null
+++ b/Engine/Plugins/VirtualProduction/Avalanche/Source/AvalancheMaterial/Public/MaterialBridge/AvaMaterialBridgeVisitContext.h
@@ -0,0 +1,30 @@
+// Copyright Epic Games, Inc. All Rights Reserved.
+
+#pragma once
+
+#include "AvaDataView.h"
+
+#define UE_API AVALANCHEMATERIAL_API
+
+namespace UE::Ava
+{
+
+/** Short-lived context information of a visited material */
+class FMaterialBridgeVisitContext
+{
+public:
+	FMaterialBridgeVisitContext() = default;
+
+	explicit UE_API FMaterialBridgeVisitContext(UObject* InContainer);
+	explicit UE_API FMaterialBridgeVisitContext(UObject* InContainer, TNotNull<const FMaterialBridgeVisitContext*> InParentContext);
+
+	/** Container of the visited material */
+	FDataView Container;
+
+	/** Parent of this context for additional information (e.g. container of the container) */
+	const FMaterialBridgeVisitContext* ParentContext = nullptr;
+};
+
+} // UE::Ava
+
+#undef UE_API
diff --git a/Engine/Plugins/VirtualProduction/Avalanche/Source/AvalancheMaterial/Public/MaterialBridge/AvaPrimitiveComponentMaterialBridge.h b/Engine/Plugins/VirtualProduction/Avalanche/Source/AvalancheMaterial/Public/MaterialBridge/AvaPrimitiveComponentMaterialBridge.h
new file mode 100644
--- /dev/null
+++ b/Engine/Plugins/VirtualProduction/Avalanche/Source/AvalancheMaterial/Public/MaterialBridge/AvaPrimitiveComponentMaterialBridge.h
@@ -0,0 +1,24 @@
+// Copyright Epic Games, Inc. All Rights Reserved.
+
+#pragma once
+
+#include "AvaMaterialBridge.h"
+
+#define UE_API AVALANCHEMATERIAL_API
+
+namespace UE::Ava
+{
+
+/** Material Bridge for Primitive Components */
+class FPrimitiveComponentMaterialBridge : public FMaterialBridge
+{
+protected:
+	//~ Begin FMaterialBridge
+	UE_API virtual const UStruct* OnGetBridgedType() const override;
+	UE_API virtual EVisitFlow OnVisitMaterials(const FMaterialBridgeVisitContext& InContext, TFunctionRef<EVisitFlow(const FMaterialBridgeVisitContext&, UMaterialInterface*)> InFunc, const FVisitOptions& InOptions) const override;
+	//~ End FMaterialBridge
+};
+
+} // UE::Ava
+
+#undef UE_API
diff --git a/Engine/Plugins/VirtualProduction/Avalanche/Source/AvalancheMaterial/Public/MaterialBridge/AvaStaticMeshMaterialBridge.h b/Engine/Plugins/VirtualProduction/Avalanche/Source/AvalancheMaterial/Public/MaterialBridge/AvaStaticMeshMaterialBridge.h
new file mode 100644
--- /dev/null
+++ b/Engine/Plugins/VirtualProduction/Avalanche/Source/AvalancheMaterial/Public/MaterialBridge/AvaStaticMeshMaterialBridge.h
@@ -0,0 +1,27 @@
+// Copyright Epic Games, Inc. All Rights Reserved.
+
+#pragma once
+
+#include "AvaMaterialBridge.h"
+
+#define UE_API AVALANCHEMATERIAL_API
+
+namespace UE::Ava
+{
+
+/**
+ * Material Bridge for Static Meshes
+ * Only deals with the Static Mesh's static materials.
+ */
+class FStaticMeshMaterialBridge : public FMaterialBridge
+{
+protected:
+	//~ Begin FMaterialBridge
+	UE_API virtual const UStruct* OnGetBridgedType() const override;
+	UE_API virtual EVisitFlow OnVisitMaterials(const FMaterialBridgeVisitContext& InContext, TFunctionRef<EVisitFlow(const FMaterialBridgeVisitContext&, UMaterialInterface*)> InFunc, const FVisitOptions& InOptions) const override;
+	//~ End FMaterialBridge
+};
+
+} // UE::Ava
+
+#undef UE_API
diff --git a/Engine/Plugins/VirtualProduction/Avalanche/Source/AvalancheMedia/Private/Playable/Playables/AvaPlayableLevelStreaming.cpp b/Engine/Plugins/VirtualProduction/Avalanche/Source/AvalancheMedia/Private/Playable/Playables/AvaPlayableLevelStreaming.cpp
--- a/Engine/Plugins/VirtualProduction/Avalanche/Source/AvalancheMedia/Private/Playable/Playables/AvaPlayableLevelStreaming.cpp
+++ b/Engine/Plugins/VirtualProduction/Avalanche/Source/AvalancheMedia/Private/Playable/Playables/AvaPlayableLevelStreaming.cpp
@@ -1,8 +1,6 @@
 // Copyright Epic Games, Inc. All Rights Reserved.
 
 #include "Playable/Playables/AvaPlayableLevelStreaming.h"
-#include "AvaMaterialCacheHelper.h"
-#include "AvaMaterialCacheSettings.h"
 #include "AvaRemoteControlRebind.h"
 #include "AvaRemoteControlUtils.h"
 #include "AvaScene.h"
@@ -15,15 +13,17 @@
 #include "Framework/AvaSoftAssetPtr.h"
 #include "IAvaMediaModule.h"
 #include "IAvaRemoteControlInterface.h"
+#include "MaterialCache/AvaMaterialCacheHelper.h"
+#include "MaterialCache/AvaMaterialCacheSettings.h"
 #include "Playable/AvaPlayableAssetUserData.h"
 #include "Playable/AvaPlayableGroup.h"
 #include "Playable/AvaPlayableGroupManager.h"
 #include "Playable/AvaPlayableUtils.h"
 #include "Playback/AvaPlaybackUtils.h"
 #include "SceneView.h"
-#include "UnrealEngine.h"
 #include "Streaming/LevelStreamingDelegates.h"
 #include "UObject/Package.h"
+#include "UnrealEngine.h"
 
 #define LOCTEXT_NAMESPACE "AvaPlayableLevelStreaming"
 
diff --git a/Engine/Plugins/VirtualProduction/Avalanche/Source/AvalancheMedia/Private/Playback/AvaPlaybackUtils.cpp b/Engine/Plugins/VirtualProduction/Avalanche/Source/AvalancheMedia/Private/Playback/AvaPlaybackUtils.cpp
--- a/Engine/Plugins/VirtualProduction/Avalanche/Source/AvalancheMedia/Private/Playback/AvaPlaybackUtils.cpp
+++ b/Engine/Plugins/VirtualProduction/Avalanche/Source/AvalancheMedia/Private/Playback/AvaPlaybackUtils.cpp
@@ -5,13 +5,13 @@
 #include "AssetRegistry/AssetData.h"
 #include "AssetRegistry/IAssetRegistry.h"
 #include "AvaAssetTags.h"
-#include "AvaMaterialCacheHelper.h"
-#include "AvaMaterialCacheSettings.h"
 #include "ContentStreaming.h"
 #include "Engine/StreamableRenderAsset.h"
 #include "Framework/AvaSoftAssetPtr.h"
 #include "HAL/PlatformFileManager.h"
 #include "IAvaMediaModule.h"
+#include "MaterialCache/AvaMaterialCacheHelper.h"
+#include "MaterialCache/AvaMaterialCacheSettings.h"
 #include "Misc/Paths.h"
 #include "RenderingThread.h"
 #include "ShaderCompiler.h"
diff --git a/Engine/Plugins/VirtualProduction/Avalanche/Source/AvalancheMedia/Private/Rundown/AvaRundown.cpp b/Engine/Plugins/VirtualProduction/Avalanche/Source/AvalancheMedia/Private/Rundown/AvaRundown.cpp
--- a/Engine/Plugins/VirtualProduction/Avalanche/Source/AvalancheMedia/Private/Rundown/AvaRundown.cpp
+++ b/Engine/Plugins/VirtualProduction/Avalanche/Source/AvalancheMedia/Private/Rundown/AvaRundown.cpp
@@ -2,11 +2,11 @@
 
 #include "Rundown/AvaRundown.h"
 
-#include "AvaMaterialCacheHelper.h"
-#include "AvaMaterialCacheSettings.h"
 #include "AvaMediaSettings.h"
 #include "Broadcast/AvaBroadcast.h"
 #include "IAvaMediaModule.h"
+#include "MaterialCache/AvaMaterialCacheHelper.h"
+#include "MaterialCache/AvaMaterialCacheSettings.h"
 #include "Playable/AvaPlayable.h"
 #include "Playback/AvaPlaybackGraph.h"
 #include "Playback/AvaPlaybackManager.h"
diff --git a/Engine/Plugins/VirtualProduction/Avalanche/Source/AvalancheShapes/Private/AvaShapeComponentMaterialBridge.cpp b/Engine/Plugins/VirtualProduction/Avalanche/Source/AvalancheShapes/Private/AvaShapeComponentMaterialBridge.cpp
--- a/Engine/Plugins/VirtualProduction/Avalanche/Source/AvalancheShapes/Private/AvaShapeComponentMaterialBridge.cpp
+++ b/Engine/Plugins/VirtualProduction/Avalanche/Source/AvalancheShapes/Private/AvaShapeComponentMaterialBridge.cpp
@@ -11,14 +11,19 @@ const UStruct* FShapeComponentMaterialBridge::OnGetBridgedType() const
 	return UAvaShapeDynamicMeshBase::StaticClass();
 }
 
-void FShapeComponentMaterialBridge::OnVisitMaterials(const FMaterialBridgeVisitContext& InContext, TFunctionRef<void(const FMaterialBridgeVisitContext&, UMaterialInterface*)> InFunc, const FVisitOptions& InOptions) const
+EVisitFlow FShapeComponentMaterialBridge::OnVisitMaterials(const FMaterialBridgeVisitContext& InContext, TFunctionRef<EVisitFlow(const FMaterialBridgeVisitContext&, UMaterialInterface*)> InFunc, const FVisitOptions& InOptions) const
 {
 	const UAvaShapeDynamicMeshBase& Component = InContext.Container.Get<const UAvaShapeDynamicMeshBase>();
 
 	for (const TPair<int32, FAvaShapeMeshData>& Pair : Component.GetMeshDatas())
 	{
-		InFunc(InContext, Pair.Value.GetMaterial());
+		if (InFunc(InContext, Pair.Value.GetMaterial()) == EVisitFlow::Break)
+		{
+			return EVisitFlow::Break;
+		}
 	}
+
+	return EVisitFlow::Continue;
 }
 
 } // UE::Ava
diff --git a/Engine/Plugins/VirtualProduction/Avalanche/Source/AvalancheShapes/Private/AvaShapesModule.cpp b/Engine/Plugins/VirtualProduction/Avalanche/Source/AvalancheShapes/Private/AvaShapesModule.cpp
--- a/Engine/Plugins/VirtualProduction/Avalanche/Source/AvalancheShapes/Private/AvaShapesModule.cpp
+++ b/Engine/Plugins/VirtualProduction/Avalanche/Source/AvalancheShapes/Private/AvaShapesModule.cpp
@@ -1,8 +1,8 @@
 // Copyright Epic Games, Inc. All Rights Reserved.
 
 #include "AvaShapesModule.h"
-#include "AvaMaterialBridgeRegistry.h"
 #include "AvaShapeComponentMaterialBridge.h"
+#include "MaterialBridge/AvaMaterialBridgeRegistry.h"
 #include "Modules/ModuleManager.h"
 
 IMPLEMENT_MODULE(FAvaShapesModule, AvalancheShapes)
diff --git a/Engine/Plugins/VirtualProduction/Avalanche/Source/AvalancheShapes/Public/AvaShapeComponentMaterialBridge.h b/Engine/Plugins/VirtualProduction/Avalanche/Source/AvalancheShapes/Public/AvaShapeComponentMaterialBridge.h
--- a/Engine/Plugins/VirtualProduction/Avalanche/Source/AvalancheShapes/Public/AvaShapeComponentMaterialBridge.h
+++ b/Engine/Plugins/VirtualProduction/Avalanche/Source/AvalancheShapes/Public/AvaShapeComponentMaterialBridge.h
@@ -2,7 +2,7 @@
 
 #pragma once
 
-#include "AvaMaterialBridge.h"
+#include "MaterialBridge/AvaMaterialBridge.h"
 
 #define UE_API AVALANCHESHAPES_API
 
@@ -15,7 +15,7 @@ class FShapeComponentMaterialBridge : public FMaterialBridge
 protected:
 	//~ Begin FMaterialBridge
 	UE_API virtual const UStruct* OnGetBridgedType() const override;
-	UE_API virtual void OnVisitMaterials(const FMaterialBridgeVisitContext& InContext, TFunctionRef<void(const FMaterialBridgeVisitContext&, UMaterialInterface*)> InFunc, const FVisitOptions& InOptions) const override;
+	UE_API virtual EVisitFlow OnVisitMaterials(const FMaterialBridgeVisitContext& InContext, TFunctionRef<EVisitFlow(const FMaterialBridgeVisitContext&, UMaterialInterface*)> InFunc, const FVisitOptions& InOptions) const override;
 	//~ End FMaterialBridge
 };
 
diff --git a/Engine/Plugins/VirtualProduction/Avalanche/Source/AvalancheText/Private/AvaText3DComponentMaterialBridge.cpp b/Engine/Plugins/VirtualProduction/Avalanche/Source/AvalancheText/Private/AvaText3DComponentMaterialBridge.cpp
--- a/Engine/Plugins/VirtualProduction/Avalanche/Source/AvalancheText/Private/AvaText3DComponentMaterialBridge.cpp
+++ b/Engine/Plugins/VirtualProduction/Avalanche/Source/AvalancheText/Private/AvaText3DComponentMaterialBridge.cpp
@@ -13,19 +13,23 @@ const UStruct* FText3DComponentMaterialBridge::OnGetBridgedType() const
 	return UText3DComponent::StaticClass();
 }
 
-void FText3DComponentMaterialBridge::OnVisitMaterials(const FMaterialBridgeVisitContext& InContext, TFunctionRef<void(const FMaterialBridgeVisitContext&, UMaterialInterface*)> InFunc, const FVisitOptions& InOptions) const
+EVisitFlow FText3DComponentMaterialBridge::OnVisitMaterials(const FMaterialBridgeVisitContext& InContext, TFunctionRef<EVisitFlow(const FMaterialBridgeVisitContext&, UMaterialInterface*)> InFunc, const FVisitOptions& InOptions) const
 {
 	const UText3DComponent& Component = InContext.Container.Get<const UText3DComponent>();
 
+	EVisitFlow VisitFlow = EVisitFlow::Continue;
+
 	if (const UText3DMaterialExtensionBase* MaterialExtension = Component.GetMaterialExtension())
 	{
 		MaterialExtension->ForEachMaterial(
-			[&InFunc, &InContext](const UE::Text3D::Material::FMaterialParameters&, UMaterialInterface* InMaterial)->bool
+			[&InFunc, &InContext, &VisitFlow](const UE::Text3D::Material::FMaterialParameters&, UMaterialInterface* InMaterial)->bool
 			{
-				InFunc(InContext, InMaterial);
-				return true; // continue				
+				VisitFlow = InFunc(InContext, InMaterial);
+				return VisitFlow == EVisitFlow::Continue;
 			});
 	}
+
+	return VisitFlow;
 }
 
 } // UE::Ava
diff --git a/Engine/Plugins/VirtualProduction/Avalanche/Source/AvalancheText/Private/AvaTextModule.cpp b/Engine/Plugins/VirtualProduction/Avalanche/Source/AvalancheText/Private/AvaTextModule.cpp
--- a/Engine/Plugins/VirtualProduction/Avalanche/Source/AvalancheText/Private/AvaTextModule.cpp
+++ b/Engine/Plugins/VirtualProduction/Avalanche/Source/AvalancheText/Private/AvaTextModule.cpp
@@ -1,8 +1,8 @@
 // Copyright Epic Games, Inc. All Rights Reserved.
 
 #include "AvaTextModule.h"
-#include "AvaMaterialBridgeRegistry.h"
 #include "AvaText3DComponentMaterialBridge.h"
+#include "MaterialBridge/AvaMaterialBridgeRegistry.h"
 #include "Modules/ModuleManager.h"
 
 IMPLEMENT_MODULE(FAvaTextModule, AvalancheText)
diff --git a/Engine/Plugins/VirtualProduction/Avalanche/Source/AvalancheText/Public/AvaText3DComponentMaterialBridge.h b/Engine/Plugins/VirtualProduction/Avalanche/Source/AvalancheText/Public/AvaText3DComponentMaterialBridge.h
--- a/Engine/Plugins/VirtualProduction/Avalanche/Source/AvalancheText/Public/AvaText3DComponentMaterialBridge.h
+++ b/Engine/Plugins/VirtualProduction/Avalanche/Source/AvalancheText/Public/AvaText3DComponentMaterialBridge.h
@@ -2,7 +2,7 @@
 
 #pragma once
 
-#include "AvaMaterialBridge.h"
+#include "MaterialBridge/AvaMaterialBridge.h"
 
 #define UE_API AVALANCHETEXT_API
 
@@ -15,7 +15,7 @@ class FText3DComponentMaterialBridge : public FMaterialBridge
 protected:
 	//~ Begin FMaterialBridge
 	UE_API virtual const UStruct* OnGetBridgedType() const override;
-	UE_API virtual void OnVisitMaterials(const FMaterialBridgeVisitContext& InContext, TFunctionRef<void(const FMaterialBridgeVisitContext&, UMaterialInterface*)> InFunc, const FVisitOptions& InOptions) const override;
+	UE_API virtual EVisitFlow OnVisitMaterials(const FMaterialBridgeVisitContext& InContext, TFunctionRef<EVisitFlow(const FMaterialBridgeVisitContext&, UMaterialInterface*)> InFunc, const FVisitOptions& InOptions) const override;
 	//~ End FMaterialBridge
 };
 
