Motion Design: added priority to bridge and allowed multiple bridges to be registered for a given type. Bridges are picked by their type first, then by their priority order the first that returns 'true' on ::Supports will be the chosen bridge.

#Synced-CL 48481255

diff --git a/Engine/Plugins/VirtualProduction/Avalanche/Source/AvalancheMaterial/Private/AvaMaterialModule.cpp b/Engine/Plugins/VirtualProduction/Avalanche/Source/AvalancheMaterial/Private/AvaMaterialModule.cpp
--- a/Engine/Plugins/VirtualProduction/Avalanche/Source/AvalancheMaterial/Private/AvaMaterialModule.cpp
+++ b/Engine/Plugins/VirtualProduction/Avalanche/Source/AvalancheMaterial/Private/AvaMaterialModule.cpp
@@ -17,12 +17,14 @@ IMPLEMENT_MODULE(FAvaMaterialModule, AvalancheMaterial)
 
 void FAvaMaterialModule::StartupModule()
 {
-	UE::Ava::FMaterialBridgeRegistry& MaterialBridgeRegistry = UE::Ava::FMaterialBridgeRegistry::GetMutable();
+	// Low priority for built-in handlers
+	constexpr uint32 MaterialBridgePriority = 0;
 
 	// Register built-in material bridges
-	MaterialBridgeRegistry.Register<UE::Ava::FLevelMaterialBridge>();
-	MaterialBridgeRegistry.Register<UE::Ava::FActorMaterialBridge>();
-	MaterialBridgeRegistry.Register<UE::Ava::FPrimitiveComponentMaterialBridge>();
-	MaterialBridgeRegistry.Register<UE::Ava::FDecalComponentMaterialBridge>();
-	MaterialBridgeRegistry.Register<UE::Ava::FStaticMeshMaterialBridge>();
+	UE::Ava::FMaterialBridgeRegistry& MaterialBridgeRegistry = UE::Ava::FMaterialBridgeRegistry::GetMutable();
+	MaterialBridgeRegistry.Register<UE::Ava::FLevelMaterialBridge>(MaterialBridgePriority);
+	MaterialBridgeRegistry.Register<UE::Ava::FActorMaterialBridge>(MaterialBridgePriority);
+	MaterialBridgeRegistry.Register<UE::Ava::FPrimitiveComponentMaterialBridge>(MaterialBridgePriority);
+	MaterialBridgeRegistry.Register<UE::Ava::FDecalComponentMaterialBridge>(MaterialBridgePriority);
+	MaterialBridgeRegistry.Register<UE::Ava::FStaticMeshMaterialBridge>(MaterialBridgePriority);
 }
diff --git a/Engine/Plugins/VirtualProduction/Avalanche/Source/AvalancheMaterial/Private/MaterialBridge/AvaActorMaterialBridge.cpp b/Engine/Plugins/VirtualProduction/Avalanche/Source/AvalancheMaterial/Private/MaterialBridge/AvaActorMaterialBridge.cpp
--- a/Engine/Plugins/VirtualProduction/Avalanche/Source/AvalancheMaterial/Private/MaterialBridge/AvaActorMaterialBridge.cpp
+++ b/Engine/Plugins/VirtualProduction/Avalanche/Source/AvalancheMaterial/Private/MaterialBridge/AvaActorMaterialBridge.cpp
@@ -12,7 +12,7 @@ const UStruct* FActorMaterialBridge::OnGetBridgedType() const
 	return AActor::StaticClass();
 }
 
-EVisitFlow FActorMaterialBridge::OnVisitMaterials(const FMaterialBridgeVisitContext& InContext, TFunctionRef<EVisitFlow(const FMaterialBridgeVisitContext&, UMaterialInterface*)> InFunc, const FVisitOptions& InOptions) const
+EMaterialBridgeVisitFlow FActorMaterialBridge::OnVisitMaterials(const FVisitContext& InContext, TFunctionRef<EVisitFlow(const FVisitContext&, UMaterialInterface*)> InFunc, const FVisitOptions& InOptions) const
 {
 	const AActor& Actor = InContext.Container.Get<const AActor>();
 
@@ -27,9 +27,9 @@ EVisitFlow FActorMaterialBridge::OnVisitMaterials(const FMaterialBridgeVisitCont
 			{
 				return;
 			}
-			if (const FMaterialBridge* MaterialBridge = InOptions.MaterialBridgeRegistry->GetMaterialBridgeForObject(InComponent))
+			if (const FMaterialBridge* MaterialBridge = InOptions.MaterialBridgeRegistry->GetMaterialBridge(FConstDataView(InComponent)))
 			{
-				VisitFlow = MaterialBridge->VisitMaterials(FMaterialBridgeVisitContext(InComponent, &InContext), InFunc, InOptions);
+				VisitFlow = MaterialBridge->VisitMaterials(FVisitContext(InComponent, &InContext), InFunc, InOptions);
 			}
 		});
 
diff --git a/Engine/Plugins/VirtualProduction/Avalanche/Source/AvalancheMaterial/Private/MaterialBridge/AvaDecalComponentMaterialBridge.cpp b/Engine/Plugins/VirtualProduction/Avalanche/Source/AvalancheMaterial/Private/MaterialBridge/AvaDecalComponentMaterialBridge.cpp
--- a/Engine/Plugins/VirtualProduction/Avalanche/Source/AvalancheMaterial/Private/MaterialBridge/AvaDecalComponentMaterialBridge.cpp
+++ b/Engine/Plugins/VirtualProduction/Avalanche/Source/AvalancheMaterial/Private/MaterialBridge/AvaDecalComponentMaterialBridge.cpp
@@ -12,7 +12,7 @@ const UStruct* FDecalComponentMaterialBridge::OnGetBridgedType() const
 	return UDecalComponent::StaticClass();
 }
 
-EVisitFlow FDecalComponentMaterialBridge::OnVisitMaterials(const FMaterialBridgeVisitContext& InContext, TFunctionRef<EVisitFlow(const FMaterialBridgeVisitContext&, UMaterialInterface*)> InFunc, const FVisitOptions& InOptions) const
+EMaterialBridgeVisitFlow FDecalComponentMaterialBridge::OnVisitMaterials(const FVisitContext& InContext, TFunctionRef<EVisitFlow(const FVisitContext&, UMaterialInterface*)> InFunc, const FVisitOptions& InOptions) const
 {
 	const UDecalComponent& Component = InContext.Container.Get<const UDecalComponent>();
 	return InFunc(InContext, Component.GetDecalMaterial());
diff --git a/Engine/Plugins/VirtualProduction/Avalanche/Source/AvalancheMaterial/Private/MaterialBridge/AvaLevelMaterialBridge.cpp b/Engine/Plugins/VirtualProduction/Avalanche/Source/AvalancheMaterial/Private/MaterialBridge/AvaLevelMaterialBridge.cpp
--- a/Engine/Plugins/VirtualProduction/Avalanche/Source/AvalancheMaterial/Private/MaterialBridge/AvaLevelMaterialBridge.cpp
+++ b/Engine/Plugins/VirtualProduction/Avalanche/Source/AvalancheMaterial/Private/MaterialBridge/AvaLevelMaterialBridge.cpp
@@ -13,15 +13,15 @@ const UStruct* FLevelMaterialBridge::OnGetBridgedType() const
 	return ULevel::StaticClass();
 }
 
-EVisitFlow FLevelMaterialBridge::OnVisitMaterials(const FMaterialBridgeVisitContext& InContext, TFunctionRef<EVisitFlow(const FMaterialBridgeVisitContext&, UMaterialInterface*)> InFunc, const FVisitOptions& InOptions) const
+EMaterialBridgeVisitFlow FLevelMaterialBridge::OnVisitMaterials(const FVisitContext& InContext, TFunctionRef<EVisitFlow(const FVisitContext&, UMaterialInterface*)> InFunc, const FVisitOptions& InOptions) const
 {
 	const ULevel& Level = InContext.Container.Get<const ULevel>();
 
 	for (AActor* Actor : Level.Actors)
 	{
-		if (const FMaterialBridge* MaterialBridge = InOptions.MaterialBridgeRegistry->GetMaterialBridgeForObject(Actor))
+		if (const FMaterialBridge* MaterialBridge = InOptions.MaterialBridgeRegistry->GetMaterialBridge(FConstDataView(Actor)))
 		{
-			const EVisitFlow VisitFlow = MaterialBridge->VisitMaterials(FMaterialBridgeVisitContext(Actor, &InContext), InFunc, InOptions);
+			const EVisitFlow VisitFlow = MaterialBridge->VisitMaterials(FVisitContext(Actor, &InContext), InFunc, InOptions);
 			if (VisitFlow == EVisitFlow::Break)
 			{
 				return EVisitFlow::Break;
diff --git a/Engine/Plugins/VirtualProduction/Avalanche/Source/AvalancheMaterial/Private/MaterialBridge/AvaMaterialBridge.cpp b/Engine/Plugins/VirtualProduction/Avalanche/Source/AvalancheMaterial/Private/MaterialBridge/AvaMaterialBridge.cpp
--- a/Engine/Plugins/VirtualProduction/Avalanche/Source/AvalancheMaterial/Private/MaterialBridge/AvaMaterialBridge.cpp
+++ b/Engine/Plugins/VirtualProduction/Avalanche/Source/AvalancheMaterial/Private/MaterialBridge/AvaMaterialBridge.cpp
@@ -5,14 +5,10 @@
 
 namespace UE::Ava
 {
-FMaterialBridge::FVisitOptions::FVisitOptions()
-	: MaterialBridgeRegistry(&FMaterialBridgeRegistry::Get())
-{
-}
-
-void FMaterialBridge::Initialize()
+void FMaterialBridge::Initialize(uint32 InPriority)
 {
 	BridgedType = OnGetBridgedType();
+	Priority = InPriority;
 }
 
 const UStruct* FMaterialBridge::GetBridgedType() const
@@ -20,7 +16,21 @@ const UStruct* FMaterialBridge::GetBridgedType() const
 	return BridgedType;
 }
 
-EVisitFlow FMaterialBridge::VisitMaterials(const FMaterialBridgeVisitContext& InContext, TFunctionRef<EVisitFlow(const FMaterialBridgeVisitContext&, UMaterialInterface*)> InFunc, const FVisitOptions& InOptions) const
+uint32 FMaterialBridge::GetPriority() const
+{
+	return Priority;
+}
+
+bool FMaterialBridge::Supports(FConstDataView InContainer) const
+{
+	if (InContainer.IsValidFor(BridgedType))
+	{
+		return OnSupports(InContainer);
+	}
+	return false;
+}
+
+EMaterialBridgeVisitFlow FMaterialBridge::VisitMaterials(const FVisitContext& InContext, TFunctionRef<EVisitFlow(const FVisitContext&, UMaterialInterface*)> InFunc, const FVisitOptions& InOptions) const
 {
 	if (InContext.Container.IsValidFor(BridgedType))
 	{
diff --git a/Engine/Plugins/VirtualProduction/Avalanche/Source/AvalancheMaterial/Private/MaterialBridge/AvaMaterialBridgeRegistry.cpp b/Engine/Plugins/VirtualProduction/Avalanche/Source/AvalancheMaterial/Private/MaterialBridge/AvaMaterialBridgeRegistry.cpp
--- a/Engine/Plugins/VirtualProduction/Avalanche/Source/AvalancheMaterial/Private/MaterialBridge/AvaMaterialBridgeRegistry.cpp
+++ b/Engine/Plugins/VirtualProduction/Avalanche/Source/AvalancheMaterial/Private/MaterialBridge/AvaMaterialBridgeRegistry.cpp
@@ -17,46 +17,54 @@ FMaterialBridgeRegistry& FMaterialBridgeRegistry::GetMutable()
 	return MaterialBridgeRegistry;
 }
 
-const FMaterialBridge* FMaterialBridgeRegistry::GetMaterialBridgeForObject(const UObject* InObject) const
+const FMaterialBridge* FMaterialBridgeRegistry::GetMaterialBridge(FConstDataView InData) const
 {
-	if (InObject)
-	{
-		return GetMaterialBridgeForType(InObject->GetClass());
-	}
-	return nullptr;
-}
-
-const FMaterialBridge* FMaterialBridgeRegistry::GetMaterialBridgeForType(const UStruct* InStruct) const
-{
-	if (!InStruct)
+	if (!InData.IsValid())
 	{
 		return nullptr;
 	}
-	for (const UStruct* Struct : InStruct->GetSuperStructIterator())
+	for (const UStruct* Struct : InData.GetStruct()->GetSuperStructIterator())
 	{
-		if (const TUniquePtr<FMaterialBridge>* MaterialBridge = RegisteredMaterialBridges.Find(Struct))
+		const TArray<TUniquePtr<FMaterialBridge>>* MaterialBridges = RegisteredMaterialBridges.Find(Struct);
+		if (!MaterialBridges)
+		{
+			continue;
+		}
+		for (const TUniquePtr<FMaterialBridge>& MaterialBridge : *MaterialBridges)
 		{
-			return MaterialBridge->Get();
+			if (MaterialBridge.IsValid() && MaterialBridge->Supports(InData))
+			{
+				return MaterialBridge.Get();
+			}
 		}
 	}
 	return nullptr;
 }
 
-
-void FMaterialBridgeRegistry::RegisterInternal(TUniquePtr<UE::Ava::FMaterialBridge>&& InMaterialBridge)
+void FMaterialBridgeRegistry::RegisterInternal(TUniquePtr<UE::Ava::FMaterialBridge>&& InMaterialBridge, uint32 InPriority)
 {
 	if (!ensureMsgf(UObjectInitialized(), TEXT("Cannot register Material Bridges. UObject system not initialized!")))
 	{
 		return;
 	}
 
-	InMaterialBridge->Initialize();
-	const UStruct* const BridgedType = InMaterialBridge->GetBridgedType();
+	InMaterialBridge->Initialize(InPriority);
 
-	if (ensureMsgf(BridgedType, TEXT("Cannot register Material Bridge with a null bridged type!")))
+	const UStruct* const BridgedType = InMaterialBridge->GetBridgedType();
+	if (!ensureMsgf(BridgedType, TEXT("Cannot register Material Bridge with a null bridged type!")))
 	{
-		RegisteredMaterialBridges.Add(BridgedType, MoveTemp(InMaterialBridge));
+		return;
 	}
+
+	TArray<TUniquePtr<FMaterialBridge>>& MaterialBridges = RegisteredMaterialBridges.FindOrAdd(BridgedType);
+	MaterialBridges.Add(MoveTemp(InMaterialBridge));
+
+	// Sort the bridges from highest to lowest priority
+	MaterialBridges.StableSort(
+		[](const TUniquePtr<FMaterialBridge>& A, const TUniquePtr<FMaterialBridge>& B)
+		{
+			return A->GetPriority() > B->GetPriority();
+		});
 }
 
 } // UE::Ava
diff --git a/Engine/Plugins/VirtualProduction/Avalanche/Source/AvalancheMaterial/Private/MaterialBridge/AvaMaterialBridgeVisitContext.cpp b/Engine/Plugins/VirtualProduction/Avalanche/Source/AvalancheMaterial/Private/MaterialBridge/AvaMaterialBridgeVisitContext.cpp
--- a/Engine/Plugins/VirtualProduction/Avalanche/Source/AvalancheMaterial/Private/MaterialBridge/AvaMaterialBridgeVisitContext.cpp
+++ b/Engine/Plugins/VirtualProduction/Avalanche/Source/AvalancheMaterial/Private/MaterialBridge/AvaMaterialBridgeVisitContext.cpp
@@ -1,6 +1,7 @@
 // Copyright Epic Games, Inc. All Rights Reserved.
 
 #include "MaterialBridge/AvaMaterialBridgeVisitContext.h"
+#include "MaterialBridge/AvaMaterialBridgeRegistry.h"
 
 namespace UE::Ava
 {
@@ -16,4 +17,9 @@ FMaterialBridgeVisitContext::FMaterialBridgeVisitContext(UObject* InContainer, T
 {
 }
 
+FMaterialBridgeVisitOptions::FMaterialBridgeVisitOptions()
+	: MaterialBridgeRegistry(&FMaterialBridgeRegistry::Get())
+{
+}
+
 } // UE::Ava
diff --git a/Engine/Plugins/VirtualProduction/Avalanche/Source/AvalancheMaterial/Private/MaterialBridge/AvaPrimitiveComponentMaterialBridge.cpp b/Engine/Plugins/VirtualProduction/Avalanche/Source/AvalancheMaterial/Private/MaterialBridge/AvaPrimitiveComponentMaterialBridge.cpp
--- a/Engine/Plugins/VirtualProduction/Avalanche/Source/AvalancheMaterial/Private/MaterialBridge/AvaPrimitiveComponentMaterialBridge.cpp
+++ b/Engine/Plugins/VirtualProduction/Avalanche/Source/AvalancheMaterial/Private/MaterialBridge/AvaPrimitiveComponentMaterialBridge.cpp
@@ -12,7 +12,7 @@ const UStruct* FPrimitiveComponentMaterialBridge::OnGetBridgedType() const
 	return UPrimitiveComponent::StaticClass();
 }
 
-EVisitFlow FPrimitiveComponentMaterialBridge::OnVisitMaterials(const FMaterialBridgeVisitContext& InContext, TFunctionRef<EVisitFlow(const FMaterialBridgeVisitContext&, UMaterialInterface*)> InFunc, const FVisitOptions& InOptions) const
+EMaterialBridgeVisitFlow FPrimitiveComponentMaterialBridge::OnVisitMaterials(const FVisitContext& InContext, TFunctionRef<EVisitFlow(const FVisitContext&, UMaterialInterface*)> InFunc, const FVisitOptions& InOptions) const
 {
 	const UPrimitiveComponent& Component = InContext.Container.Get<const UPrimitiveComponent>();
 
diff --git a/Engine/Plugins/VirtualProduction/Avalanche/Source/AvalancheMaterial/Private/MaterialBridge/AvaStaticMeshMaterialBridge.cpp b/Engine/Plugins/VirtualProduction/Avalanche/Source/AvalancheMaterial/Private/MaterialBridge/AvaStaticMeshMaterialBridge.cpp
--- a/Engine/Plugins/VirtualProduction/Avalanche/Source/AvalancheMaterial/Private/MaterialBridge/AvaStaticMeshMaterialBridge.cpp
+++ b/Engine/Plugins/VirtualProduction/Avalanche/Source/AvalancheMaterial/Private/MaterialBridge/AvaStaticMeshMaterialBridge.cpp
@@ -12,7 +12,7 @@ const UStruct* FStaticMeshMaterialBridge::OnGetBridgedType() const
 	return UStaticMesh::StaticClass();
 }
 
-EVisitFlow FStaticMeshMaterialBridge::OnVisitMaterials(const FMaterialBridgeVisitContext& InContext, TFunctionRef<EVisitFlow(const FMaterialBridgeVisitContext&, UMaterialInterface*)> InFunc, const FVisitOptions& InOptions) const
+EMaterialBridgeVisitFlow FStaticMeshMaterialBridge::OnVisitMaterials(const FVisitContext& InContext, TFunctionRef<EVisitFlow(const FVisitContext&, UMaterialInterface*)> InFunc, const FVisitOptions& InOptions) const
 {
 	const UStaticMesh& StaticMesh = InContext.Container.Get<const UStaticMesh>();
 
diff --git a/Engine/Plugins/VirtualProduction/Avalanche/Source/AvalancheMaterial/Private/MaterialCache/AvaMaterialCacheHelper.cpp b/Engine/Plugins/VirtualProduction/Avalanche/Source/AvalancheMaterial/Private/MaterialCache/AvaMaterialCacheHelper.cpp
--- a/Engine/Plugins/VirtualProduction/Avalanche/Source/AvalancheMaterial/Private/MaterialCache/AvaMaterialCacheHelper.cpp
+++ b/Engine/Plugins/VirtualProduction/Avalanche/Source/AvalancheMaterial/Private/MaterialCache/AvaMaterialCacheHelper.cpp
@@ -527,7 +527,7 @@ bool FMaterialCacheHelper::GatherMaterialData(const UObject* InObject, FObjectDa
 
 	TRACE_CPUPROFILER_EVENT_SCOPE(FMaterialCacheHelper::GatherMaterialData);
 
-	const FMaterialBridge* MaterialBridge = FMaterialBridgeRegistry::Get().GetMaterialBridgeForObject(InObject);
+	const FMaterialBridge* MaterialBridge = FMaterialBridgeRegistry::Get().GetMaterialBridge(FConstDataView(InObject));
 	if (!MaterialBridge)
 	{
 		UE_LOG(LogAvaMaterialCache, VeryVerbose, TEXT("Object '%s' (Class: %s) does not have a valid bridge to cache materials!")
@@ -551,9 +551,9 @@ bool FMaterialCacheHelper::GatherMaterialData(const UObject* InObject, FObjectDa
 		};
 
 	// Options to not wait for the container to have materials available and instead get a callback
-	FMaterialBridge::FVisitOptions VisitOptions;
+	FMaterialBridgeVisitOptions VisitOptions;
 	VisitOptions.bTryWaitOnContainerCompletion = false;
-	VisitOptions.OnContainerPendingCompletion.BindLambda([&InObjectData](const FMaterialBridgeVisitContext& InContext)->EVisitFlow
+	VisitOptions.OnContainerPendingCompletion.BindLambda([&InObjectData](const FMaterialBridgeVisitContext& InContext)->EMaterialBridgeVisitFlow
 		{
 			if (const UObject* Container = InContext.Container.GetPtr<const UObject>())
 			{
@@ -563,16 +563,16 @@ bool FMaterialCacheHelper::GatherMaterialData(const UObject* InObject, FObjectDa
 			{
 				ensureMsgf(0, TEXT("Container of type '%s' not supported for deferring material gather"), *GetNameSafe(InContext.Container.GetStruct()));
 			}
-			return EVisitFlow::Continue;
+			return EMaterialBridgeVisitFlow::Continue;
 		});
 
 	// Gather all materials for object
 	MaterialBridge->VisitMaterials(FMaterialBridgeVisitContext(const_cast<UObject*>(InObject)),
-		[&AddIncompleteMaterial](const FMaterialBridgeVisitContext& InContext, UMaterialInterface* InMaterial)->EVisitFlow
+		[&AddIncompleteMaterial](const FMaterialBridgeVisitContext& InContext, UMaterialInterface* InMaterial)->EMaterialBridgeVisitFlow
 		{
 			if (!InMaterial)
 			{
-				return EVisitFlow::Continue;
+				return EMaterialBridgeVisitFlow::Continue;
 			}
 
 			UMaterial* const Material = InMaterial->GetMaterial();
@@ -586,7 +586,7 @@ bool FMaterialCacheHelper::GatherMaterialData(const UObject* InObject, FObjectDa
 			{
 				AddIncompleteMaterial(InMaterial);
 			}
-			return EVisitFlow::Continue;
+			return EMaterialBridgeVisitFlow::Continue;
 		}
 		, VisitOptions);
 
diff --git a/Engine/Plugins/VirtualProduction/Avalanche/Source/AvalancheMaterial/Public/MaterialBridge/AvaActorMaterialBridge.h b/Engine/Plugins/VirtualProduction/Avalanche/Source/AvalancheMaterial/Public/MaterialBridge/AvaActorMaterialBridge.h
--- a/Engine/Plugins/VirtualProduction/Avalanche/Source/AvalancheMaterial/Public/MaterialBridge/AvaActorMaterialBridge.h
+++ b/Engine/Plugins/VirtualProduction/Avalanche/Source/AvalancheMaterial/Public/MaterialBridge/AvaActorMaterialBridge.h
@@ -15,7 +15,7 @@ class FActorMaterialBridge : public FMaterialBridge
 protected:
 	//~ Begin FMaterialBridge
 	UE_API virtual const UStruct* OnGetBridgedType() const override;
-	UE_API virtual EVisitFlow OnVisitMaterials(const FMaterialBridgeVisitContext& InContext, TFunctionRef<EVisitFlow(const FMaterialBridgeVisitContext&, UMaterialInterface*)> InFunc, const FVisitOptions& InOptions) const override;
+	UE_API virtual EVisitFlow OnVisitMaterials(const FVisitContext& InContext, TFunctionRef<EVisitFlow(const FVisitContext&, UMaterialInterface*)> InFunc, const FVisitOptions& InOptions) const override;
 	//~ End FMaterialBridge
 };
 
diff --git a/Engine/Plugins/VirtualProduction/Avalanche/Source/AvalancheMaterial/Public/MaterialBridge/AvaDecalComponentMaterialBridge.h b/Engine/Plugins/VirtualProduction/Avalanche/Source/AvalancheMaterial/Public/MaterialBridge/AvaDecalComponentMaterialBridge.h
--- a/Engine/Plugins/VirtualProduction/Avalanche/Source/AvalancheMaterial/Public/MaterialBridge/AvaDecalComponentMaterialBridge.h
+++ b/Engine/Plugins/VirtualProduction/Avalanche/Source/AvalancheMaterial/Public/MaterialBridge/AvaDecalComponentMaterialBridge.h
@@ -15,7 +15,7 @@ class FDecalComponentMaterialBridge : public FMaterialBridge
 protected:
 	//~ Begin FMaterialBridge
 	UE_API virtual const UStruct* OnGetBridgedType() const override;
-	UE_API virtual EVisitFlow OnVisitMaterials(const FMaterialBridgeVisitContext& InContext, TFunctionRef<EVisitFlow(const FMaterialBridgeVisitContext&, UMaterialInterface*)> InFunc, const FVisitOptions& InOptions) const override;
+	UE_API virtual EVisitFlow OnVisitMaterials(const FVisitContext& InContext, TFunctionRef<EVisitFlow(const FVisitContext&, UMaterialInterface*)> InFunc, const FVisitOptions& InOptions) const override;
 	//~ End FMaterialBridge
 };
 
diff --git a/Engine/Plugins/VirtualProduction/Avalanche/Source/AvalancheMaterial/Public/MaterialBridge/AvaLevelMaterialBridge.h b/Engine/Plugins/VirtualProduction/Avalanche/Source/AvalancheMaterial/Public/MaterialBridge/AvaLevelMaterialBridge.h
--- a/Engine/Plugins/VirtualProduction/Avalanche/Source/AvalancheMaterial/Public/MaterialBridge/AvaLevelMaterialBridge.h
+++ b/Engine/Plugins/VirtualProduction/Avalanche/Source/AvalancheMaterial/Public/MaterialBridge/AvaLevelMaterialBridge.h
@@ -15,7 +15,7 @@ class FLevelMaterialBridge : public FMaterialBridge
 protected:
 	//~ Begin FMaterialBridge
 	UE_API virtual const UStruct* OnGetBridgedType() const override;
-	UE_API virtual EVisitFlow OnVisitMaterials(const FMaterialBridgeVisitContext& InContext, TFunctionRef<EVisitFlow(const FMaterialBridgeVisitContext&, UMaterialInterface*)> InFunc, const FVisitOptions& InOptions) const override;
+	UE_API virtual EVisitFlow OnVisitMaterials(const FVisitContext& InContext, TFunctionRef<EVisitFlow(const FVisitContext&, UMaterialInterface*)> InFunc, const FVisitOptions& InOptions) const override;
 	//~ End FMaterialBridge
 };
 
diff --git a/Engine/Plugins/VirtualProduction/Avalanche/Source/AvalancheMaterial/Public/MaterialBridge/AvaMaterialBridge.h b/Engine/Plugins/VirtualProduction/Avalanche/Source/AvalancheMaterial/Public/MaterialBridge/AvaMaterialBridge.h
--- a/Engine/Plugins/VirtualProduction/Avalanche/Source/AvalancheMaterial/Public/MaterialBridge/AvaMaterialBridge.h
+++ b/Engine/Plugins/VirtualProduction/Avalanche/Source/AvalancheMaterial/Public/MaterialBridge/AvaMaterialBridge.h
@@ -13,48 +13,33 @@ class UMaterialInterface;
 
 namespace UE::Ava
 {
-	class FMaterialBridgeRegistry;
 
-} // UE::Ava
-
-namespace UE::Ava
-{
-
-/** Indicates how iteration should be handled */
-enum class EVisitFlow : uint8
-{
-	/** Continue iteration */
-	Continue,
-	/** Break iteration */
-	Break,
-};
-
-/**
- * Base class that allows interacting with materials from any type of data view (Actors, ActorComponents, etc).
- */
+/** Base class that allows interacting with materials from any type of data view (Actors, ActorComponents, etc). */
 class FMaterialBridge
 {
 public:
+	using EVisitFlow = EMaterialBridgeVisitFlow;
+	using FVisitOptions = FMaterialBridgeVisitOptions;
+	using FVisitContext = FMaterialBridgeVisitContext;
+
 	virtual ~FMaterialBridge() = default;
 
 	/** Initializes the bridge for use */
-	void Initialize();
+	void Initialize(uint32 InPriority);
 
 	/** Retrieves the type that this bridge supports */
 	UE_API const UStruct* GetBridgedType() const;
 
-	struct FVisitOptions
-	{
-		UE_API FVisitOptions();
-		/** The registry to use to visit any materials of sub objects. Defaults to the global registry. */
-		TNotNull<const FMaterialBridgeRegistry*> MaterialBridgeRegistry;
-		/** Callback when a container is not yet complete and cannot access its materials. Only called when bTryWaitOnContainerCompletion is false. */
-		TDelegate<EVisitFlow(const FMaterialBridgeVisitContext&)> OnContainerPendingCompletion;
-		/** Whether to try to wait for container completion to access its materials */
-		bool bTryWaitOnContainerCompletion = true;
-		/** Whether to include debug materials */
-		bool bVisitDebugMaterials = false;
-	};
+	/** Get the priority of this bridge in relation to other bridges of the same bridged type */
+	uint32 GetPriority() const;
+
+	/**
+	 * Whether this Material Bridge supports the given Material Container.
+	 * @param InContainer a valid material container of the supported type of this material bridge.
+	 * @return true if this material bridge supports the container, false otherwise.
+	 */
+	UE_API bool Supports(FConstDataView InContainer) const;
+
 	/**
 	 * Visits all the used materials in the given context.
 	 * @param InContext the current context to find materials from.
@@ -62,18 +47,27 @@ public:
 	 * @param InOptions visit options. see FVisitOptions
 	 * @return the iteration result of VisitMaterials.
 	 */
-	UE_API EVisitFlow VisitMaterials(const FMaterialBridgeVisitContext& InContext, TFunctionRef<EVisitFlow(const FMaterialBridgeVisitContext&, UMaterialInterface*)> InFunc, const FVisitOptions& InOptions = {}) const;
+	UE_API EVisitFlow VisitMaterials(const FVisitContext& InContext, TFunctionRef<EVisitFlow(const FVisitContext&, UMaterialInterface*)> InFunc, const FVisitOptions& InOptions = {}) const;
 
 protected:
 	/** Retrieves the type that this bridge supports. Only called once in Initialize() */
 	virtual const UStruct* OnGetBridgedType() const = 0;
 
+	/**  Whether this material bridge supports the given material container.*/
+	virtual bool OnSupports(FConstDataView InContainer) const
+	{
+		return true;
+	}
+
 	/** Get all the used materials in the given context. */
-	virtual EVisitFlow OnVisitMaterials(const FMaterialBridgeVisitContext& InContext, TFunctionRef<EVisitFlow(const FMaterialBridgeVisitContext&, UMaterialInterface*)> InFunc, const FVisitOptions& InOptions) const = 0;
+	virtual EVisitFlow OnVisitMaterials(const FVisitContext& InContext, TFunctionRef<EVisitFlow(const FVisitContext&, UMaterialInterface*)> InFunc, const FVisitOptions& InOptions) const = 0;
 
 private:
 	/** The type that this Material Bridge works on */
 	TObjectPtr<const UStruct> BridgedType;
+
+	/** The priority of this bridge in relation to other bridges of the same bridged type */
+	uint32 Priority = 0;
 };
 
 } // UE::Ava
diff --git a/Engine/Plugins/VirtualProduction/Avalanche/Source/AvalancheMaterial/Public/MaterialBridge/AvaMaterialBridgeRegistry.h b/Engine/Plugins/VirtualProduction/Avalanche/Source/AvalancheMaterial/Public/MaterialBridge/AvaMaterialBridgeRegistry.h
--- a/Engine/Plugins/VirtualProduction/Avalanche/Source/AvalancheMaterial/Public/MaterialBridge/AvaMaterialBridgeRegistry.h
+++ b/Engine/Plugins/VirtualProduction/Avalanche/Source/AvalancheMaterial/Public/MaterialBridge/AvaMaterialBridgeRegistry.h
@@ -12,6 +12,7 @@
 namespace UE::Ava
 {
 	class FMaterialBridge;
+	struct FConstDataView;
 }
 
 namespace UE::Ava
@@ -25,23 +26,21 @@ public:
 	UE_API static FMaterialBridgeRegistry& GetMutable();
 
 	template<typename InMaterialBridgeType, typename... InArgTypes> requires (UE::CDerivedFrom<InMaterialBridgeType, UE::Ava::FMaterialBridge>)
-	void Register(InArgTypes&&... InArgs)
+	void Register(uint32 InPriority, InArgTypes&&... InArgs)
 	{
-		this->RegisterInternal(MakeUnique<InMaterialBridgeType>(InArgs...));
+		this->RegisterInternal(MakeUnique<InMaterialBridgeType>(InArgs...), InPriority);
 	}
 
 	/** Tries to find the most relevant material bridge for a given object, returning a null if object is invalid, or if a relevant material bridge was not found. */
-	UE_API const FMaterialBridge* GetMaterialBridgeForObject(const UObject* InObject) const;
-
-	/** Tries to find the most relevant material bridge for a given type, returning null if a relevant material bridge was not found. */
-	UE_API const FMaterialBridge* GetMaterialBridgeForType(const UStruct* InType) const;
+	UE_API const FMaterialBridge* GetMaterialBridge(FConstDataView InData) const;
 
 protected:
-	UE_API void RegisterInternal(TUniquePtr<UE::Ava::FMaterialBridge>&& InMaterialBridge);
+	/** Initializes the material bridge and adds it to this registry */
+	UE_API void RegisterInternal(TUniquePtr<UE::Ava::FMaterialBridge>&& InMaterialBridge, uint32 InPriority);
 
 private:
 	/** Registered material bridges by their supported type */
-	TMap<TObjectKey<UStruct>, TUniquePtr<FMaterialBridge>> RegisteredMaterialBridges;
+	TMap<TObjectKey<UStruct>, TArray<TUniquePtr<FMaterialBridge>>> RegisteredMaterialBridges;
 };
 
 } // UE::Ava
diff --git a/Engine/Plugins/VirtualProduction/Avalanche/Source/AvalancheMaterial/Public/MaterialBridge/AvaMaterialBridgeVisitContext.h b/Engine/Plugins/VirtualProduction/Avalanche/Source/AvalancheMaterial/Public/MaterialBridge/AvaMaterialBridgeVisitContext.h
--- a/Engine/Plugins/VirtualProduction/Avalanche/Source/AvalancheMaterial/Public/MaterialBridge/AvaMaterialBridgeVisitContext.h
+++ b/Engine/Plugins/VirtualProduction/Avalanche/Source/AvalancheMaterial/Public/MaterialBridge/AvaMaterialBridgeVisitContext.h
@@ -8,6 +8,21 @@
 
 namespace UE::Ava
 {
+	class FMaterialBridgeRegistry;
+
+} // UE::Ava
+
+namespace UE::Ava
+{
+
+/** Indicates how iteration should be handled */
+enum class EMaterialBridgeVisitFlow : uint8
+{
+	/** Continue iteration */
+	Continue,
+	/** Break iteration */
+	Break,
+};
 
 /** Short-lived context information of a visited material */
 class FMaterialBridgeVisitContext
@@ -25,6 +40,24 @@ public:
 	const FMaterialBridgeVisitContext* ParentContext = nullptr;
 };
 
+/** Options to how visit should take place */
+struct FMaterialBridgeVisitOptions
+{
+	UE_API FMaterialBridgeVisitOptions();
+
+	/** The registry to use to visit any materials of sub objects. Defaults to the global registry. */
+	TNotNull<const FMaterialBridgeRegistry*> MaterialBridgeRegistry;
+
+	/** Callback when a container is not yet complete and cannot access its materials. Only called when bTryWaitOnContainerCompletion is false. */
+	TDelegate<EMaterialBridgeVisitFlow(const FMaterialBridgeVisitContext&)> OnContainerPendingCompletion;
+
+	/** Whether to try to wait for container completion to access its materials */
+	bool bTryWaitOnContainerCompletion = true;
+
+	/** Whether to include debug materials */
+	bool bVisitDebugMaterials = false;
+};
+
 } // UE::Ava
 
 #undef UE_API
diff --git a/Engine/Plugins/VirtualProduction/Avalanche/Source/AvalancheMaterial/Public/MaterialBridge/AvaPrimitiveComponentMaterialBridge.h b/Engine/Plugins/VirtualProduction/Avalanche/Source/AvalancheMaterial/Public/MaterialBridge/AvaPrimitiveComponentMaterialBridge.h
--- a/Engine/Plugins/VirtualProduction/Avalanche/Source/AvalancheMaterial/Public/MaterialBridge/AvaPrimitiveComponentMaterialBridge.h
+++ b/Engine/Plugins/VirtualProduction/Avalanche/Source/AvalancheMaterial/Public/MaterialBridge/AvaPrimitiveComponentMaterialBridge.h
@@ -15,7 +15,7 @@ class FPrimitiveComponentMaterialBridge : public FMaterialBridge
 protected:
 	//~ Begin FMaterialBridge
 	UE_API virtual const UStruct* OnGetBridgedType() const override;
-	UE_API virtual EVisitFlow OnVisitMaterials(const FMaterialBridgeVisitContext& InContext, TFunctionRef<EVisitFlow(const FMaterialBridgeVisitContext&, UMaterialInterface*)> InFunc, const FVisitOptions& InOptions) const override;
+	UE_API virtual EVisitFlow OnVisitMaterials(const FVisitContext& InContext, TFunctionRef<EVisitFlow(const FVisitContext&, UMaterialInterface*)> InFunc, const FVisitOptions& InOptions) const override;
 	//~ End FMaterialBridge
 };
 
diff --git a/Engine/Plugins/VirtualProduction/Avalanche/Source/AvalancheMaterial/Public/MaterialBridge/AvaStaticMeshMaterialBridge.h b/Engine/Plugins/VirtualProduction/Avalanche/Source/AvalancheMaterial/Public/MaterialBridge/AvaStaticMeshMaterialBridge.h
--- a/Engine/Plugins/VirtualProduction/Avalanche/Source/AvalancheMaterial/Public/MaterialBridge/AvaStaticMeshMaterialBridge.h
+++ b/Engine/Plugins/VirtualProduction/Avalanche/Source/AvalancheMaterial/Public/MaterialBridge/AvaStaticMeshMaterialBridge.h
@@ -18,7 +18,7 @@ class FStaticMeshMaterialBridge : public FMaterialBridge
 protected:
 	//~ Begin FMaterialBridge
 	UE_API virtual const UStruct* OnGetBridgedType() const override;
-	UE_API virtual EVisitFlow OnVisitMaterials(const FMaterialBridgeVisitContext& InContext, TFunctionRef<EVisitFlow(const FMaterialBridgeVisitContext&, UMaterialInterface*)> InFunc, const FVisitOptions& InOptions) const override;
+	UE_API virtual EVisitFlow OnVisitMaterials(const FVisitContext& InContext, TFunctionRef<EVisitFlow(const FVisitContext&, UMaterialInterface*)> InFunc, const FVisitOptions& InOptions) const override;
 	//~ End FMaterialBridge
 };
 
diff --git a/Engine/Plugins/VirtualProduction/Avalanche/Source/AvalancheShapes/Private/AvaShapeComponentMaterialBridge.cpp b/Engine/Plugins/VirtualProduction/Avalanche/Source/AvalancheShapes/Private/AvaShapeComponentMaterialBridge.cpp
--- a/Engine/Plugins/VirtualProduction/Avalanche/Source/AvalancheShapes/Private/AvaShapeComponentMaterialBridge.cpp
+++ b/Engine/Plugins/VirtualProduction/Avalanche/Source/AvalancheShapes/Private/AvaShapeComponentMaterialBridge.cpp
@@ -11,7 +11,7 @@ const UStruct* FShapeComponentMaterialBridge::OnGetBridgedType() const
 	return UAvaShapeDynamicMeshBase::StaticClass();
 }
 
-EVisitFlow FShapeComponentMaterialBridge::OnVisitMaterials(const FMaterialBridgeVisitContext& InContext, TFunctionRef<EVisitFlow(const FMaterialBridgeVisitContext&, UMaterialInterface*)> InFunc, const FVisitOptions& InOptions) const
+EMaterialBridgeVisitFlow FShapeComponentMaterialBridge::OnVisitMaterials(const FVisitContext& InContext, TFunctionRef<EVisitFlow(const FVisitContext&, UMaterialInterface*)> InFunc, const FVisitOptions& InOptions) const
 {
 	const UAvaShapeDynamicMeshBase& Component = InContext.Container.Get<const UAvaShapeDynamicMeshBase>();
 
diff --git a/Engine/Plugins/VirtualProduction/Avalanche/Source/AvalancheShapes/Private/AvaShapesModule.cpp b/Engine/Plugins/VirtualProduction/Avalanche/Source/AvalancheShapes/Private/AvaShapesModule.cpp
--- a/Engine/Plugins/VirtualProduction/Avalanche/Source/AvalancheShapes/Private/AvaShapesModule.cpp
+++ b/Engine/Plugins/VirtualProduction/Avalanche/Source/AvalancheShapes/Private/AvaShapesModule.cpp
@@ -9,6 +9,8 @@ IMPLEMENT_MODULE(FAvaShapesModule, AvalancheShapes)
 
 void FAvaShapesModule::StartupModule()
 {
+	constexpr uint32 MaterialBridgePriority = 0;
+
 	UE::Ava::FMaterialBridgeRegistry& MaterialBridgeRegistry = UE::Ava::FMaterialBridgeRegistry::GetMutable();
-	MaterialBridgeRegistry.Register<UE::Ava::FShapeComponentMaterialBridge>();
+	MaterialBridgeRegistry.Register<UE::Ava::FShapeComponentMaterialBridge>(MaterialBridgePriority);
 }
diff --git a/Engine/Plugins/VirtualProduction/Avalanche/Source/AvalancheShapes/Public/AvaShapeComponentMaterialBridge.h b/Engine/Plugins/VirtualProduction/Avalanche/Source/AvalancheShapes/Public/AvaShapeComponentMaterialBridge.h
--- a/Engine/Plugins/VirtualProduction/Avalanche/Source/AvalancheShapes/Public/AvaShapeComponentMaterialBridge.h
+++ b/Engine/Plugins/VirtualProduction/Avalanche/Source/AvalancheShapes/Public/AvaShapeComponentMaterialBridge.h
@@ -15,7 +15,7 @@ class FShapeComponentMaterialBridge : public FMaterialBridge
 protected:
 	//~ Begin FMaterialBridge
 	UE_API virtual const UStruct* OnGetBridgedType() const override;
-	UE_API virtual EVisitFlow OnVisitMaterials(const FMaterialBridgeVisitContext& InContext, TFunctionRef<EVisitFlow(const FMaterialBridgeVisitContext&, UMaterialInterface*)> InFunc, const FVisitOptions& InOptions) const override;
+	UE_API virtual EVisitFlow OnVisitMaterials(const FVisitContext& InContext, TFunctionRef<EVisitFlow(const FVisitContext&, UMaterialInterface*)> InFunc, const FVisitOptions& InOptions) const override;
 	//~ End FMaterialBridge
 };
 
diff --git a/Engine/Plugins/VirtualProduction/Avalanche/Source/AvalancheText/Private/AvaText3DComponentMaterialBridge.cpp b/Engine/Plugins/VirtualProduction/Avalanche/Source/AvalancheText/Private/AvaText3DComponentMaterialBridge.cpp
--- a/Engine/Plugins/VirtualProduction/Avalanche/Source/AvalancheText/Private/AvaText3DComponentMaterialBridge.cpp
+++ b/Engine/Plugins/VirtualProduction/Avalanche/Source/AvalancheText/Private/AvaText3DComponentMaterialBridge.cpp
@@ -13,7 +13,7 @@ const UStruct* FText3DComponentMaterialBridge::OnGetBridgedType() const
 	return UText3DComponent::StaticClass();
 }
 
-EVisitFlow FText3DComponentMaterialBridge::OnVisitMaterials(const FMaterialBridgeVisitContext& InContext, TFunctionRef<EVisitFlow(const FMaterialBridgeVisitContext&, UMaterialInterface*)> InFunc, const FVisitOptions& InOptions) const
+EMaterialBridgeVisitFlow FText3DComponentMaterialBridge::OnVisitMaterials(const FVisitContext& InContext, TFunctionRef<EVisitFlow(const FVisitContext&, UMaterialInterface*)> InFunc, const FVisitOptions& InOptions) const
 {
 	const UText3DComponent& Component = InContext.Container.Get<const UText3DComponent>();
 
diff --git a/Engine/Plugins/VirtualProduction/Avalanche/Source/AvalancheText/Private/AvaTextModule.cpp b/Engine/Plugins/VirtualProduction/Avalanche/Source/AvalancheText/Private/AvaTextModule.cpp
--- a/Engine/Plugins/VirtualProduction/Avalanche/Source/AvalancheText/Private/AvaTextModule.cpp
+++ b/Engine/Plugins/VirtualProduction/Avalanche/Source/AvalancheText/Private/AvaTextModule.cpp
@@ -9,6 +9,8 @@ IMPLEMENT_MODULE(FAvaTextModule, AvalancheText)
 
 void FAvaTextModule::StartupModule()
 {
+	constexpr uint32 MaterialBridgePriority = 0;
+
 	UE::Ava::FMaterialBridgeRegistry& MaterialBridgeRegistry = UE::Ava::FMaterialBridgeRegistry::GetMutable();
-	MaterialBridgeRegistry.Register<UE::Ava::FText3DComponentMaterialBridge>();
+	MaterialBridgeRegistry.Register<UE::Ava::FText3DComponentMaterialBridge>(MaterialBridgePriority);
 }
diff --git a/Engine/Plugins/VirtualProduction/Avalanche/Source/AvalancheText/Public/AvaText3DComponentMaterialBridge.h b/Engine/Plugins/VirtualProduction/Avalanche/Source/AvalancheText/Public/AvaText3DComponentMaterialBridge.h
--- a/Engine/Plugins/VirtualProduction/Avalanche/Source/AvalancheText/Public/AvaText3DComponentMaterialBridge.h
+++ b/Engine/Plugins/VirtualProduction/Avalanche/Source/AvalancheText/Public/AvaText3DComponentMaterialBridge.h
@@ -15,7 +15,7 @@ class FText3DComponentMaterialBridge : public FMaterialBridge
 protected:
 	//~ Begin FMaterialBridge
 	UE_API virtual const UStruct* OnGetBridgedType() const override;
-	UE_API virtual EVisitFlow OnVisitMaterials(const FMaterialBridgeVisitContext& InContext, TFunctionRef<EVisitFlow(const FMaterialBridgeVisitContext&, UMaterialInterface*)> InFunc, const FVisitOptions& InOptions) const override;
+	UE_API virtual EVisitFlow OnVisitMaterials(const FVisitContext& InContext, TFunctionRef<EVisitFlow(const FVisitContext&, UMaterialInterface*)> InFunc, const FVisitOptions& InOptions) const override;
 	//~ End FMaterialBridge
 };
 
