Motion Design: renamed 'VisitMaterials' to 'AccessSlots' and changed its parameters to indicate it's read access.
This is a change that prepares for the new addition of 'AccessSlots' overload for write accessl
+ Moved common built-in material bridges to a common folder for clarity.

#Synced-CL 48481255

diff --git a/Engine/Plugins/VirtualProduction/Avalanche/Source/AvalancheCore/Public/AvaEnums.h b/Engine/Plugins/VirtualProduction/Avalanche/Source/AvalancheCore/Public/AvaEnums.h
new file mode 100644
--- /dev/null
+++ b/Engine/Plugins/VirtualProduction/Avalanche/Source/AvalancheCore/Public/AvaEnums.h
@@ -0,0 +1,19 @@
+// Copyright Epic Games, Inc. All Rights Reserved.
+
+#pragma once
+
+#include "HAL/Platform.h"
+
+namespace UE::Ava
+{
+
+/** Indicates how iteration should be handled */
+enum class EControlFlow : uint8
+{
+	/** Continue iteration */
+	Continue,
+	/** Break iteration */
+	Break,
+};
+
+} // UE::Ava
diff --git a/Engine/Plugins/VirtualProduction/Avalanche/Source/AvalancheMaterial/Private/AvaMaterialModule.cpp b/Engine/Plugins/VirtualProduction/Avalanche/Source/AvalancheMaterial/Private/AvaMaterialModule.cpp
--- a/Engine/Plugins/VirtualProduction/Avalanche/Source/AvalancheMaterial/Private/AvaMaterialModule.cpp
+++ b/Engine/Plugins/VirtualProduction/Avalanche/Source/AvalancheMaterial/Private/AvaMaterialModule.cpp
@@ -1,14 +1,14 @@
 // Copyright Epic Games, Inc. All Rights Reserved.
 
 #include "AvaMaterialModule.h"
-#include "MaterialBridge/AvaActorMaterialBridge.h"
-#include "MaterialBridge/AvaDecalComponentMaterialBridge.h"
-#include "MaterialBridge/AvaLevelMaterialBridge.h"
 #include "MaterialBridge/AvaMaterialBridge.h"
 #include "MaterialBridge/AvaMaterialBridgeLog.h"
 #include "MaterialBridge/AvaMaterialBridgeRegistry.h"
-#include "MaterialBridge/AvaPrimitiveComponentMaterialBridge.h"
-#include "MaterialBridge/AvaStaticMeshMaterialBridge.h"
+#include "MaterialBridge/Common/AvaActorMaterialBridge.h"
+#include "MaterialBridge/Common/AvaDecalComponentMaterialBridge.h"
+#include "MaterialBridge/Common/AvaLevelMaterialBridge.h"
+#include "MaterialBridge/Common/AvaPrimitiveComponentMaterialBridge.h"
+#include "MaterialBridge/Common/AvaStaticMeshMaterialBridge.h"
 #include "Modules/ModuleManager.h"
 
 DEFINE_LOG_CATEGORY(LogAvaMaterialBridge);
@@ -17,14 +17,14 @@ IMPLEMENT_MODULE(FAvaMaterialModule, AvalancheMaterial)
 
 void FAvaMaterialModule::StartupModule()
 {
-	// Low priority for built-in handlers
-	constexpr uint32 MaterialBridgePriority = 0;
+	// Low priority for common bridges
+	constexpr uint32 DefaultPriority = 0;
 
-	// Register built-in material bridges
+	// Register common material bridges
 	UE::Ava::FMaterialBridgeRegistry& MaterialBridgeRegistry = UE::Ava::FMaterialBridgeRegistry::GetMutable();
-	MaterialBridgeRegistry.Register<UE::Ava::FLevelMaterialBridge>(MaterialBridgePriority);
-	MaterialBridgeRegistry.Register<UE::Ava::FActorMaterialBridge>(MaterialBridgePriority);
-	MaterialBridgeRegistry.Register<UE::Ava::FPrimitiveComponentMaterialBridge>(MaterialBridgePriority);
-	MaterialBridgeRegistry.Register<UE::Ava::FDecalComponentMaterialBridge>(MaterialBridgePriority);
-	MaterialBridgeRegistry.Register<UE::Ava::FStaticMeshMaterialBridge>(MaterialBridgePriority);
+	MaterialBridgeRegistry.Register<UE::Ava::FLevelMaterialBridge>(DefaultPriority);
+	MaterialBridgeRegistry.Register<UE::Ava::FActorMaterialBridge>(DefaultPriority);
+	MaterialBridgeRegistry.Register<UE::Ava::FPrimitiveComponentMaterialBridge>(DefaultPriority);
+	MaterialBridgeRegistry.Register<UE::Ava::FDecalComponentMaterialBridge>(DefaultPriority);
+	MaterialBridgeRegistry.Register<UE::Ava::FStaticMeshMaterialBridge>(DefaultPriority);
 }
diff --git a/Engine/Plugins/VirtualProduction/Avalanche/Source/AvalancheMaterial/Private/MaterialBridge/AvaActorMaterialBridge.cpp b/Engine/Plugins/VirtualProduction/Avalanche/Source/AvalancheMaterial/Private/MaterialBridge/AvaActorMaterialBridge.cpp
deleted file mode 100644
diff --git a/Engine/Plugins/VirtualProduction/Avalanche/Source/AvalancheMaterial/Private/MaterialBridge/AvaDecalComponentMaterialBridge.cpp b/Engine/Plugins/VirtualProduction/Avalanche/Source/AvalancheMaterial/Private/MaterialBridge/AvaDecalComponentMaterialBridge.cpp
deleted file mode 100644
diff --git a/Engine/Plugins/VirtualProduction/Avalanche/Source/AvalancheMaterial/Private/MaterialBridge/AvaLevelMaterialBridge.cpp b/Engine/Plugins/VirtualProduction/Avalanche/Source/AvalancheMaterial/Private/MaterialBridge/AvaLevelMaterialBridge.cpp
deleted file mode 100644
diff --git a/Engine/Plugins/VirtualProduction/Avalanche/Source/AvalancheMaterial/Private/MaterialBridge/AvaMaterialBridge.cpp b/Engine/Plugins/VirtualProduction/Avalanche/Source/AvalancheMaterial/Private/MaterialBridge/AvaMaterialBridge.cpp
--- a/Engine/Plugins/VirtualProduction/Avalanche/Source/AvalancheMaterial/Private/MaterialBridge/AvaMaterialBridge.cpp
+++ b/Engine/Plugins/VirtualProduction/Avalanche/Source/AvalancheMaterial/Private/MaterialBridge/AvaMaterialBridge.cpp
@@ -1,10 +1,12 @@
 // Copyright Epic Games, Inc. All Rights Reserved.
 
 #include "MaterialBridge/AvaMaterialBridge.h"
+#include "MaterialBridge/AvaMaterialBridgeLog.h"
 #include "MaterialBridge/AvaMaterialBridgeRegistry.h"
 
 namespace UE::Ava
 {
+
 void FMaterialBridge::Initialize(uint32 InPriority)
 {
 	BridgedType = OnGetBridgedType();
@@ -21,22 +23,22 @@ uint32 FMaterialBridge::GetPriority() const
 	return Priority;
 }
 
-bool FMaterialBridge::Supports(FConstDataView InContainer) const
+bool FMaterialBridge::IsMaterialContainerSupported(FConstDataView InMaterialContainer) const
 {
-	if (InContainer.IsValidFor(BridgedType))
+	if (InMaterialContainer.IsValidFor(BridgedType))
 	{
-		return OnSupports(InContainer);
+		return OnIsMaterialContainerSupported(InMaterialContainer);
 	}
 	return false;
 }
 
-EMaterialBridgeVisitFlow FMaterialBridge::VisitMaterials(const FVisitContext& InContext, TFunctionRef<EVisitFlow(const FVisitContext&, UMaterialInterface*)> InFunc, const FVisitOptions& InOptions) const
+EControlFlow FMaterialBridge::AccessSlots(const FReadSlotContext& InContext, TFunctionRef<EControlFlow(const FReadSlotContext&, const FSlot&)> InFunc, const FReadSlotOptions& InOptions) const
 {
-	if (InContext.Container.IsValidFor(BridgedType))
+	if (InContext.MaterialContainer.IsValidFor(BridgedType))
 	{
-		return OnVisitMaterials(InContext, InFunc, InOptions);
+		return OnAccessSlots(InContext, InFunc, InOptions);
 	}
-	return EVisitFlow::Continue;
+	return EControlFlow::Continue;
 }
 
 } // UE::Ava
diff --git a/Engine/Plugins/VirtualProduction/Avalanche/Source/AvalancheMaterial/Private/MaterialBridge/AvaMaterialBridgeRegistry.cpp b/Engine/Plugins/VirtualProduction/Avalanche/Source/AvalancheMaterial/Private/MaterialBridge/AvaMaterialBridgeRegistry.cpp
--- a/Engine/Plugins/VirtualProduction/Avalanche/Source/AvalancheMaterial/Private/MaterialBridge/AvaMaterialBridgeRegistry.cpp
+++ b/Engine/Plugins/VirtualProduction/Avalanche/Source/AvalancheMaterial/Private/MaterialBridge/AvaMaterialBridgeRegistry.cpp
@@ -17,13 +17,13 @@ FMaterialBridgeRegistry& FMaterialBridgeRegistry::GetMutable()
 	return MaterialBridgeRegistry;
 }
 
-const FMaterialBridge* FMaterialBridgeRegistry::GetMaterialBridge(FConstDataView InData) const
+const FMaterialBridge* FMaterialBridgeRegistry::GetMaterialBridge(FConstDataView InMaterialContainer) const
 {
-	if (!InData.IsValid())
+	if (!InMaterialContainer.IsValid())
 	{
 		return nullptr;
 	}
-	for (const UStruct* Struct : InData.GetStruct()->GetSuperStructIterator())
+	for (const UStruct* Struct : InMaterialContainer.GetStruct()->GetSuperStructIterator())
 	{
 		const TArray<TUniquePtr<FMaterialBridge>>* MaterialBridges = RegisteredMaterialBridges.Find(Struct);
 		if (!MaterialBridges)
@@ -32,7 +32,7 @@ const FMaterialBridge* FMaterialBridgeRegistry::GetMaterialBridge(FConstDataView
 		}
 		for (const TUniquePtr<FMaterialBridge>& MaterialBridge : *MaterialBridges)
 		{
-			if (MaterialBridge.IsValid() && MaterialBridge->Supports(InData))
+			if (MaterialBridge.IsValid() && MaterialBridge->IsMaterialContainerSupported(InMaterialContainer))
 			{
 				return MaterialBridge.Get();
 			}
diff --git a/Engine/Plugins/VirtualProduction/Avalanche/Source/AvalancheMaterial/Private/MaterialBridge/AvaMaterialBridgeSlot.cpp b/Engine/Plugins/VirtualProduction/Avalanche/Source/AvalancheMaterial/Private/MaterialBridge/AvaMaterialBridgeSlot.cpp
new file mode 100644
--- /dev/null
+++ b/Engine/Plugins/VirtualProduction/Avalanche/Source/AvalancheMaterial/Private/MaterialBridge/AvaMaterialBridgeSlot.cpp
@@ -0,0 +1,24 @@
+// Copyright Epic Games, Inc. All Rights Reserved.
+
+#include "MaterialBridge/AvaMaterialBridgeSlot.h"
+
+namespace UE::Ava
+{
+
+FMaterialBridgeSlot::FMaterialBridgeSlot(UMaterialInterface* InMaterial)
+	: OriginalMaterial(InMaterial)
+	, CurrentMaterial(InMaterial)
+{
+}
+
+void FMaterialBridgeSlot::SetMaterial(UMaterialInterface* InMaterial)
+{
+	CurrentMaterial = InMaterial;
+}
+
+bool FMaterialBridgeSlot::HasChangedMaterial() const
+{
+	return OriginalMaterial != CurrentMaterial;
+}
+
+} // UE::Ava
diff --git a/Engine/Plugins/VirtualProduction/Avalanche/Source/AvalancheMaterial/Private/MaterialBridge/AvaMaterialBridgeVisitContext.cpp b/Engine/Plugins/VirtualProduction/Avalanche/Source/AvalancheMaterial/Private/MaterialBridge/AvaMaterialBridgeVisitContext.cpp
deleted file mode 100644
diff --git a/Engine/Plugins/VirtualProduction/Avalanche/Source/AvalancheMaterial/Private/MaterialBridge/AvaPrimitiveComponentMaterialBridge.cpp b/Engine/Plugins/VirtualProduction/Avalanche/Source/AvalancheMaterial/Private/MaterialBridge/AvaPrimitiveComponentMaterialBridge.cpp
deleted file mode 100644
diff --git a/Engine/Plugins/VirtualProduction/Avalanche/Source/AvalancheMaterial/Private/MaterialBridge/AvaStaticMeshMaterialBridge.cpp b/Engine/Plugins/VirtualProduction/Avalanche/Source/AvalancheMaterial/Private/MaterialBridge/AvaStaticMeshMaterialBridge.cpp
deleted file mode 100644
diff --git a/Engine/Plugins/VirtualProduction/Avalanche/Source/AvalancheMaterial/Private/MaterialBridge/Common/AvaActorMaterialBridge.cpp b/Engine/Plugins/VirtualProduction/Avalanche/Source/AvalancheMaterial/Private/MaterialBridge/Common/AvaActorMaterialBridge.cpp
new file mode 100644
--- /dev/null
+++ b/Engine/Plugins/VirtualProduction/Avalanche/Source/AvalancheMaterial/Private/MaterialBridge/Common/AvaActorMaterialBridge.cpp
@@ -0,0 +1,54 @@
+// Copyright Epic Games, Inc. All Rights Reserved.
+
+#include "MaterialBridge/Common/AvaActorMaterialBridge.h"
+#include "GameFramework/Actor.h"
+#include "MaterialBridge/AvaMaterialBridgeRegistry.h"
+
+namespace UE::Ava::Private
+{
+// Iterates the components of an actor and calls the given func for every valid component with a found material bridge
+EControlFlow ForEachBridgedComponent(const AActor& InActor, TNotNull<const FMaterialBridgeRegistry*> InRegistry, TFunctionRef<EControlFlow(TNotNull<UActorComponent*>, TNotNull<const FMaterialBridge*>)> InFunc)
+{
+	EControlFlow ControlFlow = EControlFlow::Continue;
+
+	constexpr bool bIncludeFromChildActors = true;
+
+	InActor.ForEachComponent(bIncludeFromChildActors,
+		[&ControlFlow, &InFunc, &InRegistry](UActorComponent* InComponent)
+		{
+			if (!InComponent || ControlFlow != EControlFlow::Continue)
+			{
+				return;
+			}
+			if (const FMaterialBridge* MaterialBridge = InRegistry->GetMaterialBridge(FConstDataView(InComponent)))
+			{
+				ControlFlow = InFunc(InComponent, MaterialBridge);
+			}
+		});
+
+	return ControlFlow;
+}
+
+} // UE::Ava::Private
+
+namespace UE::Ava
+{
+
+const UStruct* FActorMaterialBridge::OnGetBridgedType() const
+{
+	return AActor::StaticClass();
+}
+
+EControlFlow FActorMaterialBridge::OnAccessSlots(const FReadSlotContext& InContext, TFunctionRef<EControlFlow(const FReadSlotContext&, const FSlot&)> InFunc, const FReadSlotOptions& InOptions) const
+{
+	const AActor& Actor = InContext.MaterialContainer.Get<const AActor>();
+
+	return Private::ForEachBridgedComponent(Actor, InOptions.MaterialBridgeRegistry,
+		[&InContext, &InFunc, &InOptions](TNotNull<UActorComponent*> InComponent, TNotNull<const FMaterialBridge*> InMaterialBridge)->EControlFlow
+		{
+			return InMaterialBridge->AccessSlots(FReadSlotContext(InComponent, &InContext), InFunc, InOptions);
+		});
+
+}
+
+} // UE::Ava
diff --git a/Engine/Plugins/VirtualProduction/Avalanche/Source/AvalancheMaterial/Private/MaterialBridge/Common/AvaDecalComponentMaterialBridge.cpp b/Engine/Plugins/VirtualProduction/Avalanche/Source/AvalancheMaterial/Private/MaterialBridge/Common/AvaDecalComponentMaterialBridge.cpp
new file mode 100644
--- /dev/null
+++ b/Engine/Plugins/VirtualProduction/Avalanche/Source/AvalancheMaterial/Private/MaterialBridge/Common/AvaDecalComponentMaterialBridge.cpp
@@ -0,0 +1,23 @@
+// Copyright Epic Games, Inc. All Rights Reserved.
+
+#include "MaterialBridge/Common/AvaDecalComponentMaterialBridge.h"
+#include "AvaDataView.h"
+#include "Components/DecalComponent.h"
+
+namespace UE::Ava
+{
+
+const UStruct* FDecalComponentMaterialBridge::OnGetBridgedType() const
+{
+	return UDecalComponent::StaticClass();
+}
+
+EControlFlow FDecalComponentMaterialBridge::OnAccessSlots(const FReadSlotContext& InContext, TFunctionRef<EControlFlow(const FReadSlotContext&, const FSlot&)> InFunc, const FReadSlotOptions& InOptions) const
+{
+	const UDecalComponent& Component = InContext.MaterialContainer.Get<const UDecalComponent>();
+
+	const FSlot Slot(Component.GetDecalMaterial());
+	return InFunc(InContext, Slot);
+}
+
+} // UE::Ava
diff --git a/Engine/Plugins/VirtualProduction/Avalanche/Source/AvalancheMaterial/Private/MaterialBridge/Common/AvaLevelMaterialBridge.cpp b/Engine/Plugins/VirtualProduction/Avalanche/Source/AvalancheMaterial/Private/MaterialBridge/Common/AvaLevelMaterialBridge.cpp
new file mode 100644
--- /dev/null
+++ b/Engine/Plugins/VirtualProduction/Avalanche/Source/AvalancheMaterial/Private/MaterialBridge/Common/AvaLevelMaterialBridge.cpp
@@ -0,0 +1,35 @@
+// Copyright Epic Games, Inc. All Rights Reserved.
+
+#include "MaterialBridge/Common/AvaLevelMaterialBridge.h"
+#include "Engine/Level.h"
+#include "GameFramework/Actor.h"
+#include "MaterialBridge/AvaMaterialBridgeRegistry.h"
+
+namespace UE::Ava
+{
+
+const UStruct* FLevelMaterialBridge::OnGetBridgedType() const
+{
+	return ULevel::StaticClass();
+}
+
+EControlFlow FLevelMaterialBridge::OnAccessSlots(const FReadSlotContext& InContext, TFunctionRef<EControlFlow(const FReadSlotContext&, const FSlot&)> InFunc, const FReadSlotOptions& InOptions) const
+{
+	const ULevel& Level = InContext.MaterialContainer.Get<const ULevel>();
+
+	for (AActor* Actor : Level.Actors)
+	{
+		if (const FMaterialBridge* MaterialBridge = InOptions.MaterialBridgeRegistry->GetMaterialBridge(FConstDataView(Actor)))
+		{
+			const EControlFlow ControlFlow = MaterialBridge->AccessSlots(FReadSlotContext(Actor, &InContext), InFunc, InOptions);
+			if (ControlFlow == EControlFlow::Break)
+			{
+				return EControlFlow::Break;
+			}
+		}
+	}
+
+	return EControlFlow::Continue;
+}
+
+} // UE::Ava
diff --git a/Engine/Plugins/VirtualProduction/Avalanche/Source/AvalancheMaterial/Private/MaterialBridge/Common/AvaPrimitiveComponentMaterialBridge.cpp b/Engine/Plugins/VirtualProduction/Avalanche/Source/AvalancheMaterial/Private/MaterialBridge/Common/AvaPrimitiveComponentMaterialBridge.cpp
new file mode 100644
--- /dev/null
+++ b/Engine/Plugins/VirtualProduction/Avalanche/Source/AvalancheMaterial/Private/MaterialBridge/Common/AvaPrimitiveComponentMaterialBridge.cpp
@@ -0,0 +1,34 @@
+// Copyright Epic Games, Inc. All Rights Reserved.
+
+#include "MaterialBridge/Common/AvaPrimitiveComponentMaterialBridge.h"
+#include "AvaDataView.h"
+#include "Components/PrimitiveComponent.h"
+
+namespace UE::Ava
+{
+
+const UStruct* FPrimitiveComponentMaterialBridge::OnGetBridgedType() const
+{
+	return UPrimitiveComponent::StaticClass();
+}
+
+EControlFlow FPrimitiveComponentMaterialBridge::OnAccessSlots(const FReadSlotContext& InContext, TFunctionRef<EControlFlow(const FReadSlotContext&, const FSlot&)> InFunc, const FReadSlotOptions& InOptions) const
+{
+	const UPrimitiveComponent& Component = InContext.MaterialContainer.Get<const UPrimitiveComponent>();
+
+	const int32 MaterialCount = Component.GetNumMaterials();
+
+	for (int32 MaterialIndex = 0; MaterialIndex < MaterialCount; ++MaterialIndex)
+	{
+		const FSlot Slot(Component.GetMaterial(MaterialIndex));
+
+		if (InFunc(InContext, Slot) == EControlFlow::Break)
+		{
+			return EControlFlow::Break;
+		}
+	}
+
+	return EControlFlow::Continue;
+}
+
+} // UE::Ava
diff --git a/Engine/Plugins/VirtualProduction/Avalanche/Source/AvalancheMaterial/Private/MaterialBridge/Common/AvaStaticMeshMaterialBridge.cpp b/Engine/Plugins/VirtualProduction/Avalanche/Source/AvalancheMaterial/Private/MaterialBridge/Common/AvaStaticMeshMaterialBridge.cpp
new file mode 100644
--- /dev/null
+++ b/Engine/Plugins/VirtualProduction/Avalanche/Source/AvalancheMaterial/Private/MaterialBridge/Common/AvaStaticMeshMaterialBridge.cpp
@@ -0,0 +1,56 @@
+// Copyright Epic Games, Inc. All Rights Reserved.
+
+#include "MaterialBridge/Common/AvaStaticMeshMaterialBridge.h"
+#include "MaterialBridge/AvaMaterialBridgeLog.h"
+#include "Engine/StaticMesh.h"
+
+namespace UE::Ava
+{
+
+const UStruct* FStaticMeshMaterialBridge::OnGetBridgedType() const
+{
+	return UStaticMesh::StaticClass();
+}
+
+EControlFlow FStaticMeshMaterialBridge::OnAccessSlots(const FReadSlotContext& InContext, TFunctionRef<EControlFlow(const FReadSlotContext&, const FSlot&)> InFunc, const FReadSlotOptions& InOptions) const
+{
+	const UStaticMesh& StaticMesh = InContext.MaterialContainer.Get<const UStaticMesh>();
+
+	// If static mesh is compiling and visitor does not want to wait on materials, bail out.
+	// NOTE: Only LockedProperties needs to be checked for EStaticMeshAsyncProperties::StaticMaterials, but there's currently no option to check for it directly.
+	if (StaticMesh.IsCompiling() && !InOptions.bTryWaitOnContainerCompletion)
+	{
+		if (InOptions.OnContainerPendingCompletion.IsBound())
+		{
+			UE_LOG(LogAvaMaterialBridge, Verbose, TEXT("Material Bridge could not visit materials of static mesh '%s'. It is currently compiling."), *StaticMesh.GetName());
+			return InOptions.OnContainerPendingCompletion.Execute(InContext);
+		}
+
+		UE_LOG(LogAvaMaterialBridge, Warning, TEXT("Material Bridge could not visit materials of static mesh '%s'. It is currently compiling."), *StaticMesh.GetName());
+		return EControlFlow::Continue;
+	}
+
+	for (const FStaticMaterial& StaticMaterial : StaticMesh.GetStaticMaterials())
+	{
+		// Base Material
+		{
+			const FSlot Slot(StaticMaterial.MaterialInterface);
+			if (InFunc(InContext, Slot) == EControlFlow::Break)
+			{
+				return EControlFlow::Break;
+			}
+		}
+
+		// Overlay Material
+		{
+			const FSlot Slot(StaticMaterial.OverlayMaterialInterface);
+			if (InFunc(InContext, Slot) == EControlFlow::Break)
+			{
+				return EControlFlow::Break;
+			}
+		}
+	}
+	return EControlFlow::Continue;
+}
+
+} // UE::Ava
diff --git a/Engine/Plugins/VirtualProduction/Avalanche/Source/AvalancheMaterial/Private/MaterialBridge/Context/AvaMaterialBridgeReadSlotContext.cpp b/Engine/Plugins/VirtualProduction/Avalanche/Source/AvalancheMaterial/Private/MaterialBridge/Context/AvaMaterialBridgeReadSlotContext.cpp
new file mode 100644
--- /dev/null
+++ b/Engine/Plugins/VirtualProduction/Avalanche/Source/AvalancheMaterial/Private/MaterialBridge/Context/AvaMaterialBridgeReadSlotContext.cpp
@@ -0,0 +1,25 @@
+// Copyright Epic Games, Inc. All Rights Reserved.
+
+#include "MaterialBridge/Context/AvaMaterialBridgeReadSlotContext.h"
+#include "MaterialBridge/AvaMaterialBridgeRegistry.h"
+
+namespace UE::Ava
+{
+
+FMaterialBridgeReadSlotContext::FMaterialBridgeReadSlotContext(const UObject* InContainer)
+	: MaterialContainer(InContainer)
+{
+}
+
+FMaterialBridgeReadSlotContext::FMaterialBridgeReadSlotContext(const UObject* InContainer, TNotNull<const FMaterialBridgeReadSlotContext*> InParentContext)
+	: MaterialContainer(InContainer)
+	, ParentContext(InParentContext)
+{
+}
+
+FMaterialBridgeReadSlotOptions::FMaterialBridgeReadSlotOptions()
+	: MaterialBridgeRegistry(&FMaterialBridgeRegistry::Get())
+{
+}
+
+} // UE::Ava
diff --git a/Engine/Plugins/VirtualProduction/Avalanche/Source/AvalancheMaterial/Private/MaterialCache/AvaMaterialCacheHelper.cpp b/Engine/Plugins/VirtualProduction/Avalanche/Source/AvalancheMaterial/Private/MaterialCache/AvaMaterialCacheHelper.cpp
--- a/Engine/Plugins/VirtualProduction/Avalanche/Source/AvalancheMaterial/Private/MaterialCache/AvaMaterialCacheHelper.cpp
+++ b/Engine/Plugins/VirtualProduction/Avalanche/Source/AvalancheMaterial/Private/MaterialCache/AvaMaterialCacheHelper.cpp
@@ -551,44 +551,45 @@ bool FMaterialCacheHelper::GatherMaterialData(const UObject* InObject, FObjectDa
 		};
 
 	// Options to not wait for the container to have materials available and instead get a callback
-	FMaterialBridgeVisitOptions VisitOptions;
-	VisitOptions.bTryWaitOnContainerCompletion = false;
-	VisitOptions.OnContainerPendingCompletion.BindLambda([&InObjectData](const FMaterialBridgeVisitContext& InContext)->EMaterialBridgeVisitFlow
+	FMaterialBridgeReadSlotOptions ReadSlotOptions;
+	ReadSlotOptions.bTryWaitOnContainerCompletion = false;
+	ReadSlotOptions.OnContainerPendingCompletion.BindLambda([&InObjectData](const FMaterialBridgeReadSlotContext& InContext)->EControlFlow
 		{
-			if (const UObject* Container = InContext.Container.GetPtr<const UObject>())
+			if (const UObject* Container = InContext.MaterialContainer.GetPtr<const UObject>())
 			{
 				InObjectData.PendingSubobjects.Add(Container);
 			}
 			else
 			{
-				ensureMsgf(0, TEXT("Container of type '%s' not supported for deferring material gather"), *GetNameSafe(InContext.Container.GetStruct()));
+				ensureMsgf(0, TEXT("Container of type '%s' not supported for deferring material gather"), *GetNameSafe(InContext.MaterialContainer.GetStruct()));
 			}
-			return EMaterialBridgeVisitFlow::Continue;
+			return EControlFlow::Continue;
 		});
 
 	// Gather all materials for object
-	MaterialBridge->VisitMaterials(FMaterialBridgeVisitContext(const_cast<UObject*>(InObject)),
-		[&AddIncompleteMaterial](const FMaterialBridgeVisitContext& InContext, UMaterialInterface* InMaterial)->EMaterialBridgeVisitFlow
+	MaterialBridge->AccessSlots(FMaterialBridgeReadSlotContext(InObject),
+		[&AddIncompleteMaterial](const FMaterialBridgeReadSlotContext& InContext, const FMaterialBridgeSlot& InSlot)->EControlFlow
 		{
-			if (!InMaterial)
+			UMaterialInterface* const SlotMaterial = InSlot.GetMaterial();
+			if (!SlotMaterial)
 			{
-				return EMaterialBridgeVisitFlow::Continue;
+				return EControlFlow::Continue;
 			}
 
-			UMaterial* const Material = InMaterial->GetMaterial();
+			UMaterial* const Material = SlotMaterial->GetMaterial();
 
 			// Add interface's material if it is different from the interface
-			if (Material && Material != InMaterial && !Material->IsComplete())
+			if (Material && Material != SlotMaterial && !Material->IsComplete())
 			{
 				AddIncompleteMaterial(Material);
 			}
-			if (!InMaterial->IsComplete())
+			if (!SlotMaterial->IsComplete())
 			{
-				AddIncompleteMaterial(InMaterial);
+				AddIncompleteMaterial(SlotMaterial);
 			}
-			return EMaterialBridgeVisitFlow::Continue;
+			return EControlFlow::Continue;
 		}
-		, VisitOptions);
+		, ReadSlotOptions);
 
 	return true;
 }
diff --git a/Engine/Plugins/VirtualProduction/Avalanche/Source/AvalancheMaterial/Public/MaterialBridge/AvaActorMaterialBridge.h b/Engine/Plugins/VirtualProduction/Avalanche/Source/AvalancheMaterial/Public/MaterialBridge/AvaActorMaterialBridge.h
deleted file mode 100644
diff --git a/Engine/Plugins/VirtualProduction/Avalanche/Source/AvalancheMaterial/Public/MaterialBridge/AvaDecalComponentMaterialBridge.h b/Engine/Plugins/VirtualProduction/Avalanche/Source/AvalancheMaterial/Public/MaterialBridge/AvaDecalComponentMaterialBridge.h
deleted file mode 100644
diff --git a/Engine/Plugins/VirtualProduction/Avalanche/Source/AvalancheMaterial/Public/MaterialBridge/AvaLevelMaterialBridge.h b/Engine/Plugins/VirtualProduction/Avalanche/Source/AvalancheMaterial/Public/MaterialBridge/AvaLevelMaterialBridge.h
deleted file mode 100644
diff --git a/Engine/Plugins/VirtualProduction/Avalanche/Source/AvalancheMaterial/Public/MaterialBridge/AvaMaterialBridge.h b/Engine/Plugins/VirtualProduction/Avalanche/Source/AvalancheMaterial/Public/MaterialBridge/AvaMaterialBridge.h
--- a/Engine/Plugins/VirtualProduction/Avalanche/Source/AvalancheMaterial/Public/MaterialBridge/AvaMaterialBridge.h
+++ b/Engine/Plugins/VirtualProduction/Avalanche/Source/AvalancheMaterial/Public/MaterialBridge/AvaMaterialBridge.h
@@ -2,8 +2,10 @@
 
 #pragma once
 
-#include "AvaMaterialBridgeVisitContext.h"
+#include "AvaEnums.h"
+#include "AvaMaterialBridgeSlot.h"
 #include "Containers/ContainersFwd.h"
+#include "Context/AvaMaterialBridgeReadSlotContext.h"
 #include "StructUtils/StructUtilsTypes.h"
 #include "UObject/ObjectPtr.h"
 
@@ -14,13 +16,16 @@ class UMaterialInterface;
 namespace UE::Ava
 {
 
-/** Base class that allows interacting with materials from any type of data view (Actors, ActorComponents, etc). */
+/** 
+ * Base class that allows interacting with material containers from any type of data view (Actors, ActorComponents, etc). 
+ * A Material Bridge allows the user to generically access a material container's material slots for read (see AccessSlots)
+ */
 class FMaterialBridge
 {
 public:
-	using EVisitFlow = EMaterialBridgeVisitFlow;
-	using FVisitOptions = FMaterialBridgeVisitOptions;
-	using FVisitContext = FMaterialBridgeVisitContext;
+	using FSlot              = FMaterialBridgeSlot;
+	using FReadSlotOptions   = FMaterialBridgeReadSlotOptions;
+	using FReadSlotContext   = FMaterialBridgeReadSlotContext;
 
 	virtual ~FMaterialBridge() = default;
 
@@ -35,32 +40,35 @@ public:
 
 	/**
 	 * Whether this Material Bridge supports the given Material Container.
-	 * @param InContainer a valid material container of the supported type of this material bridge.
+	 * @param InMaterialContainer a valid material container of the supported type of this material bridge.
 	 * @return true if this material bridge supports the container, false otherwise.
 	 */
-	UE_API bool Supports(FConstDataView InContainer) const;
+	UE_API bool IsMaterialContainerSupported(FConstDataView InMaterialContainer) const;
 
 	/**
-	 * Visits all the used materials in the given context.
-	 * @param InContext the current context to find materials from.
-	 * @param InFunc function to call for each material found.
-	 * @param InOptions visit options. see FVisitOptions
-	 * @return the iteration result of VisitMaterials.
+	 * Access the found material slots in the given context for read.
+	 * @param InContext the current context to access materials
+	 * @param InFunc function to call for each material slot accessed.
+	 * @param InOptions access options. see FMaterialBridgeReadSlotAccessOptions
+	 * @return the iteration result of AccessSlots.
 	 */
-	UE_API EVisitFlow VisitMaterials(const FVisitContext& InContext, TFunctionRef<EVisitFlow(const FVisitContext&, UMaterialInterface*)> InFunc, const FVisitOptions& InOptions = {}) const;
+	UE_API EControlFlow AccessSlots(const FReadSlotContext& InContext, TFunctionRef<EControlFlow(const FReadSlotContext&, const FSlot&)> InFunc, const FReadSlotOptions& InOptions) const;
 
 protected:
 	/** Retrieves the type that this bridge supports. Only called once in Initialize() */
 	virtual const UStruct* OnGetBridgedType() const = 0;
 
 	/**  Whether this material bridge supports the given material container.*/
-	virtual bool OnSupports(FConstDataView InContainer) const
+	virtual bool OnIsMaterialContainerSupported(FConstDataView InMaterialContainer) const
 	{
 		return true;
 	}
 
-	/** Get all the used materials in the given context. */
-	virtual EVisitFlow OnVisitMaterials(const FVisitContext& InContext, TFunctionRef<EVisitFlow(const FVisitContext&, UMaterialInterface*)> InFunc, const FVisitOptions& InOptions) const = 0;
+	/** Accesses the found material slots in the given context for read */
+	virtual EControlFlow OnAccessSlots(const FReadSlotContext& InContext, TFunctionRef<EControlFlow(const FReadSlotContext&, const FSlot&)> InFunc, const FReadSlotOptions& InOptions) const
+	{
+		return EControlFlow::Continue;
+	}
 
 private:
 	/** The type that this Material Bridge works on */
diff --git a/Engine/Plugins/VirtualProduction/Avalanche/Source/AvalancheMaterial/Public/MaterialBridge/AvaMaterialBridgeRegistry.h b/Engine/Plugins/VirtualProduction/Avalanche/Source/AvalancheMaterial/Public/MaterialBridge/AvaMaterialBridgeRegistry.h
--- a/Engine/Plugins/VirtualProduction/Avalanche/Source/AvalancheMaterial/Public/MaterialBridge/AvaMaterialBridgeRegistry.h
+++ b/Engine/Plugins/VirtualProduction/Avalanche/Source/AvalancheMaterial/Public/MaterialBridge/AvaMaterialBridgeRegistry.h
@@ -31,8 +31,8 @@ public:
 		this->RegisterInternal(MakeUnique<InMaterialBridgeType>(InArgs...), InPriority);
 	}
 
-	/** Tries to find the most relevant material bridge for a given object, returning a null if object is invalid, or if a relevant material bridge was not found. */
-	UE_API const FMaterialBridge* GetMaterialBridge(FConstDataView InData) const;
+	/** Tries to find the most relevant material bridge for a given material container, returning a null if object is invalid, or if a relevant material bridge was not found. */
+	UE_API const FMaterialBridge* GetMaterialBridge(FConstDataView InMaterialContainer) const;
 
 protected:
 	/** Initializes the material bridge and adds it to this registry */
diff --git a/Engine/Plugins/VirtualProduction/Avalanche/Source/AvalancheMaterial/Public/MaterialBridge/AvaMaterialBridgeSlot.h b/Engine/Plugins/VirtualProduction/Avalanche/Source/AvalancheMaterial/Public/MaterialBridge/AvaMaterialBridgeSlot.h
new file mode 100644
--- /dev/null
+++ b/Engine/Plugins/VirtualProduction/Avalanche/Source/AvalancheMaterial/Public/MaterialBridge/AvaMaterialBridgeSlot.h
@@ -0,0 +1,41 @@
+// Copyright Epic Games, Inc. All Rights Reserved.
+
+#pragma once
+
+#include "UObject/ObjectPtr.h"
+
+#define UE_API AVALANCHEMATERIAL_API
+
+class UMaterialInterface;
+
+namespace UE::Ava
+{
+
+/** Short-lived representation of an accessed material slot */
+class FMaterialBridgeSlot
+{
+public:
+	UE_API explicit FMaterialBridgeSlot(UMaterialInterface* InMaterial);
+
+	UMaterialInterface* GetMaterial() const
+	{
+		return CurrentMaterial;
+	}
+
+	/** Sets the new material for the slot */
+	UE_API void SetMaterial(UMaterialInterface* InMaterial);
+
+	/** Returns true if the material in the slot has changed */
+	UE_API bool HasChangedMaterial() const;
+
+protected:
+	/** The original material in the slot */
+	const TObjectPtr<UMaterialInterface> OriginalMaterial;
+
+	/** The current material in the slot */
+	TObjectPtr<UMaterialInterface> CurrentMaterial;
+};
+
+} // UE::Ava
+
+#undef UE_API
diff --git a/Engine/Plugins/VirtualProduction/Avalanche/Source/AvalancheMaterial/Public/MaterialBridge/AvaMaterialBridgeVisitContext.h b/Engine/Plugins/VirtualProduction/Avalanche/Source/AvalancheMaterial/Public/MaterialBridge/AvaMaterialBridgeVisitContext.h
deleted file mode 100644
diff --git a/Engine/Plugins/VirtualProduction/Avalanche/Source/AvalancheMaterial/Public/MaterialBridge/AvaPrimitiveComponentMaterialBridge.h b/Engine/Plugins/VirtualProduction/Avalanche/Source/AvalancheMaterial/Public/MaterialBridge/AvaPrimitiveComponentMaterialBridge.h
deleted file mode 100644
diff --git a/Engine/Plugins/VirtualProduction/Avalanche/Source/AvalancheMaterial/Public/MaterialBridge/AvaStaticMeshMaterialBridge.h b/Engine/Plugins/VirtualProduction/Avalanche/Source/AvalancheMaterial/Public/MaterialBridge/AvaStaticMeshMaterialBridge.h
deleted file mode 100644
diff --git a/Engine/Plugins/VirtualProduction/Avalanche/Source/AvalancheMaterial/Public/MaterialBridge/Common/AvaActorMaterialBridge.h b/Engine/Plugins/VirtualProduction/Avalanche/Source/AvalancheMaterial/Public/MaterialBridge/Common/AvaActorMaterialBridge.h
new file mode 100644
--- /dev/null
+++ b/Engine/Plugins/VirtualProduction/Avalanche/Source/AvalancheMaterial/Public/MaterialBridge/Common/AvaActorMaterialBridge.h
@@ -0,0 +1,24 @@
+// Copyright Epic Games, Inc. All Rights Reserved.
+
+#pragma once
+
+#include "MaterialBridge/AvaMaterialBridge.h"
+
+#define UE_API AVALANCHEMATERIAL_API
+
+namespace UE::Ava
+{
+
+/** Material Bridge for Actors */
+class FActorMaterialBridge : public FMaterialBridge
+{
+protected:
+	//~ Begin FMaterialBridge
+	UE_API virtual const UStruct* OnGetBridgedType() const override;
+	UE_API virtual EControlFlow OnAccessSlots(const FReadSlotContext& InContext, TFunctionRef<EControlFlow(const FReadSlotContext&, const FSlot&)> InFunc, const FReadSlotOptions& InOptions) const override;
+	//~ End FMaterialBridge
+};
+
+} // UE::Ava
+
+#undef UE_API
diff --git a/Engine/Plugins/VirtualProduction/Avalanche/Source/AvalancheMaterial/Public/MaterialBridge/Common/AvaDecalComponentMaterialBridge.h b/Engine/Plugins/VirtualProduction/Avalanche/Source/AvalancheMaterial/Public/MaterialBridge/Common/AvaDecalComponentMaterialBridge.h
new file mode 100644
--- /dev/null
+++ b/Engine/Plugins/VirtualProduction/Avalanche/Source/AvalancheMaterial/Public/MaterialBridge/Common/AvaDecalComponentMaterialBridge.h
@@ -0,0 +1,24 @@
+// Copyright Epic Games, Inc. All Rights Reserved.
+
+#pragma once
+
+#include "MaterialBridge/AvaMaterialBridge.h"
+
+#define UE_API AVALANCHEMATERIAL_API
+
+namespace UE::Ava
+{
+
+/** Material Bridge for Decal Components */
+class FDecalComponentMaterialBridge : public FMaterialBridge
+{
+protected:
+	//~ Begin FMaterialBridge
+	UE_API virtual const UStruct* OnGetBridgedType() const override;
+	UE_API virtual EControlFlow OnAccessSlots(const FReadSlotContext& InContext, TFunctionRef<EControlFlow(const FReadSlotContext&, const FSlot&)> InFunc, const FReadSlotOptions& InOptions) const override;
+	//~ End FMaterialBridge
+};
+
+} // UE::Ava
+
+#undef UE_API
diff --git a/Engine/Plugins/VirtualProduction/Avalanche/Source/AvalancheMaterial/Public/MaterialBridge/Common/AvaLevelMaterialBridge.h b/Engine/Plugins/VirtualProduction/Avalanche/Source/AvalancheMaterial/Public/MaterialBridge/Common/AvaLevelMaterialBridge.h
new file mode 100644
--- /dev/null
+++ b/Engine/Plugins/VirtualProduction/Avalanche/Source/AvalancheMaterial/Public/MaterialBridge/Common/AvaLevelMaterialBridge.h
@@ -0,0 +1,24 @@
+// Copyright Epic Games, Inc. All Rights Reserved.
+
+#pragma once
+
+#include "MaterialBridge/AvaMaterialBridge.h"
+
+#define UE_API AVALANCHEMATERIAL_API
+
+namespace UE::Ava
+{
+
+/** Material Bridge for Levels */
+class FLevelMaterialBridge : public FMaterialBridge
+{
+protected:
+	//~ Begin FMaterialBridge
+	UE_API virtual const UStruct* OnGetBridgedType() const override;
+	UE_API virtual EControlFlow OnAccessSlots(const FReadSlotContext& InContext, TFunctionRef<EControlFlow(const FReadSlotContext&, const FSlot&)> InFunc, const FReadSlotOptions& InOptions) const override;
+	//~ End FMaterialBridge
+};
+
+} // UE::Ava
+
+#undef UE_API
diff --git a/Engine/Plugins/VirtualProduction/Avalanche/Source/AvalancheMaterial/Public/MaterialBridge/Common/AvaPrimitiveComponentMaterialBridge.h b/Engine/Plugins/VirtualProduction/Avalanche/Source/AvalancheMaterial/Public/MaterialBridge/Common/AvaPrimitiveComponentMaterialBridge.h
new file mode 100644
--- /dev/null
+++ b/Engine/Plugins/VirtualProduction/Avalanche/Source/AvalancheMaterial/Public/MaterialBridge/Common/AvaPrimitiveComponentMaterialBridge.h
@@ -0,0 +1,23 @@
+// Copyright Epic Games, Inc. All Rights Reserved.
+
+#pragma once
+
+#include "MaterialBridge/AvaMaterialBridge.h"
+
+#define UE_API AVALANCHEMATERIAL_API
+
+namespace UE::Ava
+{
+/** Material Bridge for Primitive Components */
+class FPrimitiveComponentMaterialBridge : public FMaterialBridge
+{
+protected:
+	//~ Begin FMaterialBridge
+	UE_API virtual const UStruct* OnGetBridgedType() const override;
+	UE_API virtual EControlFlow OnAccessSlots(const FReadSlotContext& InContext, TFunctionRef<EControlFlow(const FReadSlotContext&, const FSlot&)> InFunc, const FReadSlotOptions& InOptions) const override;
+	//~ End FMaterialBridge
+};
+
+} // UE::Ava
+
+#undef UE_API
diff --git a/Engine/Plugins/VirtualProduction/Avalanche/Source/AvalancheMaterial/Public/MaterialBridge/Common/AvaStaticMeshMaterialBridge.h b/Engine/Plugins/VirtualProduction/Avalanche/Source/AvalancheMaterial/Public/MaterialBridge/Common/AvaStaticMeshMaterialBridge.h
new file mode 100644
--- /dev/null
+++ b/Engine/Plugins/VirtualProduction/Avalanche/Source/AvalancheMaterial/Public/MaterialBridge/Common/AvaStaticMeshMaterialBridge.h
@@ -0,0 +1,27 @@
+// Copyright Epic Games, Inc. All Rights Reserved.
+
+#pragma once
+
+#include "MaterialBridge/AvaMaterialBridge.h"
+
+#define UE_API AVALANCHEMATERIAL_API
+
+namespace UE::Ava
+{
+
+/**
+ * Material Bridge for Static Meshes
+ * Only deals with the Static Mesh's static materials.
+ */
+class FStaticMeshMaterialBridge : public FMaterialBridge
+{
+protected:
+	//~ Begin FMaterialBridge
+	UE_API virtual const UStruct* OnGetBridgedType() const override;
+	UE_API virtual EControlFlow OnAccessSlots(const FReadSlotContext& InContext, TFunctionRef<EControlFlow(const FReadSlotContext&, const FSlot&)> InFunc, const FReadSlotOptions& InOptions) const override;
+	//~ End FMaterialBridge
+};
+
+} // UE::Ava
+
+#undef UE_API
diff --git a/Engine/Plugins/VirtualProduction/Avalanche/Source/AvalancheMaterial/Public/MaterialBridge/Context/AvaMaterialBridgeReadSlotContext.h b/Engine/Plugins/VirtualProduction/Avalanche/Source/AvalancheMaterial/Public/MaterialBridge/Context/AvaMaterialBridgeReadSlotContext.h
new file mode 100644
--- /dev/null
+++ b/Engine/Plugins/VirtualProduction/Avalanche/Source/AvalancheMaterial/Public/MaterialBridge/Context/AvaMaterialBridgeReadSlotContext.h
@@ -0,0 +1,53 @@
+// Copyright Epic Games, Inc. All Rights Reserved.
+
+#pragma once
+
+#include "AvaDataView.h"
+#include "AvaEnums.h"
+
+#define UE_API AVALANCHEMATERIAL_API
+
+namespace UE::Ava
+{
+	class FMaterialBridgeRegistry;
+	class FMaterialBridgeReadSlotContext;
+
+} // UE::Ava
+
+namespace UE::Ava
+{
+
+/** Options to how read access should take place */
+struct FMaterialBridgeReadSlotOptions
+{
+	UE_API FMaterialBridgeReadSlotOptions();
+
+	/** The registry to use to visit any materials of sub objects. Defaults to the global registry. */
+	TNotNull<const FMaterialBridgeRegistry*> MaterialBridgeRegistry;
+
+	/** Callback when a container is not yet complete and cannot access its materials. Only called when bTryWaitOnContainerCompletion is false. */
+	TDelegate<EControlFlow(const FMaterialBridgeReadSlotContext&)> OnContainerPendingCompletion;
+
+	/** Whether to try to wait for container completion to access its materials */
+	bool bTryWaitOnContainerCompletion = true;
+};
+
+/** Short-lived context information when accessing a material slot for read  */
+class FMaterialBridgeReadSlotContext
+{
+public:
+	FMaterialBridgeReadSlotContext() = default;
+
+	explicit UE_API FMaterialBridgeReadSlotContext(const UObject* InContainer);
+	explicit UE_API FMaterialBridgeReadSlotContext(const UObject* InContainer, TNotNull<const FMaterialBridgeReadSlotContext*> InParentContext);
+
+	/** Container of the visited material */
+	FConstDataView MaterialContainer;
+
+	/** Parent of this context for additional information (e.g. container of the container) */
+	const FMaterialBridgeReadSlotContext* ParentContext = nullptr;
+};
+
+} // UE::Ava
+
+#undef UE_API
diff --git a/Engine/Plugins/VirtualProduction/Avalanche/Source/AvalancheShapes/Private/AvaShapeComponentMaterialBridge.cpp b/Engine/Plugins/VirtualProduction/Avalanche/Source/AvalancheShapes/Private/AvaShapeComponentMaterialBridge.cpp
deleted file mode 100644
diff --git a/Engine/Plugins/VirtualProduction/Avalanche/Source/AvalancheShapes/Private/AvaShapeMaterialBridge.cpp b/Engine/Plugins/VirtualProduction/Avalanche/Source/AvalancheShapes/Private/AvaShapeMaterialBridge.cpp
new file mode 100644
--- /dev/null
+++ b/Engine/Plugins/VirtualProduction/Avalanche/Source/AvalancheShapes/Private/AvaShapeMaterialBridge.cpp
@@ -0,0 +1,30 @@
+// Copyright Epic Games, Inc. All Rights Reserved.
+
+#include "AvaShapeMaterialBridge.h"
+#include "DynamicMeshes/AvaShapeDynMeshBase.h"
+
+namespace UE::Ava
+{
+
+const UStruct* FShapeMaterialBridge::OnGetBridgedType() const
+{
+	return UAvaShapeDynamicMeshBase::StaticClass();
+}
+
+EControlFlow FShapeMaterialBridge::OnAccessSlots(const FReadSlotContext& InContext, TFunctionRef<EControlFlow(const FReadSlotContext&, const FSlot&)> InFunc, const FReadSlotOptions& InOptions) const
+{
+	const UAvaShapeDynamicMeshBase& Shape = InContext.MaterialContainer.Get<const UAvaShapeDynamicMeshBase>();
+
+	for (const TPair<int32, FAvaShapeMeshData>& Pair : Shape.GetMeshDatas())
+	{
+		const FSlot Slot(Pair.Value.GetMaterial());
+		if (InFunc(InContext, Slot) == EControlFlow::Break)
+		{
+			return EControlFlow::Break;
+		}
+	}
+
+	return EControlFlow::Continue;
+}
+
+} // UE::Ava
diff --git a/Engine/Plugins/VirtualProduction/Avalanche/Source/AvalancheShapes/Private/AvaShapesModule.cpp b/Engine/Plugins/VirtualProduction/Avalanche/Source/AvalancheShapes/Private/AvaShapesModule.cpp
--- a/Engine/Plugins/VirtualProduction/Avalanche/Source/AvalancheShapes/Private/AvaShapesModule.cpp
+++ b/Engine/Plugins/VirtualProduction/Avalanche/Source/AvalancheShapes/Private/AvaShapesModule.cpp
@@ -1,7 +1,7 @@
 // Copyright Epic Games, Inc. All Rights Reserved.
 
 #include "AvaShapesModule.h"
-#include "AvaShapeComponentMaterialBridge.h"
+#include "AvaShapeMaterialBridge.h"
 #include "MaterialBridge/AvaMaterialBridgeRegistry.h"
 #include "Modules/ModuleManager.h"
 
@@ -9,8 +9,8 @@ IMPLEMENT_MODULE(FAvaShapesModule, AvalancheShapes)
 
 void FAvaShapesModule::StartupModule()
 {
-	constexpr uint32 MaterialBridgePriority = 0;
+	constexpr uint32 DefaultPriority = 0;
 
 	UE::Ava::FMaterialBridgeRegistry& MaterialBridgeRegistry = UE::Ava::FMaterialBridgeRegistry::GetMutable();
-	MaterialBridgeRegistry.Register<UE::Ava::FShapeComponentMaterialBridge>(MaterialBridgePriority);
+	MaterialBridgeRegistry.Register<UE::Ava::FShapeMaterialBridge>(DefaultPriority);
 }
diff --git a/Engine/Plugins/VirtualProduction/Avalanche/Source/AvalancheShapes/Public/AvaShapeComponentMaterialBridge.h b/Engine/Plugins/VirtualProduction/Avalanche/Source/AvalancheShapes/Public/AvaShapeComponentMaterialBridge.h
deleted file mode 100644
diff --git a/Engine/Plugins/VirtualProduction/Avalanche/Source/AvalancheShapes/Public/AvaShapeMaterialBridge.h b/Engine/Plugins/VirtualProduction/Avalanche/Source/AvalancheShapes/Public/AvaShapeMaterialBridge.h
new file mode 100644
--- /dev/null
+++ b/Engine/Plugins/VirtualProduction/Avalanche/Source/AvalancheShapes/Public/AvaShapeMaterialBridge.h
@@ -0,0 +1,24 @@
+// Copyright Epic Games, Inc. All Rights Reserved.
+
+#pragma once
+
+#include "MaterialBridge/AvaMaterialBridge.h"
+
+#define UE_API AVALANCHESHAPES_API
+
+namespace UE::Ava
+{
+
+/** Material Bridge for Motion Design Shape Components */
+class FShapeMaterialBridge : public FMaterialBridge
+{
+protected:
+	//~ Begin FMaterialBridge
+	UE_API virtual const UStruct* OnGetBridgedType() const override;
+	UE_API virtual EControlFlow OnAccessSlots(const FReadSlotContext& InContext, TFunctionRef<EControlFlow(const FReadSlotContext&, const FSlot&)> InFunc, const FReadSlotOptions& InOptions) const override;
+	//~ End FMaterialBridge
+};
+
+} // UE::Ava
+
+#undef UE_API
diff --git a/Engine/Plugins/VirtualProduction/Avalanche/Source/AvalancheText/Private/AvaText3DComponentMaterialBridge.cpp b/Engine/Plugins/VirtualProduction/Avalanche/Source/AvalancheText/Private/AvaText3DComponentMaterialBridge.cpp
deleted file mode 100644
diff --git a/Engine/Plugins/VirtualProduction/Avalanche/Source/AvalancheText/Private/AvaTextModule.cpp b/Engine/Plugins/VirtualProduction/Avalanche/Source/AvalancheText/Private/AvaTextModule.cpp
--- a/Engine/Plugins/VirtualProduction/Avalanche/Source/AvalancheText/Private/AvaTextModule.cpp
+++ b/Engine/Plugins/VirtualProduction/Avalanche/Source/AvalancheText/Private/AvaTextModule.cpp
@@ -1,16 +1,16 @@
 // Copyright Epic Games, Inc. All Rights Reserved.
 
 #include "AvaTextModule.h"
-#include "AvaText3DComponentMaterialBridge.h"
 #include "MaterialBridge/AvaMaterialBridgeRegistry.h"
+#include "MaterialBridge/AvaText3DComponentMaterialBridge.h"
 #include "Modules/ModuleManager.h"
 
 IMPLEMENT_MODULE(FAvaTextModule, AvalancheText)
 
 void FAvaTextModule::StartupModule()
 {
-	constexpr uint32 MaterialBridgePriority = 0;
+	constexpr uint32 DefaultPriority = 0;
 
 	UE::Ava::FMaterialBridgeRegistry& MaterialBridgeRegistry = UE::Ava::FMaterialBridgeRegistry::GetMutable();
-	MaterialBridgeRegistry.Register<UE::Ava::FText3DComponentMaterialBridge>(MaterialBridgePriority);
+	MaterialBridgeRegistry.Register<UE::Ava::FText3DComponentMaterialBridge>(DefaultPriority);
 }
diff --git a/Engine/Plugins/VirtualProduction/Avalanche/Source/AvalancheText/Private/MaterialBridge/AvaText3DComponentMaterialBridge.cpp b/Engine/Plugins/VirtualProduction/Avalanche/Source/AvalancheText/Private/MaterialBridge/AvaText3DComponentMaterialBridge.cpp
new file mode 100644
--- /dev/null
+++ b/Engine/Plugins/VirtualProduction/Avalanche/Source/AvalancheText/Private/MaterialBridge/AvaText3DComponentMaterialBridge.cpp
@@ -0,0 +1,36 @@
+// Copyright Epic Games, Inc. All Rights Reserved.
+
+#include "MaterialBridge/AvaText3DComponentMaterialBridge.h"
+#include "AvaDataView.h"
+#include "Extensions/Text3DMaterialExtensionBase.h"
+#include "Text3DComponent.h"
+
+namespace UE::Ava
+{
+
+const UStruct* FText3DComponentMaterialBridge::OnGetBridgedType() const
+{
+	return UText3DComponent::StaticClass();
+}
+
+EControlFlow FText3DComponentMaterialBridge::OnAccessSlots(const FReadSlotContext& InContext, TFunctionRef<EControlFlow(const FReadSlotContext&, const FSlot&)> InFunc, const FReadSlotOptions& InOptions) const
+{
+	const UText3DComponent& Component = InContext.MaterialContainer.Get<const UText3DComponent>();
+
+	EControlFlow ControlFlow = EControlFlow::Continue;
+
+	if (const UText3DMaterialExtensionBase* MaterialExtension = Component.GetMaterialExtension())
+	{
+		MaterialExtension->ForEachMaterial(
+			[&InFunc, &InContext, &ControlFlow](const UE::Text3D::Material::FMaterialParameters&, UMaterialInterface* InMaterial)->bool
+			{
+				const FSlot Slot(InMaterial);
+				ControlFlow = InFunc(InContext, Slot);
+				return ControlFlow == EControlFlow::Continue;
+			});
+	}
+
+	return ControlFlow;
+}
+
+} // UE::Ava
diff --git a/Engine/Plugins/VirtualProduction/Avalanche/Source/AvalancheText/Public/AvaText3DComponentMaterialBridge.h b/Engine/Plugins/VirtualProduction/Avalanche/Source/AvalancheText/Public/AvaText3DComponentMaterialBridge.h
deleted file mode 100644
diff --git a/Engine/Plugins/VirtualProduction/Avalanche/Source/AvalancheText/Public/MaterialBridge/AvaText3DComponentMaterialBridge.h b/Engine/Plugins/VirtualProduction/Avalanche/Source/AvalancheText/Public/MaterialBridge/AvaText3DComponentMaterialBridge.h
new file mode 100644
--- /dev/null
+++ b/Engine/Plugins/VirtualProduction/Avalanche/Source/AvalancheText/Public/MaterialBridge/AvaText3DComponentMaterialBridge.h
@@ -0,0 +1,24 @@
+// Copyright Epic Games, Inc. All Rights Reserved.
+
+#pragma once
+
+#include "MaterialBridge/AvaMaterialBridge.h"
+
+#define UE_API AVALANCHETEXT_API
+
+namespace UE::Ava
+{
+
+/** Material Bridge for Text3D Components */
+class FText3DComponentMaterialBridge : public FMaterialBridge
+{
+protected:
+	//~ Begin FMaterialBridge
+	UE_API virtual const UStruct* OnGetBridgedType() const override;
+	UE_API virtual EControlFlow OnAccessSlots(const FReadSlotContext& InContext, TFunctionRef<EControlFlow(const FReadSlotContext&, const FSlot&)> InFunc, const FReadSlotOptions& InOptions) const override;
+	//~ End FMaterialBridge
+};
+
+} // UE::Ava
+
+#undef UE_API
