Motion Design: added access slot for write + apply/store state data for material bridges

#Synced-CL 48692301

diff --git a/Engine/Plugins/VirtualProduction/Avalanche/Source/AvalancheMaterial/Private/MaterialBridge/AvaMaterialBridge.cpp b/Engine/Plugins/VirtualProduction/Avalanche/Source/AvalancheMaterial/Private/MaterialBridge/AvaMaterialBridge.cpp
--- a/Engine/Plugins/VirtualProduction/Avalanche/Source/AvalancheMaterial/Private/MaterialBridge/AvaMaterialBridge.cpp
+++ b/Engine/Plugins/VirtualProduction/Avalanche/Source/AvalancheMaterial/Private/MaterialBridge/AvaMaterialBridge.cpp
@@ -3,6 +3,9 @@
 #include "MaterialBridge/AvaMaterialBridge.h"
 #include "MaterialBridge/AvaMaterialBridgeLog.h"
 #include "MaterialBridge/AvaMaterialBridgeRegistry.h"
+#include "MaterialBridge/AvaMaterialContainerState.h"
+#include "MaterialBridge/Context/AvaMaterialBridgeReadSlotContext.h"
+#include "MaterialBridge/Context/AvaMaterialBridgeWriteSlotContext.h"
 
 namespace UE::Ava
 {
@@ -10,6 +13,7 @@ namespace UE::Ava
 void FMaterialBridge::Initialize(uint32 InPriority)
 {
 	BridgedType = OnGetBridgedType();
+	ContainerStateType = OnGetContainerStateType();
 	Priority = InPriority;
 }
 
@@ -32,7 +36,7 @@ bool FMaterialBridge::IsMaterialContainerSupported(FConstDataView InMaterialCont
 	return false;
 }
 
-EControlFlow FMaterialBridge::AccessSlots(const FReadSlotContext& InContext, TFunctionRef<EControlFlow(const FReadSlotContext&, const FSlot&)> InFunc, const FReadSlotOptions& InOptions) const
+EControlFlow FMaterialBridge::AccessSlots(const FReadSlotContext& InContext, TFunctionRef<EControlFlow(const FReadSlotContext&, const FReadSlot&)> InFunc, const FReadSlotOptions& InOptions) const
 {
 	if (InContext.MaterialContainer.IsValidFor(BridgedType))
 	{
@@ -41,4 +45,66 @@ EControlFlow FMaterialBridge::AccessSlots(const FReadSlotContext& InContext, TFu
 	return EControlFlow::Continue;
 }
 
+EControlFlow FMaterialBridge::AccessSlots(const FWriteSlotContext& InContext, TFunctionRef<EControlFlow(const FWriteSlotContext&, FWriteSlot&)> InFunc, const FWriteSlotOptions& InOptions) const
+{
+	if (InContext.MaterialContainer.IsValidFor(BridgedType))
+	{
+		return OnAccessSlots(InContext, InFunc, InOptions);
+	}
+	return EControlFlow::Continue;
+}
+
+bool FMaterialBridge::CanCreateContainerState() const
+{
+	return !!ContainerStateType;
+}
+
+FInstancedStruct FMaterialBridge::CreateContainerState() const
+{
+	if (CanCreateContainerState())
+	{
+		return FInstancedStruct(ContainerStateType);
+	}
+	return FInstancedStruct();
+}
+
+bool FMaterialBridge::ApplyState(const FApplyStateContext& InContext, FConstStructView InContainerState, const FApplyStateOptions& InOptions) const
+{
+	if (InContainerState.IsValid() && InContainerState.GetScriptStruct() == ContainerStateType)
+	{
+		OnApplyState(InContext, InContainerState, InOptions);
+		return true;
+	}
+	return false;
+}
+
+bool FMaterialBridge::StoreState(const FStoreStateContext& InContext, TNotNull<FInstancedStruct*> InOutContainerState, const FStoreStateOptions& InOptions) const
+{
+	if (!InOutContainerState->IsValid() || InOutContainerState->GetScriptStruct() != ContainerStateType)
+	{
+		FInstancedStruct NewContainerState = CreateContainerState();
+		if (!NewContainerState.IsValid())
+		{
+			UE_LOG(LogAvaMaterialBridge, Verbose, TEXT("Material Container State (Type: '%s') for Material Bridge (Bridged Type: '%s') could not be created.")
+				, *GetNameSafe(ContainerStateType)
+				, *GetNameSafe(BridgedType));
+			return false;
+		}
+
+		*InOutContainerState = MoveTemp(NewContainerState);
+	}
+	// Call the overload where state data is passed in as a FStructView 
+	return StoreState(InContext, *InOutContainerState, InOptions);
+}
+
+bool FMaterialBridge::StoreState(const FStoreStateContext& InContext, FStructView InContainerState, const FStoreStateOptions& InOptions) const
+{
+	if (InContainerState.IsValid() && InContainerState.GetScriptStruct() == ContainerStateType)
+	{
+		OnStoreState(InContext, InContainerState, InOptions);
+		return true;
+	}
+	return false;
+}
+
 } // UE::Ava
diff --git a/Engine/Plugins/VirtualProduction/Avalanche/Source/AvalancheMaterial/Private/MaterialBridge/AvaMaterialBridgeReadSlot.cpp b/Engine/Plugins/VirtualProduction/Avalanche/Source/AvalancheMaterial/Private/MaterialBridge/AvaMaterialBridgeReadSlot.cpp
new file mode 100644
--- /dev/null
+++ b/Engine/Plugins/VirtualProduction/Avalanche/Source/AvalancheMaterial/Private/MaterialBridge/AvaMaterialBridgeReadSlot.cpp
@@ -0,0 +1,13 @@
+// Copyright Epic Games, Inc. All Rights Reserved.
+
+#include "MaterialBridge/AvaMaterialBridgeReadSlot.h"
+
+namespace UE::Ava
+{
+	
+FMaterialBridgeReadSlot::FMaterialBridgeReadSlot(UMaterialInterface* InMaterial)
+	: Material(InMaterial)
+{
+}
+
+} // UE::Ava
\ No newline at end of file
diff --git a/Engine/Plugins/VirtualProduction/Avalanche/Source/AvalancheMaterial/Private/MaterialBridge/AvaMaterialBridgeSlot.cpp b/Engine/Plugins/VirtualProduction/Avalanche/Source/AvalancheMaterial/Private/MaterialBridge/AvaMaterialBridgeSlot.cpp
deleted file mode 100644
diff --git a/Engine/Plugins/VirtualProduction/Avalanche/Source/AvalancheMaterial/Private/MaterialBridge/AvaMaterialBridgeWriteSlot.cpp b/Engine/Plugins/VirtualProduction/Avalanche/Source/AvalancheMaterial/Private/MaterialBridge/AvaMaterialBridgeWriteSlot.cpp
new file mode 100644
--- /dev/null
+++ b/Engine/Plugins/VirtualProduction/Avalanche/Source/AvalancheMaterial/Private/MaterialBridge/AvaMaterialBridgeWriteSlot.cpp
@@ -0,0 +1,35 @@
+// Copyright Epic Games, Inc. All Rights Reserved.
+
+#include "MaterialBridge/AvaMaterialBridgeWriteSlot.h"
+
+namespace UE::Ava
+{
+
+FMaterialBridgeWriteSlot::FMaterialBridgeWriteSlot(UMaterialInterface* InMaterial)
+	: Material(InMaterial)
+{
+}
+
+void FMaterialBridgeWriteSlot::SetMaterial(UMaterialInterface* InMaterial)
+{
+	if (InMaterial != Material)
+	{
+		SetMaterialInternal(InMaterial);
+		OnMaterialChanged();
+	}
+}
+
+void FMaterialBridgeWriteSlot::RequestFeatures(EMaterialBridgeRequestedFeatures InFeatures)
+{
+	if (InFeatures != EMaterialBridgeRequestedFeatures::None)
+	{
+		OnRequestFeatures(InFeatures);
+	}
+}
+
+void FMaterialBridgeWriteSlot::SetMaterialInternal(UMaterialInterface* InMaterial)
+{
+	Material = InMaterial;
+}
+
+} // UE::Ava
\ No newline at end of file
diff --git a/Engine/Plugins/VirtualProduction/Avalanche/Source/AvalancheMaterial/Private/MaterialBridge/Common/AvaActorMaterialBridge.cpp b/Engine/Plugins/VirtualProduction/Avalanche/Source/AvalancheMaterial/Private/MaterialBridge/Common/AvaActorMaterialBridge.cpp
--- a/Engine/Plugins/VirtualProduction/Avalanche/Source/AvalancheMaterial/Private/MaterialBridge/Common/AvaActorMaterialBridge.cpp
+++ b/Engine/Plugins/VirtualProduction/Avalanche/Source/AvalancheMaterial/Private/MaterialBridge/Common/AvaActorMaterialBridge.cpp
@@ -2,7 +2,12 @@
 
 #include "MaterialBridge/Common/AvaActorMaterialBridge.h"
 #include "GameFramework/Actor.h"
+#include "MaterialBridge/AvaMaterialBridgeLog.h"
 #include "MaterialBridge/AvaMaterialBridgeRegistry.h"
+#include "MaterialBridge/Context/AvaMaterialBridgeApplyStateContext.h"
+#include "MaterialBridge/Context/AvaMaterialBridgeReadSlotContext.h"
+#include "MaterialBridge/Context/AvaMaterialBridgeStoreStateContext.h"
+#include "MaterialBridge/Context/AvaMaterialBridgeWriteSlotContext.h"
 
 namespace UE::Ava::Private
 {
@@ -31,6 +36,41 @@ EControlFlow ForEachBridgedComponent(const AActor& InActor, TNotNull<const FMate
 
 } // UE::Ava::Private
 
+void FAvaActorMaterialContainerStateComponent::SetComponent(UActorComponent* InComponent)
+{
+	if (InComponent)
+	{
+		ComponentWeak = InComponent;
+		Path = InComponent->GetPathName(InComponent->GetOwner());
+	}
+	else
+	{
+		ComponentWeak.Reset();
+		Path.Reset();
+	}
+}
+
+UActorComponent* FAvaActorMaterialContainerStateComponent::ResolveComponent(const AActor* InActor) const
+{
+	UActorComponent* Component = ComponentWeak.Get();
+	if (Component || !InActor)
+	{
+		return Component;
+	}
+
+	if (Path.IsEmpty())
+	{
+		Component = InActor->GetRootComponent();
+	}
+	else
+	{
+		Component = FindObject<UActorComponent>(const_cast<AActor*>(InActor), *Path);	
+	}
+
+	ComponentWeak = Component;
+	return Component;
+}
+
 namespace UE::Ava
 {
 
@@ -39,7 +79,7 @@ const UStruct* FActorMaterialBridge::OnGetBridgedType() const
 	return AActor::StaticClass();
 }
 
-EControlFlow FActorMaterialBridge::OnAccessSlots(const FReadSlotContext& InContext, TFunctionRef<EControlFlow(const FReadSlotContext&, const FSlot&)> InFunc, const FReadSlotOptions& InOptions) const
+EControlFlow FActorMaterialBridge::OnAccessSlots(const FReadSlotContext& InContext, TFunctionRef<EControlFlow(const FReadSlotContext&, const FReadSlot&)> InFunc, const FReadSlotOptions& InOptions) const
 {
 	const AActor& Actor = InContext.MaterialContainer.Get<const AActor>();
 
@@ -51,4 +91,67 @@ EControlFlow FActorMaterialBridge::OnAccessSlots(const FReadSlotContext& InConte
 
 }
 
+EControlFlow FActorMaterialBridge::OnAccessSlots(const FWriteSlotContext& InContext, TFunctionRef<EControlFlow(const FWriteSlotContext&, FWriteSlot&)> InFunc, const FWriteSlotOptions& InOptions) const
+{
+	const AActor& Actor = InContext.MaterialContainer.Get<const AActor>();
+
+	return Private::ForEachBridgedComponent(Actor, InOptions.MaterialBridgeRegistry,
+		[&InContext, &InFunc, &InOptions](TNotNull<UActorComponent*> InComponent, TNotNull<const FMaterialBridge*> InMaterialBridge)->EControlFlow
+		{
+			return InMaterialBridge->AccessSlots(FWriteSlotContext(InComponent, &InContext), InFunc, InOptions);
+		});
+}
+
+TSubScriptStructOf<FAvaMaterialContainerState> FActorMaterialBridge::OnGetContainerStateType() const
+{
+	return FContainerState::StaticStruct();
+}
+
+void FActorMaterialBridge::OnApplyState(const FApplyStateContext& InContext, FConstStructView InContainerState, const FApplyStateOptions& InOptions) const
+{
+	const AActor& Actor = InContext.MaterialContainer.Get<const AActor>();
+	const FContainerState& ContainerState = InContainerState.Get<FContainerState>();
+
+	// Call ApplyState on existing component data that still point to valid component
+	for (const FAvaActorMaterialContainerStateComponent& ComponentData : ContainerState.Components)
+	{
+		UActorComponent* const Component = ComponentData.ResolveComponent(&Actor);
+		if (!Component)
+		{
+			continue;
+		}
+
+		const FMaterialBridge* const MaterialBridge = InOptions.MaterialBridgeRegistry->GetMaterialBridge(FConstDataView(Component));
+		if (!MaterialBridge)
+		{
+			continue;
+		}
+
+		FApplyStateContext ApplyContext(Component, &InContext);
+		MaterialBridge->ApplyState(ApplyContext, ComponentData.ContainerState, InOptions);
+	}
+}
+
+void FActorMaterialBridge::OnStoreState(const FStoreStateContext& InContext, FStructView InContainerState, const FStoreStateOptions& InOptions) const
+{
+	const AActor& Actor = InContext.MaterialContainer.Get<const AActor>();
+
+	FContainerState& ContainerState = InContainerState.Get<FContainerState>();
+	ContainerState.Components.Reset();
+
+	Private::ForEachBridgedComponent(Actor, InOptions.MaterialBridgeRegistry,
+		[&InContext, &InOptions, &ContainerState](TNotNull<UActorComponent*> InComponent, TNotNull<const FMaterialBridge*> InMaterialBridge)->EControlFlow
+		{
+			FAvaActorMaterialContainerStateComponent ComponentData;
+			ComponentData.SetComponent(InComponent);
+
+			FStoreStateContext StoreContext(InComponent, &InContext);
+			if (InMaterialBridge->StoreState(StoreContext, &ComponentData.ContainerState, InOptions))
+			{
+				ContainerState.Components.Add(MoveTemp(ComponentData));
+			}
+			return EControlFlow::Continue;
+		});
+}
+
 } // UE::Ava
diff --git a/Engine/Plugins/VirtualProduction/Avalanche/Source/AvalancheMaterial/Private/MaterialBridge/Common/AvaDecalComponentMaterialBridge.cpp b/Engine/Plugins/VirtualProduction/Avalanche/Source/AvalancheMaterial/Private/MaterialBridge/Common/AvaDecalComponentMaterialBridge.cpp
--- a/Engine/Plugins/VirtualProduction/Avalanche/Source/AvalancheMaterial/Private/MaterialBridge/Common/AvaDecalComponentMaterialBridge.cpp
+++ b/Engine/Plugins/VirtualProduction/Avalanche/Source/AvalancheMaterial/Private/MaterialBridge/Common/AvaDecalComponentMaterialBridge.cpp
@@ -3,6 +3,12 @@
 #include "MaterialBridge/Common/AvaDecalComponentMaterialBridge.h"
 #include "AvaDataView.h"
 #include "Components/DecalComponent.h"
+#include "MaterialBridge/AvaMaterialBridgeReadSlot.h"
+#include "MaterialBridge/AvaMaterialBridgeWriteSlot.h"
+#include "MaterialBridge/Context/AvaMaterialBridgeApplyStateContext.h"
+#include "MaterialBridge/Context/AvaMaterialBridgeReadSlotContext.h"
+#include "MaterialBridge/Context/AvaMaterialBridgeStoreStateContext.h"
+#include "MaterialBridge/Context/AvaMaterialBridgeWriteSlotContext.h"
 
 namespace UE::Ava
 {
@@ -12,12 +18,48 @@ const UStruct* FDecalComponentMaterialBridge::OnGetBridgedType() const
 	return UDecalComponent::StaticClass();
 }
 
-EControlFlow FDecalComponentMaterialBridge::OnAccessSlots(const FReadSlotContext& InContext, TFunctionRef<EControlFlow(const FReadSlotContext&, const FSlot&)> InFunc, const FReadSlotOptions& InOptions) const
+EControlFlow FDecalComponentMaterialBridge::OnAccessSlots(const FReadSlotContext& InContext, TFunctionRef<EControlFlow(const FReadSlotContext&, const FReadSlot&)> InFunc, const FReadSlotOptions& InOptions) const
 {
 	const UDecalComponent& Component = InContext.MaterialContainer.Get<const UDecalComponent>();
 
-	const FSlot Slot(Component.GetDecalMaterial());
+	const FReadSlot Slot(Component.GetDecalMaterial());
 	return InFunc(InContext, Slot);
 }
 
+EControlFlow FDecalComponentMaterialBridge::OnAccessSlots(const FWriteSlotContext& InContext, TFunctionRef<EControlFlow(const FWriteSlotContext&, FWriteSlot&)> InFunc, const FWriteSlotOptions& InOptions) const
+{
+	UDecalComponent& Component = InContext.MaterialContainer.GetMutable<UDecalComponent>();
+
+	FWriteSlot Slot(Component.GetDecalMaterial());
+
+	const EControlFlow ControlFlow = InFunc(InContext, Slot);
+
+	if (Component.GetDecalMaterial() != Slot.GetMaterial())
+	{
+		Component.SetDecalMaterial(Slot.GetMaterial());
+	}
+	return ControlFlow;
+}
+
+TSubScriptStructOf<FAvaMaterialContainerState> FDecalComponentMaterialBridge::OnGetContainerStateType() const
+{
+	return FContainerState::StaticStruct();
+}
+
+void FDecalComponentMaterialBridge::OnApplyState(const FApplyStateContext& InContext, FConstStructView InContainerState, const FApplyStateOptions& InOptions) const
+{
+	UDecalComponent& Component = InContext.MaterialContainer.GetMutable<UDecalComponent>();
+
+	const FContainerState& ContainerState = InContainerState.Get<FContainerState>();
+	Component.SetDecalMaterial(ContainerState.DecalMaterial);
+}
+
+void FDecalComponentMaterialBridge::OnStoreState(const FStoreStateContext& InContext, FStructView InContainerState, const FStoreStateOptions& InOptions) const
+{
+	const UDecalComponent& Component = InContext.MaterialContainer.Get<const UDecalComponent>();
+
+	FContainerState& ContainerState = InContainerState.Get<FContainerState>();
+	ContainerState.DecalMaterial = Component.GetDecalMaterial();
+}
+
 } // UE::Ava
diff --git a/Engine/Plugins/VirtualProduction/Avalanche/Source/AvalancheMaterial/Private/MaterialBridge/Common/AvaLevelMaterialBridge.cpp b/Engine/Plugins/VirtualProduction/Avalanche/Source/AvalancheMaterial/Private/MaterialBridge/Common/AvaLevelMaterialBridge.cpp
--- a/Engine/Plugins/VirtualProduction/Avalanche/Source/AvalancheMaterial/Private/MaterialBridge/Common/AvaLevelMaterialBridge.cpp
+++ b/Engine/Plugins/VirtualProduction/Avalanche/Source/AvalancheMaterial/Private/MaterialBridge/Common/AvaLevelMaterialBridge.cpp
@@ -4,6 +4,8 @@
 #include "Engine/Level.h"
 #include "GameFramework/Actor.h"
 #include "MaterialBridge/AvaMaterialBridgeRegistry.h"
+#include "MaterialBridge/Context/AvaMaterialBridgeReadSlotContext.h"
+#include "MaterialBridge/Context/AvaMaterialBridgeWriteSlotContext.h"
 
 namespace UE::Ava
 {
@@ -13,7 +15,7 @@ const UStruct* FLevelMaterialBridge::OnGetBridgedType() const
 	return ULevel::StaticClass();
 }
 
-EControlFlow FLevelMaterialBridge::OnAccessSlots(const FReadSlotContext& InContext, TFunctionRef<EControlFlow(const FReadSlotContext&, const FSlot&)> InFunc, const FReadSlotOptions& InOptions) const
+EControlFlow FLevelMaterialBridge::OnAccessSlots(const FReadSlotContext& InContext, TFunctionRef<EControlFlow(const FReadSlotContext&, const FReadSlot&)> InFunc, const FReadSlotOptions& InOptions) const
 {
 	const ULevel& Level = InContext.MaterialContainer.Get<const ULevel>();
 
@@ -32,4 +34,23 @@ EControlFlow FLevelMaterialBridge::OnAccessSlots(const FReadSlotContext& InConte
 	return EControlFlow::Continue;
 }
 
+EControlFlow FLevelMaterialBridge::OnAccessSlots(const FWriteSlotContext& InContext, TFunctionRef<EControlFlow(const FWriteSlotContext&, FWriteSlot&)> InFunc, const FWriteSlotOptions& InOptions) const
+{
+	const ULevel& Level = InContext.MaterialContainer.Get<const ULevel>();
+
+	for (AActor* Actor : Level.Actors)
+	{
+		if (const FMaterialBridge* MaterialBridge = InOptions.MaterialBridgeRegistry->GetMaterialBridge(FConstDataView(Actor)))
+		{
+			const EControlFlow ControlFlow = MaterialBridge->AccessSlots(FWriteSlotContext(Actor, &InContext), InFunc, InOptions);
+			if (ControlFlow == EControlFlow::Break)
+			{
+				return EControlFlow::Break;
+			}
+		}
+	}
+
+	return EControlFlow::Continue;
+}
+
 } // UE::Ava
diff --git a/Engine/Plugins/VirtualProduction/Avalanche/Source/AvalancheMaterial/Private/MaterialBridge/Common/AvaPrimitiveComponentMaterialBridge.cpp b/Engine/Plugins/VirtualProduction/Avalanche/Source/AvalancheMaterial/Private/MaterialBridge/Common/AvaPrimitiveComponentMaterialBridge.cpp
--- a/Engine/Plugins/VirtualProduction/Avalanche/Source/AvalancheMaterial/Private/MaterialBridge/Common/AvaPrimitiveComponentMaterialBridge.cpp
+++ b/Engine/Plugins/VirtualProduction/Avalanche/Source/AvalancheMaterial/Private/MaterialBridge/Common/AvaPrimitiveComponentMaterialBridge.cpp
@@ -1,8 +1,14 @@
 // Copyright Epic Games, Inc. All Rights Reserved.
 
 #include "MaterialBridge/Common/AvaPrimitiveComponentMaterialBridge.h"
-#include "AvaDataView.h"
 #include "Components/PrimitiveComponent.h"
+#include "MaterialBridge/AvaMaterialBridgeReadSlot.h"
+#include "MaterialBridge/AvaMaterialBridgeWriteSlot.h"
+#include "MaterialBridge/Context/AvaMaterialBridgeApplyStateContext.h"
+#include "MaterialBridge/Context/AvaMaterialBridgeReadSlotContext.h"
+#include "MaterialBridge/Context/AvaMaterialBridgeStoreStateContext.h"
+#include "MaterialBridge/Context/AvaMaterialBridgeWriteSlotContext.h"
+#include "Misc/EnumerateRange.h"
 
 namespace UE::Ava
 {
@@ -12,16 +18,13 @@ const UStruct* FPrimitiveComponentMaterialBridge::OnGetBridgedType() const
 	return UPrimitiveComponent::StaticClass();
 }
 
-EControlFlow FPrimitiveComponentMaterialBridge::OnAccessSlots(const FReadSlotContext& InContext, TFunctionRef<EControlFlow(const FReadSlotContext&, const FSlot&)> InFunc, const FReadSlotOptions& InOptions) const
+EControlFlow FPrimitiveComponentMaterialBridge::OnAccessSlots(const FReadSlotContext& InContext, TFunctionRef<EControlFlow(const FReadSlotContext&, const FReadSlot&)> InFunc, const FReadSlotOptions& InOptions) const
 {
 	const UPrimitiveComponent& Component = InContext.MaterialContainer.Get<const UPrimitiveComponent>();
 
-	const int32 MaterialCount = Component.GetNumMaterials();
-
-	for (int32 MaterialIndex = 0; MaterialIndex < MaterialCount; ++MaterialIndex)
+	for (int32 MaterialIndex = 0; MaterialIndex < Component.GetNumMaterials(); ++MaterialIndex)
 	{
-		const FSlot Slot(Component.GetMaterial(MaterialIndex));
-
+		const FReadSlot Slot(Component.GetMaterial(MaterialIndex));
 		if (InFunc(InContext, Slot) == EControlFlow::Break)
 		{
 			return EControlFlow::Break;
@@ -31,4 +34,55 @@ EControlFlow FPrimitiveComponentMaterialBridge::OnAccessSlots(const FReadSlotCon
 	return EControlFlow::Continue;
 }
 
+EControlFlow FPrimitiveComponentMaterialBridge::OnAccessSlots(const FWriteSlotContext& InContext, TFunctionRef<EControlFlow(const FWriteSlotContext&, FWriteSlot&)> InFunc, const FWriteSlotOptions& InOptions) const
+{
+	UPrimitiveComponent& Component = InContext.MaterialContainer.GetMutable<UPrimitiveComponent>();
+
+	for (int32 MaterialIndex = 0; MaterialIndex < Component.GetNumMaterials(); ++MaterialIndex)
+	{
+		FWriteSlot Slot(Component.GetMaterial(MaterialIndex));
+		const EControlFlow ControlFlow = InFunc(InContext, Slot);
+
+		if (Component.GetMaterial(MaterialIndex) != Slot.GetMaterial())
+		{
+			Component.SetMaterial(MaterialIndex, Slot.GetMaterial());
+		}
+		if (ControlFlow == EControlFlow::Break)
+		{
+			return EControlFlow::Break;
+		}
+	}
+	return EControlFlow::Continue;
+}
+
+TSubScriptStructOf<FAvaMaterialContainerState> FPrimitiveComponentMaterialBridge::OnGetContainerStateType() const
+{
+	return FContainerState::StaticStruct();
+}
+
+void FPrimitiveComponentMaterialBridge::OnApplyState(const FApplyStateContext& InContext, FConstStructView InContainerState, const FApplyStateOptions& InOptions) const
+{
+	UPrimitiveComponent& Component = InContext.MaterialContainer.GetMutable<UPrimitiveComponent>();
+
+	const FContainerState& ContainerState = InContainerState.Get<FContainerState>();
+
+	for (TEnumerateRef<const TObjectPtr<UMaterialInterface>> Material : EnumerateRange(ContainerState.Materials))
+	{
+		Component.SetMaterial(Material.GetIndex(), *Material);
+	}
+}
+
+void FPrimitiveComponentMaterialBridge::OnStoreState(const FStoreStateContext& InContext, FStructView InContainerState, const FStoreStateOptions& InOptions) const
+{
+	const UPrimitiveComponent& Component = InContext.MaterialContainer.Get<const UPrimitiveComponent>();
+	FContainerState& ContainerState = InContainerState.Get<FContainerState>();
+
+	ContainerState.Materials.SetNum(Component.GetNumMaterials());
+
+	for (TEnumerateRef<TObjectPtr<UMaterialInterface>> Material : EnumerateRange(ContainerState.Materials))
+	{
+		*Material = Component.GetMaterial(Material.GetIndex());
+	}
+}
+
 } // UE::Ava
diff --git a/Engine/Plugins/VirtualProduction/Avalanche/Source/AvalancheMaterial/Private/MaterialBridge/Common/AvaStaticMeshMaterialBridge.cpp b/Engine/Plugins/VirtualProduction/Avalanche/Source/AvalancheMaterial/Private/MaterialBridge/Common/AvaStaticMeshMaterialBridge.cpp
--- a/Engine/Plugins/VirtualProduction/Avalanche/Source/AvalancheMaterial/Private/MaterialBridge/Common/AvaStaticMeshMaterialBridge.cpp
+++ b/Engine/Plugins/VirtualProduction/Avalanche/Source/AvalancheMaterial/Private/MaterialBridge/Common/AvaStaticMeshMaterialBridge.cpp
@@ -1,8 +1,12 @@
 // Copyright Epic Games, Inc. All Rights Reserved.
 
 #include "MaterialBridge/Common/AvaStaticMeshMaterialBridge.h"
-#include "MaterialBridge/AvaMaterialBridgeLog.h"
 #include "Engine/StaticMesh.h"
+#include "MaterialBridge/AvaMaterialBridgeLog.h"
+#include "MaterialBridge/AvaMaterialBridgeReadSlot.h"
+#include "MaterialBridge/AvaMaterialBridgeWriteSlot.h"
+#include "MaterialBridge/Context/AvaMaterialBridgeReadSlotContext.h"
+#include "MaterialBridge/Context/AvaMaterialBridgeWriteSlotContext.h"
 
 namespace UE::Ava
 {
@@ -12,13 +16,14 @@ const UStruct* FStaticMeshMaterialBridge::OnGetBridgedType() const
 	return UStaticMesh::StaticClass();
 }
 
-EControlFlow FStaticMeshMaterialBridge::OnAccessSlots(const FReadSlotContext& InContext, TFunctionRef<EControlFlow(const FReadSlotContext&, const FSlot&)> InFunc, const FReadSlotOptions& InOptions) const
+EControlFlow FStaticMeshMaterialBridge::OnAccessSlots(const FReadSlotContext& InContext, TFunctionRef<EControlFlow(const FReadSlotContext&, const FReadSlot&)> InFunc, const FReadSlotOptions& InOptions) const
 {
 	const UStaticMesh& StaticMesh = InContext.MaterialContainer.Get<const UStaticMesh>();
 
-	// If static mesh is compiling and visitor does not want to wait on materials, bail out.
-	// NOTE: Only LockedProperties needs to be checked for EStaticMeshAsyncProperties::StaticMaterials, but there's currently no option to check for it directly.
-	if (StaticMesh.IsCompiling() && !InOptions.bTryWaitOnContainerCompletion)
+	// When static mesh is compiling GetStaticMaterials() waits until LockedProperties does not have EStaticMeshAsyncProperties::StaticMaterials set.
+	// If static mesh is compiling and caller does not want to wait on materials, bail out.
+	// NOTE: Ideally only EStaticMeshAsyncProperties::StaticMaterials needs to be checked in LockedProperties, but there's currently no option to check for it directly.
+	if (StaticMesh.IsCompiling() && InOptions.bTrySkipWaitOnContainerCompletion)
 	{
 		if (InOptions.OnContainerPendingCompletion.IsBound())
 		{
@@ -34,7 +39,7 @@ EControlFlow FStaticMeshMaterialBridge::OnAccessSlots(const FReadSlotContext& In
 	{
 		// Base Material
 		{
-			const FSlot Slot(StaticMaterial.MaterialInterface);
+			const FReadSlot Slot(StaticMaterial.MaterialInterface);
 			if (InFunc(InContext, Slot) == EControlFlow::Break)
 			{
 				return EControlFlow::Break;
@@ -43,7 +48,7 @@ EControlFlow FStaticMeshMaterialBridge::OnAccessSlots(const FReadSlotContext& In
 
 		// Overlay Material
 		{
-			const FSlot Slot(StaticMaterial.OverlayMaterialInterface);
+			const FReadSlot Slot(StaticMaterial.OverlayMaterialInterface);
 			if (InFunc(InContext, Slot) == EControlFlow::Break)
 			{
 				return EControlFlow::Break;
@@ -53,4 +58,55 @@ EControlFlow FStaticMeshMaterialBridge::OnAccessSlots(const FReadSlotContext& In
 	return EControlFlow::Continue;
 }
 
+EControlFlow FStaticMeshMaterialBridge::OnAccessSlots(const FWriteSlotContext& InContext, TFunctionRef<EControlFlow(const FWriteSlotContext&, FWriteSlot&)> InFunc, const FWriteSlotOptions& InOptions) const
+{
+	UStaticMesh& StaticMesh = InContext.MaterialContainer.GetMutable<UStaticMesh>();
+
+	// When static mesh is compiling GetStaticMaterials() waits until LockedProperties does not have EStaticMeshAsyncProperties::StaticMaterials set.
+	// If static mesh is compiling and caller does not want to wait on materials, bail out.
+	// NOTE: Ideally only EStaticMeshAsyncProperties::StaticMaterials needs to be checked in LockedProperties, but there's currently no option to check for it directly.
+	if (StaticMesh.IsCompiling() && InOptions.bTrySkipWaitOnContainerCompletion)
+	{
+		if (InOptions.OnContainerPendingCompletion.IsBound())
+		{
+			UE_LOG(LogAvaMaterialBridge, Verbose, TEXT("Material Bridge could not access material slots of static mesh '%s'. Mesh is currently compiling."), *StaticMesh.GetName());
+			return InOptions.OnContainerPendingCompletion.Execute(InContext);
+		}
+
+		UE_LOG(LogAvaMaterialBridge, Warning, TEXT("Material Bridge could not access material slots of static mesh '%s'. Mesh is currently compiling."), *StaticMesh.GetName());
+		return EControlFlow::Continue;
+	}
+
+	for (FStaticMaterial& StaticMaterial : StaticMesh.GetStaticMaterials())
+	{
+		// Base Material
+		{
+			FWriteSlot Slot(StaticMaterial.MaterialInterface);
+
+			const EControlFlow ControlFlow = InFunc(InContext, Slot);
+			StaticMaterial.MaterialInterface = Slot.GetMaterial();
+
+			if (ControlFlow == EControlFlow::Break)
+			{
+				return EControlFlow::Break;
+			}
+		}
+
+		// Overlay Material
+		{
+			FWriteSlot Slot(StaticMaterial.OverlayMaterialInterface);
+
+			const EControlFlow ControlFlow = InFunc(InContext, Slot);
+			StaticMaterial.OverlayMaterialInterface = Slot.GetMaterial();
+
+			if (ControlFlow == EControlFlow::Break)
+			{
+				return EControlFlow::Break;
+			}
+		}
+	}
+
+	return EControlFlow::Continue;
+}
+
 } // UE::Ava
diff --git a/Engine/Plugins/VirtualProduction/Avalanche/Source/AvalancheMaterial/Private/MaterialBridge/Context/AvaMaterialBridgeApplyStateContext.cpp b/Engine/Plugins/VirtualProduction/Avalanche/Source/AvalancheMaterial/Private/MaterialBridge/Context/AvaMaterialBridgeApplyStateContext.cpp
new file mode 100644
--- /dev/null
+++ b/Engine/Plugins/VirtualProduction/Avalanche/Source/AvalancheMaterial/Private/MaterialBridge/Context/AvaMaterialBridgeApplyStateContext.cpp
@@ -0,0 +1,38 @@
+// Copyright Epic Games, Inc. All Rights Reserved.
+
+#include "MaterialBridge/Context/AvaMaterialBridgeApplyStateContext.h"
+#include "MaterialBridge/AvaMaterialBridgeRegistry.h"
+
+namespace UE::Ava
+{
+
+FMaterialBridgeApplyStateOptions::FMaterialBridgeApplyStateOptions()
+	: MaterialBridgeRegistry(&FMaterialBridgeRegistry::Get())
+{
+}
+
+FMaterialBridgeApplyStateContext::FMaterialBridgeApplyStateContext(UObject* InMaterialContainer)
+	: MaterialContainer(InMaterialContainer)
+{
+}
+
+FMaterialBridgeApplyStateContext::FMaterialBridgeApplyStateContext(UObject* InMaterialContainer, TNotNull<const FMaterialBridgeApplyStateContext*> InParentContext)
+	: MaterialContainer(InMaterialContainer)
+	, ParentContext(InParentContext)
+{
+}
+
+UObject* FMaterialBridgeApplyStateContext::GetMaterialContainerObject() const
+{
+	if (UObject* MaterialContainerObject = MaterialContainer.GetMutablePtr<UObject>())
+	{
+		return MaterialContainerObject;
+	}
+	if (ParentContext)
+	{
+		return ParentContext->GetMaterialContainerObject();
+	}
+	return nullptr;
+}
+
+} // UE::Ava
diff --git a/Engine/Plugins/VirtualProduction/Avalanche/Source/AvalancheMaterial/Private/MaterialBridge/Context/AvaMaterialBridgeStoreStateContext.cpp b/Engine/Plugins/VirtualProduction/Avalanche/Source/AvalancheMaterial/Private/MaterialBridge/Context/AvaMaterialBridgeStoreStateContext.cpp
new file mode 100644
--- /dev/null
+++ b/Engine/Plugins/VirtualProduction/Avalanche/Source/AvalancheMaterial/Private/MaterialBridge/Context/AvaMaterialBridgeStoreStateContext.cpp
@@ -0,0 +1,39 @@
+// Copyright Epic Games, Inc. All Rights Reserved.
+
+#include "MaterialBridge/Context/AvaMaterialBridgeStoreStateContext.h"
+#include "MaterialBridge/AvaMaterialBridgeRegistry.h"
+#include "StructUtils/InstancedStruct.h"
+
+namespace UE::Ava
+{
+
+FMaterialBridgeStoreStateOptions::FMaterialBridgeStoreStateOptions()
+	: MaterialBridgeRegistry(&FMaterialBridgeRegistry::Get())
+{
+}
+
+FMaterialBridgeStoreStateContext::FMaterialBridgeStoreStateContext(const UObject* InMaterialContainer)
+	: MaterialContainer(InMaterialContainer)
+{
+}
+
+FMaterialBridgeStoreStateContext::FMaterialBridgeStoreStateContext(const UObject* InMaterialContainer, TNotNull<const FMaterialBridgeStoreStateContext*> InParentContext)
+	: MaterialContainer(InMaterialContainer)
+	, ParentContext(InParentContext)
+{
+}
+
+const UObject* FMaterialBridgeStoreStateContext::GetMaterialContainerObject() const
+{
+	if (const UObject* MaterialContainerObject = MaterialContainer.GetPtr<const UObject>())
+	{
+		return MaterialContainerObject;
+	}
+	if (ParentContext)
+	{
+		return ParentContext->GetMaterialContainerObject();
+	}
+	return nullptr;
+}
+
+} // UE::Ava
diff --git a/Engine/Plugins/VirtualProduction/Avalanche/Source/AvalancheMaterial/Private/MaterialBridge/Context/AvaMaterialBridgeWriteSlotContext.cpp b/Engine/Plugins/VirtualProduction/Avalanche/Source/AvalancheMaterial/Private/MaterialBridge/Context/AvaMaterialBridgeWriteSlotContext.cpp
new file mode 100644
--- /dev/null
+++ b/Engine/Plugins/VirtualProduction/Avalanche/Source/AvalancheMaterial/Private/MaterialBridge/Context/AvaMaterialBridgeWriteSlotContext.cpp
@@ -0,0 +1,43 @@
+// Copyright Epic Games, Inc. All Rights Reserved.
+
+#include "MaterialBridge/Context/AvaMaterialBridgeWriteSlotContext.h"
+#include "MaterialBridge/AvaMaterialBridgeRegistry.h"
+
+namespace UE::Ava
+{
+
+FMaterialBridgeWriteSlotOptions::FMaterialBridgeWriteSlotOptions()
+	: MaterialBridgeRegistry(&FMaterialBridgeRegistry::Get())
+{
+}
+
+FMaterialBridgeWriteSlotContext::FMaterialBridgeWriteSlotContext(FDataView InMaterialContainer)
+	: MaterialContainer(InMaterialContainer)
+{
+}
+
+FMaterialBridgeWriteSlotContext::FMaterialBridgeWriteSlotContext(UObject* InMaterialContainer)
+	: MaterialContainer(InMaterialContainer)
+{
+}
+
+FMaterialBridgeWriteSlotContext::FMaterialBridgeWriteSlotContext(UObject* InMaterialContainer, TNotNull<const FMaterialBridgeWriteSlotContext*> InParentContext)
+	: MaterialContainer(InMaterialContainer)
+	, ParentContext(InParentContext)
+{
+}
+
+UObject* FMaterialBridgeWriteSlotContext::GetMaterialContainerObject() const
+{
+	if (UObject* MaterialContainerObject = MaterialContainer.GetMutablePtr<UObject>())
+	{
+		return MaterialContainerObject;
+	}
+	if (ParentContext)
+	{
+		return ParentContext->GetMaterialContainerObject();
+	}
+	return nullptr;
+}
+
+} // UE::Ava
diff --git a/Engine/Plugins/VirtualProduction/Avalanche/Source/AvalancheMaterial/Private/MaterialCache/AvaMaterialCacheHelper.cpp b/Engine/Plugins/VirtualProduction/Avalanche/Source/AvalancheMaterial/Private/MaterialCache/AvaMaterialCacheHelper.cpp
--- a/Engine/Plugins/VirtualProduction/Avalanche/Source/AvalancheMaterial/Private/MaterialCache/AvaMaterialCacheHelper.cpp
+++ b/Engine/Plugins/VirtualProduction/Avalanche/Source/AvalancheMaterial/Private/MaterialCache/AvaMaterialCacheHelper.cpp
@@ -3,7 +3,9 @@
 #include "MaterialCache/AvaMaterialCacheHelper.h"
 #include "AvaDataView.h"
 #include "MaterialBridge/AvaMaterialBridge.h"
+#include "MaterialBridge/AvaMaterialBridgeReadSlot.h"
 #include "MaterialBridge/AvaMaterialBridgeRegistry.h"
+#include "MaterialBridge/Context/AvaMaterialBridgeReadSlotContext.h"
 #include "MaterialCache/AvaMaterialCacheLog.h"
 #include "MaterialCache/AvaMaterialCacheSettings.h"
 #include "MaterialShared.h"
@@ -552,7 +554,7 @@ bool FMaterialCacheHelper::GatherMaterialData(const UObject* InObject, FObjectDa
 
 	// Options to not wait for the container to have materials available and instead get a callback
 	FMaterialBridgeReadSlotOptions ReadSlotOptions;
-	ReadSlotOptions.bTryWaitOnContainerCompletion = false;
+	ReadSlotOptions.bTrySkipWaitOnContainerCompletion = true;
 	ReadSlotOptions.OnContainerPendingCompletion.BindLambda([&InObjectData](const FMaterialBridgeReadSlotContext& InContext)->EControlFlow
 		{
 			if (const UObject* Container = InContext.MaterialContainer.GetPtr<const UObject>())
@@ -568,7 +570,7 @@ bool FMaterialCacheHelper::GatherMaterialData(const UObject* InObject, FObjectDa
 
 	// Gather all materials for object
 	MaterialBridge->AccessSlots(FMaterialBridgeReadSlotContext(InObject),
-		[&AddIncompleteMaterial](const FMaterialBridgeReadSlotContext& InContext, const FMaterialBridgeSlot& InSlot)->EControlFlow
+		[&AddIncompleteMaterial](const FMaterialBridgeReadSlotContext& InContext, const FMaterialBridgeReadSlot& InSlot)->EControlFlow
 		{
 			UMaterialInterface* const SlotMaterial = InSlot.GetMaterial();
 			if (!SlotMaterial)
diff --git a/Engine/Plugins/VirtualProduction/Avalanche/Source/AvalancheMaterial/Public/MaterialBridge/AvaMaterialBridge.h b/Engine/Plugins/VirtualProduction/Avalanche/Source/AvalancheMaterial/Public/MaterialBridge/AvaMaterialBridge.h
--- a/Engine/Plugins/VirtualProduction/Avalanche/Source/AvalancheMaterial/Public/MaterialBridge/AvaMaterialBridge.h
+++ b/Engine/Plugins/VirtualProduction/Avalanche/Source/AvalancheMaterial/Public/MaterialBridge/AvaMaterialBridge.h
@@ -2,30 +2,71 @@
 
 #pragma once
 
+#include "AvaDataView.h"
 #include "AvaEnums.h"
-#include "AvaMaterialBridgeSlot.h"
 #include "Containers/ContainersFwd.h"
-#include "Context/AvaMaterialBridgeReadSlotContext.h"
-#include "StructUtils/StructUtilsTypes.h"
+#include "StructUtils/StructView.h"
+#include "Templates/SubScriptStructOf.h"
 #include "UObject/ObjectPtr.h"
 
 #define UE_API AVALANCHEMATERIAL_API
 
-class UMaterialInterface;
+struct FAvaMaterialContainerState;
 
 namespace UE::Ava
 {
+	class FMaterialBridgeReadSlot;
+	class FMaterialBridgeReadSlotContext;
+	class FMaterialBridgeWriteSlot;
+	class FMaterialBridgeWriteSlotContext;
+	struct FMaterialBridgeApplyStateContext;
+	struct FMaterialBridgeApplyStateOptions;
+	struct FMaterialBridgeReadSlotOptions;
+	struct FMaterialBridgeStoreStateContext;
+	struct FMaterialBridgeStoreStateOptions;
+	struct FMaterialBridgeWriteSlotOptions;
+}
 
-/** 
- * Base class that allows interacting with material containers from any type of data view (Actors, ActorComponents, etc). 
- * A Material Bridge allows the user to generically access a material container's material slots for read (see AccessSlots)
+namespace UE::Ava
+{
+/**
+ * Base class that allows interacting with material containers from any type (Actors, ActorComponents, etc.). 
+ * A Material Bridge allows:
+ *
+ * 1) Access material slots: 
+ *    - Access can either be to read or write. 
+ *    - Read access allows only getting the material
+ *    - Write access allows getting the material slot and replacing its material.
+ *
+ * 2) Apply and store material container state
+ *    - Prior to replacing materials in a slot, the state of a material container can be stored via StoreState.
+ *    - The Material Bridge can then revert the material container back to its original state via ApplyState with the data that was stored prior to change.
+ *    - The expected flow could look like the following:
+ *        1) Call StoreState. Caller owns the state data.
+ *        2) Call AccessSlots (with write context) and replace slot material with a new one.
+ *        3) When needing to revert back, call ApplyState with the state data saved from (1)
+ *
+ * Things to note:
+ * - The Material Bridge does not own material container state data.
+ *   It can only create and interact with it via Apply/Store.
+ *   The responsibility of the ownership lies with the caller.
+ *
+ * - ApplyState and StoreState is only concerned with the state of the material container, not the material itself.
+ *   For example the parameters of a Material Instance Dynamic or any other property within a material object are not considered.
  */
 class FMaterialBridge
 {
 public:
-	using FSlot              = FMaterialBridgeSlot;
+	using FReadSlot          = FMaterialBridgeReadSlot;
 	using FReadSlotOptions   = FMaterialBridgeReadSlotOptions;
 	using FReadSlotContext   = FMaterialBridgeReadSlotContext;
+	using FWriteSlot         = FMaterialBridgeWriteSlot;
+	using FWriteSlotOptions  = FMaterialBridgeWriteSlotOptions;
+	using FWriteSlotContext  = FMaterialBridgeWriteSlotContext;
+	using FApplyStateOptions = FMaterialBridgeApplyStateOptions;
+	using FApplyStateContext = FMaterialBridgeApplyStateContext;
+	using FStoreStateOptions = FMaterialBridgeStoreStateOptions;
+	using FStoreStateContext = FMaterialBridgeStoreStateContext;
 
 	virtual ~FMaterialBridge() = default;
 
@@ -52,7 +93,50 @@ public:
 	 * @param InOptions access options. see FMaterialBridgeReadSlotAccessOptions
 	 * @return the iteration result of AccessSlots.
 	 */
-	UE_API EControlFlow AccessSlots(const FReadSlotContext& InContext, TFunctionRef<EControlFlow(const FReadSlotContext&, const FSlot&)> InFunc, const FReadSlotOptions& InOptions) const;
+	UE_API EControlFlow AccessSlots(const FReadSlotContext& InContext, TFunctionRef<EControlFlow(const FReadSlotContext&, const FReadSlot&)> InFunc, const FReadSlotOptions& InOptions) const;
+
+	/**
+	 * Access the found material slots in the given context for write.
+	 * @param InContext the current context to access materials.
+	 * @param InFunc function to call for each material slot accessed.
+	 * @param InOptions access options. see FMaterialBridgeWriteSlotAccessOptions
+	 * @return the iteration result of AccessSlots.
+	 */
+	UE_API EControlFlow AccessSlots(const FWriteSlotContext& InContext, TFunctionRef<EControlFlow(const FWriteSlotContext&, FWriteSlot&)> InFunc, const FWriteSlotOptions& InOptions) const;
+
+	/** Returns whether this material bridge supports container states */
+	UE_API bool CanCreateContainerState() const;
+
+	/** Creates a material container state instance that can be used for this material bridge */
+	UE_API FInstancedStruct CreateContainerState() const;
+
+	/** 
+	 * Applies the material container state to the context object.
+	 * @param InContext context used when applying the container state. see FMaterialBridgeApplyStateContext
+	 * @param InContainerState const view of the container state. Required to match the container state type of this bridge.
+	 * @param InOptions additional options to how apply should take place. see FMaterialBridgeApplyStateOptions
+	 * @return true if the container state was valid and was processed
+	 */
+	UE_API bool ApplyState(const FApplyStateContext& InContext, FConstStructView InContainerState, const FApplyStateOptions& InOptions) const;
+
+	/** 
+	 * Stores the state of the context object in the material container state
+	 * This overload initializes the container state if it is invalid.
+	 * @param InContext context used when storing the container state. see FMaterialBridgeStoreStateContext
+	 * @param InOutContainerState the instanced struct where the container state will be stored
+	 * @param InOptions additional options to how store should take place. see FMaterialBridgeStoreStateOptions
+	 * @return true if the container state was valid (or initialized to valid) and was processed
+	 */
+	UE_API bool StoreState(const FStoreStateContext& InContext, TNotNull<FInstancedStruct*> InOutContainerState, const FStoreStateOptions& InOptions) const;
+
+	/** 
+	 * Stores the state of the context object in the material container state
+	 * @param InContext context used when storing the container state. see FMaterialBridgeStoreStateContext
+	 * @param InContainerState the container state to store to. Required to match the container state type of this bridge.
+	 * @param InOptions additional options to how store should take place. see FMaterialBridgeStoreStateOptions
+	 * @return true if the container state was valid and was processed
+	 */
+	UE_API bool StoreState(const FStoreStateContext& InContext, FStructView InContainerState, const FStoreStateOptions& InOptions) const;
 
 protected:
 	/** Retrieves the type that this bridge supports. Only called once in Initialize() */
@@ -65,15 +149,40 @@ protected:
 	}
 
 	/** Accesses the found material slots in the given context for read */
-	virtual EControlFlow OnAccessSlots(const FReadSlotContext& InContext, TFunctionRef<EControlFlow(const FReadSlotContext&, const FSlot&)> InFunc, const FReadSlotOptions& InOptions) const
+	virtual EControlFlow OnAccessSlots(const FReadSlotContext& InContext, TFunctionRef<EControlFlow(const FReadSlotContext&, const FReadSlot&)> InFunc, const FReadSlotOptions& InOptions) const
 	{
 		return EControlFlow::Continue;
 	}
 
+	/** Accesses the found material slots in the given context for write */
+	virtual EControlFlow OnAccessSlots(const FWriteSlotContext& InContext, TFunctionRef<EControlFlow(const FWriteSlotContext&, FWriteSlot&)> InFunc, const FWriteSlotOptions& InOptions) const
+	{
+		return EControlFlow::Continue;
+	}
+
+	/** Returns the material container state type required by this bridge, if any */
+	virtual TSubScriptStructOf<FAvaMaterialContainerState> OnGetContainerStateType() const
+	{
+		return TSubScriptStructOf<FAvaMaterialContainerState>();
+	}
+
+	/** Applies the container state to the context object */
+	virtual void OnApplyState(const FApplyStateContext& InContext, FConstStructView InContainerState, const FApplyStateOptions& InOptions) const
+	{
+	}
+
+	/** Stores the state of the context object in the container state */
+	virtual void OnStoreState(const FStoreStateContext& InContext, FStructView InContainerState, const FStoreStateOptions& InOptions) const
+	{
+	}
+
 private:
 	/** The type that this Material Bridge works on */
 	TObjectPtr<const UStruct> BridgedType;
 
+	/** The material container state type of this bridge, if any */
+	TSubScriptStructOf<FAvaMaterialContainerState> ContainerStateType;
+
 	/** The priority of this bridge in relation to other bridges of the same bridged type */
 	uint32 Priority = 0;
 };
diff --git a/Engine/Plugins/VirtualProduction/Avalanche/Source/AvalancheMaterial/Public/MaterialBridge/AvaMaterialBridgeEnums.h b/Engine/Plugins/VirtualProduction/Avalanche/Source/AvalancheMaterial/Public/MaterialBridge/AvaMaterialBridgeEnums.h
new file mode 100644
--- /dev/null
+++ b/Engine/Plugins/VirtualProduction/Avalanche/Source/AvalancheMaterial/Public/MaterialBridge/AvaMaterialBridgeEnums.h
@@ -0,0 +1,22 @@
+// Copyright Epic Games, Inc. All Rights Reserved.
+
+#pragma once
+
+#include "HAL/Platform.h"
+#include "Misc/EnumClassFlags.h"
+
+namespace UE::Ava
+{
+
+/** The requested features for the material container to apply */
+enum class EMaterialBridgeRequestedFeatures : uint8
+{
+	/** No requested features */
+	None,
+
+	/** Set the material container to support translucent materials */
+	Translucency,
+};
+ENUM_CLASS_FLAGS(EMaterialBridgeRequestedFeatures);
+
+} // UE::Ava
diff --git a/Engine/Plugins/VirtualProduction/Avalanche/Source/AvalancheMaterial/Public/MaterialBridge/AvaMaterialBridgeReadSlot.h b/Engine/Plugins/VirtualProduction/Avalanche/Source/AvalancheMaterial/Public/MaterialBridge/AvaMaterialBridgeReadSlot.h
new file mode 100644
--- /dev/null
+++ b/Engine/Plugins/VirtualProduction/Avalanche/Source/AvalancheMaterial/Public/MaterialBridge/AvaMaterialBridgeReadSlot.h
@@ -0,0 +1,33 @@
+// Copyright Epic Games, Inc. All Rights Reserved.
+
+#pragma once
+
+#include "UObject/ObjectPtr.h"
+
+#define UE_API AVALANCHEMATERIAL_API
+
+class UMaterialInterface;
+
+namespace UE::Ava
+{
+
+/** Short-lived representation of an accessed material slot for read */
+class FMaterialBridgeReadSlot
+{
+public:
+	UE_API explicit FMaterialBridgeReadSlot(UMaterialInterface* InMaterial);
+
+	/** Gets the material for this slot */
+	UMaterialInterface* GetMaterial() const
+	{
+		return Material;
+	}
+
+protected:
+	/** Material in this slot */
+	TObjectPtr<UMaterialInterface> Material;
+};
+
+} // UE::Ava
+
+#undef UE_API
diff --git a/Engine/Plugins/VirtualProduction/Avalanche/Source/AvalancheMaterial/Public/MaterialBridge/AvaMaterialBridgeSlot.h b/Engine/Plugins/VirtualProduction/Avalanche/Source/AvalancheMaterial/Public/MaterialBridge/AvaMaterialBridgeSlot.h
deleted file mode 100644
diff --git a/Engine/Plugins/VirtualProduction/Avalanche/Source/AvalancheMaterial/Public/MaterialBridge/AvaMaterialBridgeWriteSlot.h b/Engine/Plugins/VirtualProduction/Avalanche/Source/AvalancheMaterial/Public/MaterialBridge/AvaMaterialBridgeWriteSlot.h
new file mode 100644
--- /dev/null
+++ b/Engine/Plugins/VirtualProduction/Avalanche/Source/AvalancheMaterial/Public/MaterialBridge/AvaMaterialBridgeWriteSlot.h
@@ -0,0 +1,60 @@
+// Copyright Epic Games, Inc. All Rights Reserved.
+
+#pragma once
+
+#include "MaterialBridge/AvaMaterialBridgeEnums.h"
+#include "UObject/ObjectPtr.h"
+
+#define UE_API AVALANCHEMATERIAL_API
+
+class UMaterialInterface;
+
+namespace UE::Ava
+{
+
+/** Base struct for short-lived representation of an accessed material slot for write */
+class FMaterialBridgeWriteSlot
+{
+public:
+	UE_API explicit FMaterialBridgeWriteSlot(UMaterialInterface* InMaterial);
+
+	virtual ~FMaterialBridgeWriteSlot() = default;
+
+	/** Gets the material for this slot */
+	UMaterialInterface* GetMaterial() const
+	{
+		return Material;
+	}
+
+	/** Sets the new material for the slot */
+	UE_API void SetMaterial(UMaterialInterface* InMaterial);
+
+	/**
+	 * Requests features for this material slot
+	 * NOTE: this is not meant to modify any property of the material, 
+	 * but rather meant to modify the state of the material container to adapt to the requested features.
+	 */
+	UE_API void RequestFeatures(EMaterialBridgeRequestedFeatures InFeatures);
+
+protected:
+	/** Sets the material property directly without calling the change notification */
+	UE_API void SetMaterialInternal(UMaterialInterface* InMaterial);
+
+	/** Notifies when a new material has been set for the slot */
+	virtual void OnMaterialChanged()
+	{
+	}
+
+	/** Requests features for this material slot */
+	virtual void OnRequestFeatures(EMaterialBridgeRequestedFeatures InFeatures)
+	{
+	}
+
+private:
+	/** Material in this slot */
+	TObjectPtr<UMaterialInterface> Material;
+};
+
+} // UE::Ava
+
+#undef UE_API
diff --git a/Engine/Plugins/VirtualProduction/Avalanche/Source/AvalancheMaterial/Public/MaterialBridge/AvaMaterialContainerState.h b/Engine/Plugins/VirtualProduction/Avalanche/Source/AvalancheMaterial/Public/MaterialBridge/AvaMaterialContainerState.h
new file mode 100644
--- /dev/null
+++ b/Engine/Plugins/VirtualProduction/Avalanche/Source/AvalancheMaterial/Public/MaterialBridge/AvaMaterialContainerState.h
@@ -0,0 +1,24 @@
+// Copyright Epic Games, Inc. All Rights Reserved.
+
+#pragma once
+
+#include "AvaMaterialContainerState.generated.h"
+
+class UMaterialInterface;
+
+/** Base struct for all material container states */
+USTRUCT()
+struct FAvaMaterialContainerState
+{
+	GENERATED_BODY()
+};
+
+/** Material container state that holds the materials as an array-like api */
+USTRUCT()
+struct FAvaMaterialArrayState : public FAvaMaterialContainerState
+{
+	GENERATED_BODY()
+
+	UPROPERTY()
+	TArray<TObjectPtr<UMaterialInterface>> Materials;
+};
diff --git a/Engine/Plugins/VirtualProduction/Avalanche/Source/AvalancheMaterial/Public/MaterialBridge/Common/AvaActorMaterialBridge.h b/Engine/Plugins/VirtualProduction/Avalanche/Source/AvalancheMaterial/Public/MaterialBridge/Common/AvaActorMaterialBridge.h
--- a/Engine/Plugins/VirtualProduction/Avalanche/Source/AvalancheMaterial/Public/MaterialBridge/Common/AvaActorMaterialBridge.h
+++ b/Engine/Plugins/VirtualProduction/Avalanche/Source/AvalancheMaterial/Public/MaterialBridge/Common/AvaActorMaterialBridge.h
@@ -3,19 +3,65 @@
 #pragma once
 
 #include "MaterialBridge/AvaMaterialBridge.h"
+#include "MaterialBridge/AvaMaterialContainerState.h"
+#include "AvaActorMaterialBridge.generated.h"
 
 #define UE_API AVALANCHEMATERIAL_API
 
+class UActorComponent;
+
+/** Struct holding the path for each actor component and the container state data mapped to it  */
+USTRUCT()
+struct FAvaActorMaterialContainerStateComponent
+{
+	GENERATED_BODY()
+
+	/** Sets the main component of this material container state */
+	UE_API void SetComponent(UActorComponent* InComponent);
+
+	/** Tries to resolve component by first resolving the weak ptr, or else using the component path and the given actor */
+	UE_API UActorComponent* ResolveComponent(const AActor* InActor) const;
+
+	/** Path to the actor component */
+	UPROPERTY()
+	FString Path;
+
+	/** Resolved component */
+	mutable TWeakObjectPtr<UActorComponent> ComponentWeak;
+
+	/** The container state data for the component. The type varies as each component could have a different Material Bridge */
+	UPROPERTY()
+	FInstancedStruct ContainerState;
+};
+
+/** Material container state for actors */
+USTRUCT()
+struct FAvaActorMaterialContainerState : public FAvaMaterialContainerState
+{
+	GENERATED_BODY()
+
+	/** All the actor component data found for the actor mapped to this data */
+	UPROPERTY()
+	TArray<FAvaActorMaterialContainerStateComponent> Components;
+};
+
 namespace UE::Ava
 {
 
 /** Material Bridge for Actors */
 class FActorMaterialBridge : public FMaterialBridge
 {
+public:
+	using FContainerState = FAvaActorMaterialContainerState;
+
 protected:
 	//~ Begin FMaterialBridge
 	UE_API virtual const UStruct* OnGetBridgedType() const override;
-	UE_API virtual EControlFlow OnAccessSlots(const FReadSlotContext& InContext, TFunctionRef<EControlFlow(const FReadSlotContext&, const FSlot&)> InFunc, const FReadSlotOptions& InOptions) const override;
+	UE_API virtual EControlFlow OnAccessSlots(const FReadSlotContext& InContext, TFunctionRef<EControlFlow(const FReadSlotContext&, const FReadSlot&)> InFunc, const FReadSlotOptions& InOptions) const override;
+	UE_API virtual EControlFlow OnAccessSlots(const FWriteSlotContext& InContext, TFunctionRef<EControlFlow(const FWriteSlotContext&, FWriteSlot&)> InFunc, const FWriteSlotOptions& InOptions) const override;
+	UE_API virtual TSubScriptStructOf<FAvaMaterialContainerState> OnGetContainerStateType() const override;
+	UE_API virtual void OnApplyState(const FApplyStateContext& InContext, FConstStructView InContainerState, const FApplyStateOptions& InOptions) const override;
+	UE_API virtual void OnStoreState(const FStoreStateContext& InContext, FStructView InContainerState, const FStoreStateOptions& InOptions) const override;
 	//~ End FMaterialBridge
 };
 
diff --git a/Engine/Plugins/VirtualProduction/Avalanche/Source/AvalancheMaterial/Public/MaterialBridge/Common/AvaDecalComponentMaterialBridge.h b/Engine/Plugins/VirtualProduction/Avalanche/Source/AvalancheMaterial/Public/MaterialBridge/Common/AvaDecalComponentMaterialBridge.h
--- a/Engine/Plugins/VirtualProduction/Avalanche/Source/AvalancheMaterial/Public/MaterialBridge/Common/AvaDecalComponentMaterialBridge.h
+++ b/Engine/Plugins/VirtualProduction/Avalanche/Source/AvalancheMaterial/Public/MaterialBridge/Common/AvaDecalComponentMaterialBridge.h
@@ -3,19 +3,39 @@
 #pragma once
 
 #include "MaterialBridge/AvaMaterialBridge.h"
+#include "MaterialBridge/AvaMaterialContainerState.h"
+#include "AvaDecalComponentMaterialBridge.generated.h"
 
 #define UE_API AVALANCHEMATERIAL_API
 
+/** Material container state for decal components */
+USTRUCT()
+struct FAvaDecalComponentMaterialContainerState : public FAvaMaterialContainerState
+{
+	GENERATED_BODY()
+
+	/** The decal material of the component */
+	UPROPERTY()
+	TObjectPtr<UMaterialInterface> DecalMaterial;
+};
+
 namespace UE::Ava
 {
 
 /** Material Bridge for Decal Components */
 class FDecalComponentMaterialBridge : public FMaterialBridge
 {
+public:
+	using FContainerState = FAvaDecalComponentMaterialContainerState;
+
 protected:
 	//~ Begin FMaterialBridge
 	UE_API virtual const UStruct* OnGetBridgedType() const override;
-	UE_API virtual EControlFlow OnAccessSlots(const FReadSlotContext& InContext, TFunctionRef<EControlFlow(const FReadSlotContext&, const FSlot&)> InFunc, const FReadSlotOptions& InOptions) const override;
+	UE_API virtual EControlFlow OnAccessSlots(const FReadSlotContext& InContext, TFunctionRef<EControlFlow(const FReadSlotContext&, const FReadSlot&)> InFunc, const FReadSlotOptions& InOptions) const override;
+	UE_API virtual EControlFlow OnAccessSlots(const FWriteSlotContext& InContext, TFunctionRef<EControlFlow(const FWriteSlotContext&, FWriteSlot&)> InFunc, const FWriteSlotOptions& InOptions) const override;
+	UE_API virtual TSubScriptStructOf<FAvaMaterialContainerState> OnGetContainerStateType() const override;
+	UE_API virtual void OnApplyState(const FApplyStateContext& InContext, FConstStructView InContainerState, const FApplyStateOptions& InOptions) const override;
+	UE_API virtual void OnStoreState(const FStoreStateContext& InContext, FStructView InContainerState, const FStoreStateOptions& InOptions) const override;
 	//~ End FMaterialBridge
 };
 
diff --git a/Engine/Plugins/VirtualProduction/Avalanche/Source/AvalancheMaterial/Public/MaterialBridge/Common/AvaLevelMaterialBridge.h b/Engine/Plugins/VirtualProduction/Avalanche/Source/AvalancheMaterial/Public/MaterialBridge/Common/AvaLevelMaterialBridge.h
--- a/Engine/Plugins/VirtualProduction/Avalanche/Source/AvalancheMaterial/Public/MaterialBridge/Common/AvaLevelMaterialBridge.h
+++ b/Engine/Plugins/VirtualProduction/Avalanche/Source/AvalancheMaterial/Public/MaterialBridge/Common/AvaLevelMaterialBridge.h
@@ -9,13 +9,19 @@
 namespace UE::Ava
 {
 
-/** Material Bridge for Levels */
+/**
+ * Material Bridge for Levels
+ *
+ * NOTE: Does not support Material Container State as it's probably better for users to do it per actor on-demand than to store/instantiate an entire level's material data.
+ * If it's ever needed, functionality to do this incrementally would be needed so as not to stall on one frame while creating large amounts of data.
+ */
 class FLevelMaterialBridge : public FMaterialBridge
 {
 protected:
 	//~ Begin FMaterialBridge
 	UE_API virtual const UStruct* OnGetBridgedType() const override;
-	UE_API virtual EControlFlow OnAccessSlots(const FReadSlotContext& InContext, TFunctionRef<EControlFlow(const FReadSlotContext&, const FSlot&)> InFunc, const FReadSlotOptions& InOptions) const override;
+	UE_API virtual EControlFlow OnAccessSlots(const FReadSlotContext& InContext, TFunctionRef<EControlFlow(const FReadSlotContext&, const FReadSlot&)> InFunc, const FReadSlotOptions& InOptions) const override;
+	UE_API virtual EControlFlow OnAccessSlots(const FWriteSlotContext& InContext, TFunctionRef<EControlFlow(const FWriteSlotContext&, FWriteSlot&)> InFunc, const FWriteSlotOptions& InOptions) const override;
 	//~ End FMaterialBridge
 };
 
diff --git a/Engine/Plugins/VirtualProduction/Avalanche/Source/AvalancheMaterial/Public/MaterialBridge/Common/AvaPrimitiveComponentMaterialBridge.h b/Engine/Plugins/VirtualProduction/Avalanche/Source/AvalancheMaterial/Public/MaterialBridge/Common/AvaPrimitiveComponentMaterialBridge.h
--- a/Engine/Plugins/VirtualProduction/Avalanche/Source/AvalancheMaterial/Public/MaterialBridge/Common/AvaPrimitiveComponentMaterialBridge.h
+++ b/Engine/Plugins/VirtualProduction/Avalanche/Source/AvalancheMaterial/Public/MaterialBridge/Common/AvaPrimitiveComponentMaterialBridge.h
@@ -3,18 +3,38 @@
 #pragma once
 
 #include "MaterialBridge/AvaMaterialBridge.h"
+#include "MaterialBridge/AvaMaterialContainerState.h"
+#include "AvaPrimitiveComponentMaterialBridge.generated.h"
 
 #define UE_API AVALANCHEMATERIAL_API
 
+/** Material container state for primitive components */
+USTRUCT()
+struct FAvaPrimitiveComponentMaterialContainerState : public FAvaMaterialContainerState
+{
+	GENERATED_BODY()
+
+	/** Materials found via the UPrimitiveComponent material api */
+	UPROPERTY()
+	TArray<TObjectPtr<UMaterialInterface>> Materials;
+};
+
 namespace UE::Ava
 {
 /** Material Bridge for Primitive Components */
 class FPrimitiveComponentMaterialBridge : public FMaterialBridge
 {
+public:
+	using FContainerState = FAvaPrimitiveComponentMaterialContainerState;
+
 protected:
 	//~ Begin FMaterialBridge
 	UE_API virtual const UStruct* OnGetBridgedType() const override;
-	UE_API virtual EControlFlow OnAccessSlots(const FReadSlotContext& InContext, TFunctionRef<EControlFlow(const FReadSlotContext&, const FSlot&)> InFunc, const FReadSlotOptions& InOptions) const override;
+	UE_API virtual EControlFlow OnAccessSlots(const FReadSlotContext& InContext, TFunctionRef<EControlFlow(const FReadSlotContext&, const FReadSlot&)> InFunc, const FReadSlotOptions& InOptions) const override;
+	UE_API virtual EControlFlow OnAccessSlots(const FWriteSlotContext& InContext, TFunctionRef<EControlFlow(const FWriteSlotContext&, FWriteSlot&)> InFunc, const FWriteSlotOptions& InOptions) const override;
+	UE_API virtual TSubScriptStructOf<FAvaMaterialContainerState> OnGetContainerStateType() const override;
+	UE_API virtual void OnApplyState(const FApplyStateContext& InContext, FConstStructView InContainerState, const FApplyStateOptions& InOptions) const override;
+	UE_API virtual void OnStoreState(const FStoreStateContext& InContext, FStructView InContainerState, const FStoreStateOptions& InOptions) const override;
 	//~ End FMaterialBridge
 };
 
diff --git a/Engine/Plugins/VirtualProduction/Avalanche/Source/AvalancheMaterial/Public/MaterialBridge/Common/AvaStaticMeshMaterialBridge.h b/Engine/Plugins/VirtualProduction/Avalanche/Source/AvalancheMaterial/Public/MaterialBridge/Common/AvaStaticMeshMaterialBridge.h
--- a/Engine/Plugins/VirtualProduction/Avalanche/Source/AvalancheMaterial/Public/MaterialBridge/Common/AvaStaticMeshMaterialBridge.h
+++ b/Engine/Plugins/VirtualProduction/Avalanche/Source/AvalancheMaterial/Public/MaterialBridge/Common/AvaStaticMeshMaterialBridge.h
@@ -18,7 +18,8 @@ class FStaticMeshMaterialBridge : public FMaterialBridge
 protected:
 	//~ Begin FMaterialBridge
 	UE_API virtual const UStruct* OnGetBridgedType() const override;
-	UE_API virtual EControlFlow OnAccessSlots(const FReadSlotContext& InContext, TFunctionRef<EControlFlow(const FReadSlotContext&, const FSlot&)> InFunc, const FReadSlotOptions& InOptions) const override;
+	UE_API virtual EControlFlow OnAccessSlots(const FReadSlotContext& InContext, TFunctionRef<EControlFlow(const FReadSlotContext&, const FReadSlot&)> InFunc, const FReadSlotOptions& InOptions) const override;
+	UE_API virtual EControlFlow OnAccessSlots(const FWriteSlotContext& InContext, TFunctionRef<EControlFlow(const FWriteSlotContext&, FWriteSlot&)> InFunc, const FWriteSlotOptions& InOptions) const override;
 	//~ End FMaterialBridge
 };
 
diff --git a/Engine/Plugins/VirtualProduction/Avalanche/Source/AvalancheMaterial/Public/MaterialBridge/Context/AvaMaterialBridgeApplyStateContext.h b/Engine/Plugins/VirtualProduction/Avalanche/Source/AvalancheMaterial/Public/MaterialBridge/Context/AvaMaterialBridgeApplyStateContext.h
new file mode 100644
--- /dev/null
+++ b/Engine/Plugins/VirtualProduction/Avalanche/Source/AvalancheMaterial/Public/MaterialBridge/Context/AvaMaterialBridgeApplyStateContext.h
@@ -0,0 +1,50 @@
+// Copyright Epic Games, Inc. All Rights Reserved.
+
+#pragma once
+
+#include "AvaDataView.h"
+
+#define UE_API AVALANCHEMATERIAL_API
+
+namespace UE::Ava
+{
+	class FMaterialBridgeRegistry;
+	struct FMaterialBridgeApplyStateContext;
+}
+
+namespace UE::Ava
+{
+
+/** Additional options when applying state to a material container  */
+struct FMaterialBridgeApplyStateOptions
+{
+	UE_API FMaterialBridgeApplyStateOptions();
+
+	/** The registry used to get material bridges. Defaults to the global registry. */
+	TNotNull<const FMaterialBridgeRegistry*> MaterialBridgeRegistry;
+};
+
+/** Short-lived context information when applying state to a material container */
+struct FMaterialBridgeApplyStateContext
+{
+	FMaterialBridgeApplyStateContext() = default;
+
+	explicit UE_API FMaterialBridgeApplyStateContext(UObject* InMaterialContainer);
+	explicit UE_API FMaterialBridgeApplyStateContext(UObject* InMaterialContainer, TNotNull<const FMaterialBridgeApplyStateContext*> InParentContext);
+
+	/** Gets the first material container that is a UObject, or null if the data is not nested inside a UObject */
+	UE_API UObject* GetMaterialContainerObject() const;
+
+	/** The object/struct containing the material */
+	FDataView MaterialContainer;
+
+	/** Parent context. Non-null if in a nested call. Should always be valid checked */
+	const FMaterialBridgeApplyStateContext* ParentContext = nullptr;
+
+private:
+	friend struct FMaterialBridgeApplyStateChangeScope;
+};
+
+} // UE::Ava
+
+#undef UE_API
diff --git a/Engine/Plugins/VirtualProduction/Avalanche/Source/AvalancheMaterial/Public/MaterialBridge/Context/AvaMaterialBridgeReadSlotContext.h b/Engine/Plugins/VirtualProduction/Avalanche/Source/AvalancheMaterial/Public/MaterialBridge/Context/AvaMaterialBridgeReadSlotContext.h
--- a/Engine/Plugins/VirtualProduction/Avalanche/Source/AvalancheMaterial/Public/MaterialBridge/Context/AvaMaterialBridgeReadSlotContext.h
+++ b/Engine/Plugins/VirtualProduction/Avalanche/Source/AvalancheMaterial/Public/MaterialBridge/Context/AvaMaterialBridgeReadSlotContext.h
@@ -25,11 +25,14 @@ struct FMaterialBridgeReadSlotOptions
 	/** The registry to use to visit any materials of sub objects. Defaults to the global registry. */
 	TNotNull<const FMaterialBridgeRegistry*> MaterialBridgeRegistry;
 
-	/** Callback when a container is not yet complete and cannot access its materials. Only called when bTryWaitOnContainerCompletion is false. */
+	/** Delegate called when a container is not yet complete and cannot access its materials. Only called when bTrySkipWaitOnContainerCompletion is true. */
 	TDelegate<EControlFlow(const FMaterialBridgeReadSlotContext&)> OnContainerPendingCompletion;
 
-	/** Whether to try to wait for container completion to access its materials */
-	bool bTryWaitOnContainerCompletion = true;
+	/**
+	 * Whether to try to skip wait for container completion to access its materials.
+	 * 'Try' because there is no guarantee that all material containers have an API that determines whether accessing a material will cause a wait.
+	 */
+	bool bTrySkipWaitOnContainerCompletion = false;
 };
 
 /** Short-lived context information when accessing a material slot for read  */
diff --git a/Engine/Plugins/VirtualProduction/Avalanche/Source/AvalancheMaterial/Public/MaterialBridge/Context/AvaMaterialBridgeStoreStateContext.h b/Engine/Plugins/VirtualProduction/Avalanche/Source/AvalancheMaterial/Public/MaterialBridge/Context/AvaMaterialBridgeStoreStateContext.h
new file mode 100644
--- /dev/null
+++ b/Engine/Plugins/VirtualProduction/Avalanche/Source/AvalancheMaterial/Public/MaterialBridge/Context/AvaMaterialBridgeStoreStateContext.h
@@ -0,0 +1,47 @@
+// Copyright Epic Games, Inc. All Rights Reserved.
+
+#pragma once
+
+#include "AvaDataView.h"
+
+#define UE_API AVALANCHEMATERIAL_API
+
+namespace UE::Ava
+{
+	class FMaterialBridgeRegistry;
+	struct FMaterialBridgeStoreStateContext;
+}
+
+namespace UE::Ava
+{
+
+/** Additional options when applying/storing material container state  */
+struct FMaterialBridgeStoreStateOptions
+{
+	UE_API FMaterialBridgeStoreStateOptions();
+
+	/** The registry used to get material bridges. Defaults to the global registry. */
+	TNotNull<const FMaterialBridgeRegistry*> MaterialBridgeRegistry;
+};
+
+/** Short-lived context information when storing the state of a material container */
+struct FMaterialBridgeStoreStateContext
+{
+	FMaterialBridgeStoreStateContext() = default;
+
+	explicit UE_API FMaterialBridgeStoreStateContext(const UObject* InMaterialContainer);
+	explicit UE_API FMaterialBridgeStoreStateContext(const UObject* InMaterialContainer, TNotNull<const FMaterialBridgeStoreStateContext*> InParentContext);
+
+	/** Gets the first material container that is a UObject */
+	UE_API const UObject* GetMaterialContainerObject() const;
+
+	/** The object containing the material */
+	FConstDataView MaterialContainer;
+
+	/** Parent context. Non-null if in a nested call. Should always be valid checked */
+	const FMaterialBridgeStoreStateContext* ParentContext = nullptr;
+};
+	
+} // UE::Ava
+
+#undef UE_API
diff --git a/Engine/Plugins/VirtualProduction/Avalanche/Source/AvalancheMaterial/Public/MaterialBridge/Context/AvaMaterialBridgeWriteSlotContext.h b/Engine/Plugins/VirtualProduction/Avalanche/Source/AvalancheMaterial/Public/MaterialBridge/Context/AvaMaterialBridgeWriteSlotContext.h
new file mode 100644
--- /dev/null
+++ b/Engine/Plugins/VirtualProduction/Avalanche/Source/AvalancheMaterial/Public/MaterialBridge/Context/AvaMaterialBridgeWriteSlotContext.h
@@ -0,0 +1,64 @@
+// Copyright Epic Games, Inc. All Rights Reserved.
+
+#pragma once
+
+#include "AvaEnums.h"
+#include "AvaDataView.h"
+
+#define UE_API AVALANCHEMATERIAL_API
+
+namespace UE::Ava
+{
+	class FMaterialBridgeRegistry;
+	class FMaterialBridgeWriteSlotContext;
+
+} // UE::Ava
+
+namespace UE::Ava
+{
+
+/** Options used when accessing a slot for write */
+struct FMaterialBridgeWriteSlotOptions
+{
+	UE_API FMaterialBridgeWriteSlotOptions();
+
+	/** The registry to use to visit any materials of sub objects. Defaults to the global registry. */
+	TNotNull<const FMaterialBridgeRegistry*> MaterialBridgeRegistry;
+
+	/** Delegate called when a container is not yet complete and cannot access its materials. Only called when bTrySkipWaitOnContainerCompletion is true. */
+	TDelegate<EControlFlow(const FMaterialBridgeWriteSlotContext&)> OnContainerPendingCompletion;
+
+	/**
+	 * Whether to try to skip wait for container completion to access its materials.
+	 * 'Try' because there is no guarantee that all material containers have an API that determines whether accessing a material will cause a wait.
+	 */
+	bool bTrySkipWaitOnContainerCompletion = false;
+};
+
+/** Short-lived context information when accessing a material slot for write  */
+class FMaterialBridgeWriteSlotContext
+{
+public:
+	FMaterialBridgeWriteSlotContext() = default;
+
+	explicit UE_API FMaterialBridgeWriteSlotContext(FDataView InMaterialContainer);
+	explicit UE_API FMaterialBridgeWriteSlotContext(UObject* InMaterialContainer);
+	explicit UE_API FMaterialBridgeWriteSlotContext(UObject* InMaterialContainer, TNotNull<const FMaterialBridgeWriteSlotContext*> InParentContext);
+
+	/** Gets the first UObject container in the context chain */
+	UE_API UObject* GetMaterialContainerObject() const;
+
+	/** Container of the accessed material slot */
+	FDataView MaterialContainer;
+
+	/** Parent context. Non-null if in a nested call. Should always be valid checked */
+	const FMaterialBridgeWriteSlotContext* ParentContext = nullptr;
+
+private:
+	/** Keeps track of whether pre change has been called for this context */
+	mutable bool bPreChangeCalled = false;
+};
+
+} // UE::Ava
+
+#undef UE_API
diff --git a/Engine/Plugins/VirtualProduction/Avalanche/Source/AvalancheShapes/Private/AvaShapeMaterialBridge.cpp b/Engine/Plugins/VirtualProduction/Avalanche/Source/AvalancheShapes/Private/AvaShapeMaterialBridge.cpp
--- a/Engine/Plugins/VirtualProduction/Avalanche/Source/AvalancheShapes/Private/AvaShapeMaterialBridge.cpp
+++ b/Engine/Plugins/VirtualProduction/Avalanche/Source/AvalancheShapes/Private/AvaShapeMaterialBridge.cpp
@@ -2,22 +2,66 @@
 
 #include "AvaShapeMaterialBridge.h"
 #include "DynamicMeshes/AvaShapeDynMeshBase.h"
+#include "MaterialBridge/AvaMaterialBridgeReadSlot.h"
+#include "MaterialBridge/Context/AvaMaterialBridgeApplyStateContext.h"
+#include "MaterialBridge/Context/AvaMaterialBridgeReadSlotContext.h"
+#include "MaterialBridge/Context/AvaMaterialBridgeStoreStateContext.h"
+#include "MaterialBridge/Context/AvaMaterialBridgeWriteSlotContext.h"
 
 namespace UE::Ava
 {
 
+namespace Private
+{
+
+TOptional<EAvaShapeParametricMaterialTranslucency> GetTranslucency(EMaterialBridgeRequestedFeatures InRequestedFeatures)
+{
+	if (EnumHasAnyFlags(InRequestedFeatures, EMaterialBridgeRequestedFeatures::Translucency))
+	{
+		return EAvaShapeParametricMaterialTranslucency::Enabled;
+	}
+	return TOptional<EAvaShapeParametricMaterialTranslucency>();
+}
+
+} // UE::Ava::Private
+
+FShapeMaterialWriteSlot::FShapeMaterialWriteSlot(TNotNull<UAvaShapeDynamicMeshBase*> InShape, int32 InMeshIndex)
+	: FMaterialBridgeWriteSlot(InShape->GetMaterial(InMeshIndex))
+	, Shape(InShape)
+	, MeshIndex(InMeshIndex)
+{
+}
+
+void FShapeMaterialWriteSlot::OnMaterialChanged()
+{
+	Shape->SetMaterial(MeshIndex, GetMaterial());
+}
+
+void FShapeMaterialWriteSlot::OnRequestFeatures(EMaterialBridgeRequestedFeatures InFeatures)
+{
+	if (EnumHasAnyFlags(InFeatures, EMaterialBridgeRequestedFeatures::Translucency))
+	{
+		FAvaShapeParametricMaterial* const ParametricMaterial = Shape->GetParametricMaterialPtr(MeshIndex);
+		if (ParametricMaterial && Shape->GetMaterialType(MeshIndex) == EMaterialType::Parametric)
+		{
+			ParametricMaterial->SetTranslucency(EAvaShapeParametricMaterialTranslucency::Enabled);
+			SetMaterialInternal(Shape->GetMaterial(MeshIndex));
+		}
+	}
+}
+
 const UStruct* FShapeMaterialBridge::OnGetBridgedType() const
 {
 	return UAvaShapeDynamicMeshBase::StaticClass();
 }
 
-EControlFlow FShapeMaterialBridge::OnAccessSlots(const FReadSlotContext& InContext, TFunctionRef<EControlFlow(const FReadSlotContext&, const FSlot&)> InFunc, const FReadSlotOptions& InOptions) const
+EControlFlow FShapeMaterialBridge::OnAccessSlots(const FReadSlotContext& InContext, TFunctionRef<EControlFlow(const FReadSlotContext&, const FReadSlot&)> InFunc, const FReadSlotOptions& InOptions) const
 {
 	const UAvaShapeDynamicMeshBase& Shape = InContext.MaterialContainer.Get<const UAvaShapeDynamicMeshBase>();
 
 	for (const TPair<int32, FAvaShapeMeshData>& Pair : Shape.GetMeshDatas())
 	{
-		const FSlot Slot(Pair.Value.GetMaterial());
+		const FReadSlot Slot(Pair.Value.GetMaterial());
 		if (InFunc(InContext, Slot) == EControlFlow::Break)
 		{
 			return EControlFlow::Break;
@@ -27,4 +71,80 @@ EControlFlow FShapeMaterialBridge::OnAccessSlots(const FReadSlotContext& InConte
 	return EControlFlow::Continue;
 }
 
+EControlFlow FShapeMaterialBridge::OnAccessSlots(const FWriteSlotContext& InContext, TFunctionRef<EControlFlow(const FWriteSlotContext&, FWriteSlot&)> InFunc, const FWriteSlotOptions& InOptions) const
+{
+	UAvaShapeDynamicMeshBase& Shape = InContext.MaterialContainer.GetMutable<UAvaShapeDynamicMeshBase>();
+
+	const TSet<int32> MeshIndices = Shape.GetMeshesIndexes();
+
+	for (int32 MeshIndex : MeshIndices)
+	{
+		if (FAvaShapeMeshData* MeshData = Shape.GetMeshData(MeshIndex))
+		{
+			FShapeMaterialWriteSlot Slot(&Shape, MeshIndex);
+			if (InFunc(InContext, Slot) == EControlFlow::Break)
+			{
+				return EControlFlow::Break;
+			}
+		}
+	}
+
+	return EControlFlow::Continue;
+}
+
+TSubScriptStructOf<FAvaMaterialContainerState> FShapeMaterialBridge::OnGetContainerStateType() const
+{
+	return FContainerState::StaticStruct();
+}
+
+void FShapeMaterialBridge::OnApplyState(const FApplyStateContext& InContext, FConstStructView InContainerState, const FApplyStateOptions& InOptions) const
+{
+	UAvaShapeDynamicMeshBase& Shape = InContext.MaterialContainer.GetMutable<UAvaShapeDynamicMeshBase>();
+	const FContainerState& ContainerState = InContainerState.Get<FContainerState>();
+
+	if (ContainerState.Meshes.IsEmpty())
+	{
+		return;
+	}
+
+	for (const FAvaShapeMaterialContainerStateMesh& Mesh : ContainerState.Meshes)
+	{
+		if (Mesh.MaterialType == EMaterialType::Parametric)
+		{
+			if (FAvaShapeParametricMaterial* ParametricMaterial = Shape.GetParametricMaterialPtr(Mesh.Index))
+			{
+				*ParametricMaterial = Mesh.ParametricMaterial;
+			}
+		}
+		Shape.SetMaterialType(Mesh.Index, Mesh.MaterialType);
+		Shape.SetMaterial(Mesh.Index, Mesh.Material);
+	}
+
+	Shape.MarkRenderStateDirty();
+}
+
+void FShapeMaterialBridge::OnStoreState(const FStoreStateContext& InContext, FStructView InContainerState, const FStoreStateOptions& InOptions) const
+{
+	const UAvaShapeDynamicMeshBase& Shape = InContext.MaterialContainer.Get<const UAvaShapeDynamicMeshBase>();
+	FContainerState& StateData = InContainerState.Get<FContainerState>();
+
+	StateData.Meshes.Reset();
+
+	for (int32 MeshIndex : Shape.GetMeshesIndexes())
+	{
+		if (const FAvaShapeMeshData* MeshData = Shape.GetMeshData(MeshIndex))
+		{
+			FAvaShapeMaterialContainerStateMesh& MeshState = StateData.Meshes.AddDefaulted_GetRef();
+			MeshState.Index = MeshIndex;
+			MeshState.Material = MeshData->GetMaterial();
+			MeshState.MaterialType = MeshData->GetMaterialType();
+
+			if (MeshState.MaterialType == EMaterialType::Parametric)
+			{
+				MeshState.ParametricMaterial = MeshData->GetParametricMaterial();	
+			}
+		}
+	}
+}
+
 } // UE::Ava
diff --git a/Engine/Plugins/VirtualProduction/Avalanche/Source/AvalancheShapes/Private/DynamicMeshes/AvaShapeDynMeshBase.cpp b/Engine/Plugins/VirtualProduction/Avalanche/Source/AvalancheShapes/Private/DynamicMeshes/AvaShapeDynMeshBase.cpp
--- a/Engine/Plugins/VirtualProduction/Avalanche/Source/AvalancheShapes/Private/DynamicMeshes/AvaShapeDynMeshBase.cpp
+++ b/Engine/Plugins/VirtualProduction/Avalanche/Source/AvalancheShapes/Private/DynamicMeshes/AvaShapeDynMeshBase.cpp
@@ -2175,18 +2175,14 @@ bool UAvaShapeDynamicMeshBase::RegisterMesh(FAvaShapeMeshData& NewMeshData)
 	return true;
 }
 
-const FAvaShapeMeshData* UAvaShapeDynamicMeshBase::GetMeshData(int32 MeshIndex) const
+const FAvaShapeMeshData* UAvaShapeDynamicMeshBase::GetMeshData(int32 InMeshIndex) const
 {
-	const FAvaShapeMeshData* Value = MeshDatas.Find(MeshIndex);
-
-	return Value;
+	return MeshDatas.Find(InMeshIndex);
 }
 
-FAvaShapeMeshData* UAvaShapeDynamicMeshBase::GetMeshData(int32 MeshIndex)
+FAvaShapeMeshData* UAvaShapeDynamicMeshBase::GetMeshData(int32 InMeshIndex)
 {
-	FAvaShapeMeshData* Value = MeshDatas.Find(MeshIndex);
-
-	return Value;
+	return MeshDatas.Find(InMeshIndex);
 }
 
 bool UAvaShapeDynamicMeshBase::IsMeshVisible(int32 MeshIndex)
diff --git a/Engine/Plugins/VirtualProduction/Avalanche/Source/AvalancheShapes/Public/AvaShapeMaterialBridge.h b/Engine/Plugins/VirtualProduction/Avalanche/Source/AvalancheShapes/Public/AvaShapeMaterialBridge.h
--- a/Engine/Plugins/VirtualProduction/Avalanche/Source/AvalancheShapes/Public/AvaShapeMaterialBridge.h
+++ b/Engine/Plugins/VirtualProduction/Avalanche/Source/AvalancheShapes/Public/AvaShapeMaterialBridge.h
@@ -2,20 +2,80 @@
 
 #pragma once
 
+#include "AvaShapeParametricMaterial.h"
+#include "AvaShapeUVParameters.h"
 #include "MaterialBridge/AvaMaterialBridge.h"
+#include "MaterialBridge/AvaMaterialBridgeWriteSlot.h"
+#include "MaterialBridge/AvaMaterialContainerState.h"
+#include "AvaShapeMaterialBridge.generated.h"
 
 #define UE_API AVALANCHESHAPES_API
 
+/** Shape mesh state data. Only the items available through FAvaShapeMeshData getters are available */
+USTRUCT()
+struct FAvaShapeMaterialContainerStateMesh
+{
+	GENERATED_BODY()
+
+	UPROPERTY()
+	int32 Index = INDEX_NONE;
+
+	UPROPERTY()
+	TObjectPtr<UMaterialInterface> Material;
+
+	UPROPERTY()
+	EMaterialType MaterialType;
+
+	UPROPERTY()
+	FAvaShapeParametricMaterial ParametricMaterial;
+};
+
+/** Material container state for UAvaShapeDynamicMeshBase components */
+USTRUCT()
+struct FAvaShapeMaterialContainerState : public FAvaMaterialContainerState
+{
+	GENERATED_BODY()
+
+	UPROPERTY()
+	TArray<FAvaShapeMaterialContainerStateMesh> Meshes;
+};
+
 namespace UE::Ava
 {
 
+/** Represents a slot in Motion Design Shape Components when accessing the slot with write access */
+class FShapeMaterialWriteSlot : public FMaterialBridgeWriteSlot
+{
+public:
+	UE_API explicit FShapeMaterialWriteSlot(TNotNull<UAvaShapeDynamicMeshBase*> InShape, int32 InMeshIndex);
+
+protected:
+	//~ Begin FMaterialBridgeSlot
+	UE_API virtual void OnMaterialChanged() override;
+	UE_API virtual void OnRequestFeatures(EMaterialBridgeRequestedFeatures InFeatures) override;
+	//~ End FMaterialBridgeSlot
+
+	/** The dynamic mesh shape component owning the slot */
+	TNotNull<TObjectPtr<UAvaShapeDynamicMeshBase>> Shape;
+
+	/** The mesh index of the slot */
+	int32 MeshIndex;
+};
+
 /** Material Bridge for Motion Design Shape Components */
 class FShapeMaterialBridge : public FMaterialBridge
 {
+public:
+	using FContainerState = FAvaShapeMaterialContainerState;
+
 protected:
 	//~ Begin FMaterialBridge
 	UE_API virtual const UStruct* OnGetBridgedType() const override;
-	UE_API virtual EControlFlow OnAccessSlots(const FReadSlotContext& InContext, TFunctionRef<EControlFlow(const FReadSlotContext&, const FSlot&)> InFunc, const FReadSlotOptions& InOptions) const override;
+	UE_API virtual EControlFlow OnAccessSlots(const FReadSlotContext& InContext, TFunctionRef<EControlFlow(const FReadSlotContext&, const FReadSlot&)> InFunc, const FReadSlotOptions& InOptions) const override;
+	UE_API virtual EControlFlow OnAccessSlots(const FWriteSlotContext& InContext, TFunctionRef<EControlFlow(const FWriteSlotContext&, FWriteSlot&)> InFunc, const FWriteSlotOptions& InOptions) const override;
+	UE_API virtual TSubScriptStructOf<FAvaMaterialContainerState> OnGetContainerStateType() const override;
+	UE_API virtual void OnApplyState(const FApplyStateContext& InContext, FConstStructView InContainerState, const FApplyStateOptions& InOptions) const override;
+	UE_API virtual void OnStoreState(const FStoreStateContext& InContext, FStructView InContainerState, const FStoreStateOptions& InOptions) const override;
 	//~ End FMaterialBridge
 };
 
diff --git a/Engine/Plugins/VirtualProduction/Avalanche/Source/AvalancheShapes/Public/AvaShapeMesh.h b/Engine/Plugins/VirtualProduction/Avalanche/Source/AvalancheShapes/Public/AvaShapeMesh.h
--- a/Engine/Plugins/VirtualProduction/Avalanche/Source/AvalancheShapes/Public/AvaShapeMesh.h
+++ b/Engine/Plugins/VirtualProduction/Avalanche/Source/AvalancheShapes/Public/AvaShapeMesh.h
@@ -116,11 +116,21 @@ struct FAvaShapeMeshData
 		return MeshName;
 	}
 
+	EMaterialType GetMaterialType() const
+	{
+		return MaterialType;
+	}
+
 	UMaterialInterface* GetMaterial() const
 	{
 		return Material;
 	}
 
+	const FAvaShapeParametricMaterial& GetParametricMaterial() const
+	{
+		return ParametricMaterial;
+	}
+
 protected:
 	// whether the mesh is currently visible or not and should be editable
 	UPROPERTY(EditDefaultsOnly, Category="Shape")
diff --git a/Engine/Plugins/VirtualProduction/Avalanche/Source/AvalancheShapes/Public/DynamicMeshes/AvaShapeDynMeshBase.h b/Engine/Plugins/VirtualProduction/Avalanche/Source/AvalancheShapes/Public/DynamicMeshes/AvaShapeDynMeshBase.h
--- a/Engine/Plugins/VirtualProduction/Avalanche/Source/AvalancheShapes/Public/DynamicMeshes/AvaShapeDynMeshBase.h
+++ b/Engine/Plugins/VirtualProduction/Avalanche/Source/AvalancheShapes/Public/DynamicMeshes/AvaShapeDynMeshBase.h
@@ -221,7 +221,10 @@ public:
 	AVALANCHESHAPES_API TSet<int32> GetMeshesIndexes() const;
 
 	/** find a registered mesh and gets a pointer to it */
-	FAvaShapeMeshData* GetMeshData(int32 MeshIndex);
+	FAvaShapeMeshData* GetMeshData(int32 InMeshIndex);
+
+	/** find a registered mesh and gets a pointer to it */
+	const FAvaShapeMeshData* GetMeshData(int32 InMeshIndex) const;
 
 	const TMap<int32, FAvaShapeMeshData>& GetMeshDatas() const
 	{
@@ -416,9 +419,6 @@ protected:
 	TArray<FAvaSnapPoint> LocalSnapPoints;
 
 private:
-	/** find a registered mesh and gets a const pointer to it */
-	const FAvaShapeMeshData* GetMeshData(int32 MeshIndex) const;
-
 	FAvaShapeMesh* GetMesh(int32 MeshIndex);
 
 	/** Update render state (materials) */
diff --git a/Engine/Plugins/VirtualProduction/Avalanche/Source/AvalancheText/Private/AvaTextModule.cpp b/Engine/Plugins/VirtualProduction/Avalanche/Source/AvalancheText/Private/AvaTextModule.cpp
--- a/Engine/Plugins/VirtualProduction/Avalanche/Source/AvalancheText/Private/AvaTextModule.cpp
+++ b/Engine/Plugins/VirtualProduction/Avalanche/Source/AvalancheText/Private/AvaTextModule.cpp
@@ -3,6 +3,9 @@
 #include "AvaTextModule.h"
 #include "MaterialBridge/AvaMaterialBridgeRegistry.h"
 #include "MaterialBridge/AvaText3DComponentMaterialBridge.h"
+#include "MaterialBridge/AvaText3DDefaultExtensionMaterialBridge.h"
+#include "MaterialBridge/AvaText3DExtensionMaterialBridge.h"
+#include "MaterialBridge/AvaText3DManagedComponentMaterialBridge.h"
 #include "Modules/ModuleManager.h"
 
 IMPLEMENT_MODULE(FAvaTextModule, AvalancheText)
@@ -13,4 +16,11 @@ void FAvaTextModule::StartupModule()
 
 	UE::Ava::FMaterialBridgeRegistry& MaterialBridgeRegistry = UE::Ava::FMaterialBridgeRegistry::GetMutable();
 	MaterialBridgeRegistry.Register<UE::Ava::FText3DComponentMaterialBridge>(DefaultPriority);
+	MaterialBridgeRegistry.Register<UE::Ava::FText3DExtensionMaterialBridge>(DefaultPriority);
+	MaterialBridgeRegistry.Register<UE::Ava::FText3DDefaultExtensionMaterialBridge>(DefaultPriority);
+
+	// Higher than default to be prioritized over registered default material bridges of the same bridged type.
+	constexpr uint32 HigherPriority = 100;
+
+	MaterialBridgeRegistry.Register<UE::Ava::FText3DManagedComponentMaterialBridge>(HigherPriority);
 }
diff --git a/Engine/Plugins/VirtualProduction/Avalanche/Source/AvalancheText/Private/MaterialBridge/AvaText3DComponentMaterialBridge.cpp b/Engine/Plugins/VirtualProduction/Avalanche/Source/AvalancheText/Private/MaterialBridge/AvaText3DComponentMaterialBridge.cpp
--- a/Engine/Plugins/VirtualProduction/Avalanche/Source/AvalancheText/Private/MaterialBridge/AvaText3DComponentMaterialBridge.cpp
+++ b/Engine/Plugins/VirtualProduction/Avalanche/Source/AvalancheText/Private/MaterialBridge/AvaText3DComponentMaterialBridge.cpp
@@ -3,6 +3,11 @@
 #include "MaterialBridge/AvaText3DComponentMaterialBridge.h"
 #include "AvaDataView.h"
 #include "Extensions/Text3DMaterialExtensionBase.h"
+#include "MaterialBridge/AvaMaterialBridgeRegistry.h"
+#include "MaterialBridge/Context/AvaMaterialBridgeApplyStateContext.h"
+#include "MaterialBridge/Context/AvaMaterialBridgeReadSlotContext.h"
+#include "MaterialBridge/Context/AvaMaterialBridgeStoreStateContext.h"
+#include "MaterialBridge/Context/AvaMaterialBridgeWriteSlotContext.h"
 #include "Text3DComponent.h"
 
 namespace UE::Ava
@@ -13,24 +18,69 @@ const UStruct* FText3DComponentMaterialBridge::OnGetBridgedType() const
 	return UText3DComponent::StaticClass();
 }
 
-EControlFlow FText3DComponentMaterialBridge::OnAccessSlots(const FReadSlotContext& InContext, TFunctionRef<EControlFlow(const FReadSlotContext&, const FSlot&)> InFunc, const FReadSlotOptions& InOptions) const
+EControlFlow FText3DComponentMaterialBridge::OnAccessSlots(const FReadSlotContext& InContext, TFunctionRef<EControlFlow(const FReadSlotContext&, const FReadSlot&)> InFunc, const FReadSlotOptions& InOptions) const
 {
 	const UText3DComponent& Component = InContext.MaterialContainer.Get<const UText3DComponent>();
 
-	EControlFlow ControlFlow = EControlFlow::Continue;
+	if (UText3DMaterialExtensionBase* MaterialExtension = Component.GetMaterialExtension())
+	{
+		if (const FMaterialBridge* MaterialBridge = InOptions.MaterialBridgeRegistry->GetMaterialBridge(FConstDataView(MaterialExtension)))
+		{
+			return MaterialBridge->AccessSlots(FReadSlotContext(MaterialExtension, &InContext), InFunc, InOptions);
+		}
+	}
+
+	return EControlFlow::Continue;
+}
+
+EControlFlow FText3DComponentMaterialBridge::OnAccessSlots(const FWriteSlotContext& InContext, TFunctionRef<EControlFlow(const FWriteSlotContext&, FWriteSlot&)> InFunc, const FWriteSlotOptions& InOptions) const
+{
+	const UText3DComponent& Component = InContext.MaterialContainer.Get<const UText3DComponent>();
 
-	if (const UText3DMaterialExtensionBase* MaterialExtension = Component.GetMaterialExtension())
+	if (UText3DMaterialExtensionBase* MaterialExtension = Component.GetMaterialExtension())
 	{
-		MaterialExtension->ForEachMaterial(
-			[&InFunc, &InContext, &ControlFlow](const UE::Text3D::Material::FMaterialParameters&, UMaterialInterface* InMaterial)->bool
-			{
-				const FSlot Slot(InMaterial);
-				ControlFlow = InFunc(InContext, Slot);
-				return ControlFlow == EControlFlow::Continue;
-			});
+		if (const FMaterialBridge* MaterialBridge = InOptions.MaterialBridgeRegistry->GetMaterialBridge(FConstDataView(MaterialExtension)))
+		{
+			return MaterialBridge->AccessSlots(FWriteSlotContext(MaterialExtension, &InContext), InFunc, InOptions);
+		}
 	}
 
-	return ControlFlow;
+	return EControlFlow::Continue;
+}
+
+TSubScriptStructOf<FAvaMaterialContainerState> FText3DComponentMaterialBridge::OnGetContainerStateType() const
+{
+	return FContainerState::StaticStruct();
+}
+
+void FText3DComponentMaterialBridge::OnApplyState(const FApplyStateContext& InContext, FConstStructView InContainerState, const FApplyStateOptions& InOptions) const
+{
+	UText3DComponent& Component = InContext.MaterialContainer.GetMutable<UText3DComponent>();
+	const FContainerState& ContainerState = InContainerState.Get<FContainerState>();
+
+	if (UText3DMaterialExtensionBase* MaterialExtension = Component.GetMaterialExtension())
+	{
+		if (const FMaterialBridge* MaterialBridge = InOptions.MaterialBridgeRegistry->GetMaterialBridge(FConstDataView(MaterialExtension)))
+		{
+			FApplyStateContext ApplyContext(MaterialExtension, &InContext);
+			MaterialBridge->ApplyState(ApplyContext, ContainerState.MaterialExtensionStateData, InOptions);
+		}
+	}
+}
+
+void FText3DComponentMaterialBridge::OnStoreState(const FStoreStateContext& InContext, FStructView InContainerState, const FStoreStateOptions& InOptions) const
+{
+	const UText3DComponent& Component = InContext.MaterialContainer.Get<const UText3DComponent>();
+	FContainerState& StateData = InContainerState.Get<FContainerState>();
+
+	if (UText3DMaterialExtensionBase* MaterialExtension = Component.GetMaterialExtension())
+	{
+		if (const FMaterialBridge* MaterialBridge = InOptions.MaterialBridgeRegistry->GetMaterialBridge(FConstDataView(MaterialExtension)))
+		{
+			FStoreStateContext StoreContext(MaterialExtension, &InContext);
+			MaterialBridge->StoreState(StoreContext, &StateData.MaterialExtensionStateData, InOptions);
+		}
+	}
 }
 
 } // UE::Ava
diff --git a/Engine/Plugins/VirtualProduction/Avalanche/Source/AvalancheText/Private/MaterialBridge/AvaText3DDefaultExtensionMaterialBridge.cpp b/Engine/Plugins/VirtualProduction/Avalanche/Source/AvalancheText/Private/MaterialBridge/AvaText3DDefaultExtensionMaterialBridge.cpp
new file mode 100644
--- /dev/null
+++ b/Engine/Plugins/VirtualProduction/Avalanche/Source/AvalancheText/Private/MaterialBridge/AvaText3DDefaultExtensionMaterialBridge.cpp
@@ -0,0 +1,101 @@
+// Copyright Epic Games, Inc. All Rights Reserved.
+
+#include "MaterialBridge/AvaText3DDefaultExtensionMaterialBridge.h"
+#include "Extensions/Text3DDefaultMaterialExtension.h"
+#include "MaterialBridge/Context/AvaMaterialBridgeApplyStateContext.h"
+#include "MaterialBridge/Context/AvaMaterialBridgeStoreStateContext.h"
+#include "MaterialBridge/Context/AvaMaterialBridgeWriteSlotContext.h"
+#include "Text3DComponent.h"
+
+namespace UE::Ava
+{
+
+namespace Private
+{
+// Sets the material extension's blend mode and updating Text3D Component immediately to update any material properties
+// returns true if the blend mode was changed, or false if the blend mode was already the requested one.
+bool SetText3DBlendMode(TNotNull<UText3DDefaultMaterialExtension*> InMaterialExtension, EText3DMaterialBlendMode InBlendMode)
+{
+	if (InMaterialExtension->GetBlendMode() != InBlendMode)
+	{
+		InMaterialExtension->SetBlendMode(InBlendMode);
+
+		// Request immediate material update so that the material changes due to blend mode can take effect.
+		if (UText3DComponent* Text3DComponent = InMaterialExtension->GetText3DComponent())
+		{
+			Text3DComponent->RequestUpdate(EText3DRendererFlags::Material, /*bInImmediate*/true);
+		}
+		return true;
+	}
+	return false;
+}
+	
+} // UE::Ava::Private
+
+void FText3DDefaultMaterialWriteSlot::OnRequestFeatures(EMaterialBridgeRequestedFeatures InFeatures)
+{
+	UText3DDefaultMaterialExtension* MaterialExtension = Cast<UText3DDefaultMaterialExtension>(BaseMaterialExtension);
+	if (!ensure(MaterialExtension))
+	{
+		return;
+	}
+
+	if (EnumHasAnyFlags(InFeatures, EMaterialBridgeRequestedFeatures::Translucency))
+	{
+		if (Private::SetText3DBlendMode(MaterialExtension, EText3DMaterialBlendMode::Translucent))
+		{
+			SetMaterialInternal(MaterialExtension->GetMaterial(MaterialParameters));
+		}
+	}
+}
+
+const UStruct* FText3DDefaultExtensionMaterialBridge::OnGetBridgedType() const
+{
+	return UText3DDefaultMaterialExtension::StaticClass();
+}
+
+TSubScriptStructOf<FAvaMaterialContainerState> FText3DDefaultExtensionMaterialBridge::OnGetContainerStateType() const
+{
+	return FContainerState::StaticStruct();
+}
+
+EControlFlow FText3DDefaultExtensionMaterialBridge::OnAccessSlots(const FWriteSlotContext& InContext, TFunctionRef<EControlFlow(const FWriteSlotContext&, FWriteSlot&)> InFunc, const FWriteSlotOptions& InOptions) const
+{
+	UText3DDefaultMaterialExtension& MaterialExtension = InContext.MaterialContainer.GetMutable<UText3DDefaultMaterialExtension>();
+
+	EControlFlow ControlFlow = EControlFlow::Continue;
+
+	MaterialExtension.ForEachMaterial(
+		[&InFunc, &InContext, &ControlFlow, &MaterialExtension](const UE::Text3D::Material::FMaterialParameters& InParameters, UMaterialInterface* InMaterial)->bool
+		{
+			FText3DDefaultMaterialWriteSlot Slot(&MaterialExtension, InParameters);
+			ControlFlow = InFunc(InContext, Slot);
+			return ControlFlow == EControlFlow::Continue;
+		});
+
+	return ControlFlow;
+}
+
+void FText3DDefaultExtensionMaterialBridge::OnApplyState(const FApplyStateContext& InContext, FConstStructView InContainerState, const FApplyStateOptions& InOptions) const
+{
+	UText3DDefaultMaterialExtension& MaterialExtension = InContext.MaterialContainer.GetMutable<UText3DDefaultMaterialExtension>();
+	const FContainerState& ContainerState = InContainerState.Get<FContainerState>();
+
+	Private::SetText3DBlendMode(&MaterialExtension, ContainerState.BlendMode);
+
+	// Call parent implementation last to apply the materials after the blend mode update
+	Super::OnApplyState(InContext, InContainerState, InOptions);
+}
+
+void FText3DDefaultExtensionMaterialBridge::OnStoreState(const FStoreStateContext& InContext, FStructView InContainerState, const FStoreStateOptions& InOptions) const
+{
+	// Call parent implementation to store the materials
+	Super::OnStoreState(InContext, InContainerState, InOptions);
+
+	const UText3DDefaultMaterialExtension& MaterialExtension = InContext.MaterialContainer.Get<const UText3DDefaultMaterialExtension>();
+
+	FContainerState& ContainerState = InContainerState.Get<FContainerState>();
+	ContainerState.BlendMode = MaterialExtension.GetBlendMode();
+}
+
+} // UE::Ava
diff --git a/Engine/Plugins/VirtualProduction/Avalanche/Source/AvalancheText/Private/MaterialBridge/AvaText3DExtensionMaterialBridge.cpp b/Engine/Plugins/VirtualProduction/Avalanche/Source/AvalancheText/Private/MaterialBridge/AvaText3DExtensionMaterialBridge.cpp
new file mode 100644
--- /dev/null
+++ b/Engine/Plugins/VirtualProduction/Avalanche/Source/AvalancheText/Private/MaterialBridge/AvaText3DExtensionMaterialBridge.cpp
@@ -0,0 +1,113 @@
+// Copyright Epic Games, Inc. All Rights Reserved.
+
+#include "MaterialBridge/AvaText3DExtensionMaterialBridge.h"
+#include "Extensions/Text3DMaterialExtensionBase.h"
+#include "MaterialBridge/AvaMaterialBridgeReadSlot.h"
+#include "MaterialBridge/Context/AvaMaterialBridgeApplyStateContext.h"
+#include "MaterialBridge/Context/AvaMaterialBridgeReadSlotContext.h"
+#include "MaterialBridge/Context/AvaMaterialBridgeStoreStateContext.h"
+#include "MaterialBridge/Context/AvaMaterialBridgeWriteSlotContext.h"
+#include "Text3DComponent.h"
+
+namespace UE::Ava
+{
+
+FText3DMaterialWriteSlot::FText3DMaterialWriteSlot(TNotNull<UText3DMaterialExtensionBase*> InMaterialExtension, const UE::Text3D::Material::FMaterialParameters& InMaterialParameters)
+	: FMaterialBridgeWriteSlot(InMaterialExtension->GetMaterial(InMaterialParameters))
+	, BaseMaterialExtension(InMaterialExtension)
+	, MaterialParameters(InMaterialParameters)
+{
+}
+
+void FText3DMaterialWriteSlot::OnMaterialChanged()
+{
+	BaseMaterialExtension->SetMaterial(MaterialParameters, GetMaterial());
+}
+
+const UStruct* FText3DExtensionMaterialBridge::OnGetBridgedType() const
+{
+	return UText3DMaterialExtensionBase::StaticClass();
+}
+
+EControlFlow FText3DExtensionMaterialBridge::OnAccessSlots(const FReadSlotContext& InContext, TFunctionRef<EControlFlow(const FReadSlotContext&, const FReadSlot&)> InFunc, const FReadSlotOptions& InOptions) const
+{
+	const UText3DMaterialExtensionBase& Extension = InContext.MaterialContainer.Get<const UText3DMaterialExtensionBase>();
+
+	EControlFlow ControlFlow = EControlFlow::Continue;
+
+	Extension.ForEachMaterial(
+		[&InFunc, &InContext, &ControlFlow](const UE::Text3D::Material::FMaterialParameters&, UMaterialInterface* InMaterial)->bool
+		{
+			const FReadSlot Slot(InMaterial);
+			ControlFlow = InFunc(InContext, Slot);
+			return ControlFlow == EControlFlow::Continue;
+		});
+
+	return ControlFlow;
+}
+
+EControlFlow FText3DExtensionMaterialBridge::OnAccessSlots(const FWriteSlotContext& InContext, TFunctionRef<EControlFlow(const FWriteSlotContext&, FWriteSlot&)> InFunc, const FWriteSlotOptions& InOptions) const
+{
+	UText3DMaterialExtensionBase& MaterialExtension = InContext.MaterialContainer.GetMutable<UText3DMaterialExtensionBase>();
+
+	EControlFlow ControlFlow = EControlFlow::Continue;
+
+	MaterialExtension.ForEachMaterial(
+		[&InFunc, &InContext, &ControlFlow, &MaterialExtension](const UE::Text3D::Material::FMaterialParameters& InParameters, UMaterialInterface* InMaterial)->bool
+		{
+			FText3DMaterialWriteSlot Slot(&MaterialExtension, InParameters);
+			ControlFlow = InFunc(InContext, Slot);
+			return ControlFlow == EControlFlow::Continue;
+		});
+
+	return ControlFlow;
+}
+
+TSubScriptStructOf<FAvaMaterialContainerState> FText3DExtensionMaterialBridge::OnGetContainerStateType() const
+{
+	return FContainerState::StaticStruct();
+}
+
+void FText3DExtensionMaterialBridge::OnApplyState(const FApplyStateContext& InContext, FConstStructView InContainerState, const FApplyStateOptions& InOptions) const
+{
+	UText3DMaterialExtensionBase& MaterialExtension = InContext.MaterialContainer.GetMutable<UText3DMaterialExtensionBase>();
+	const FContainerState& ContainerState = InContainerState.Get<const FContainerState>();
+
+	for (const FAvaText3DExtensionMaterialData& MaterialData : ContainerState.MaterialData)
+	{
+		UE::Text3D::Material::FMaterialParameters Parameters;
+		Parameters.Tag = MaterialData.Tag;
+		Parameters.Group = MaterialData.Group;
+
+		MaterialExtension.SetMaterial(Parameters, MaterialData.Material);
+	}
+}
+
+void FText3DExtensionMaterialBridge::OnStoreState(const FStoreStateContext& InContext, FStructView InContainerState, const FStoreStateOptions& InOptions) const
+{
+	const UText3DMaterialExtensionBase& MaterialExtension = InContext.MaterialContainer.Get<const UText3DMaterialExtensionBase>();
+	FContainerState& ContainerState = InContainerState.Get<FContainerState>();
+
+	// Ensure Text3D is up-to-date with the material before storing state
+	if (UText3DComponent* Text3DComponent = MaterialExtension.GetText3DComponent())
+	{
+		Text3DComponent->RequestUpdate(EText3DRendererFlags::Material, /*bInImmediate*/true);
+	}
+
+	const int32 MaterialCount = MaterialExtension.GetMaterialCount();
+	ContainerState.MaterialData.SetNum(MaterialCount);
+
+	int32 Index = 0;
+
+	MaterialExtension.ForEachMaterial(
+		[&ContainerState, &Index, &InContext, &InOptions](const UE::Text3D::Material::FMaterialParameters& InParameters, UMaterialInterface* InMaterial)->bool
+		{
+			FAvaText3DExtensionMaterialData& MaterialData = ContainerState.MaterialData[Index++];
+			MaterialData.Tag = InParameters.Tag;
+			MaterialData.Group = InParameters.Group;
+			MaterialData.Material = InMaterial;
+			return true; // continue;
+		});
+}
+
+} // UE::Ava
diff --git a/Engine/Plugins/VirtualProduction/Avalanche/Source/AvalancheText/Private/MaterialBridge/AvaText3DManagedComponentMaterialBridge.cpp b/Engine/Plugins/VirtualProduction/Avalanche/Source/AvalancheText/Private/MaterialBridge/AvaText3DManagedComponentMaterialBridge.cpp
new file mode 100644
--- /dev/null
+++ b/Engine/Plugins/VirtualProduction/Avalanche/Source/AvalancheText/Private/MaterialBridge/AvaText3DManagedComponentMaterialBridge.cpp
@@ -0,0 +1,24 @@
+// Copyright Epic Games, Inc. All Rights Reserved.
+
+#include "MaterialBridge/AvaText3DManagedComponentMaterialBridge.h"
+#include "Components/MeshComponent.h"
+#include "Renderers/Text3DRendererBase.h"
+
+namespace UE::Ava
+{
+
+const UStruct* FText3DManagedComponentMaterialBridge::OnGetBridgedType() const
+{
+	return UMeshComponent::StaticClass();
+}
+
+bool FText3DManagedComponentMaterialBridge::OnIsMaterialContainerSupported(FConstDataView InMaterialContainer) const
+{
+	const UMeshComponent& Component = InMaterialContainer.Get<const UMeshComponent>();
+
+	// NOTE: The mesh components get created with renderer extension as outer. 
+	// These get attached to the Text3D component directly, but there are other components that could potentially be attached to the Text3D component.
+	return !!Cast<UText3DRendererBase>(Component.GetOuter());
+}
+
+} // UE::Ava
diff --git a/Engine/Plugins/VirtualProduction/Avalanche/Source/AvalancheText/Public/MaterialBridge/AvaText3DComponentMaterialBridge.h b/Engine/Plugins/VirtualProduction/Avalanche/Source/AvalancheText/Public/MaterialBridge/AvaText3DComponentMaterialBridge.h
--- a/Engine/Plugins/VirtualProduction/Avalanche/Source/AvalancheText/Public/MaterialBridge/AvaText3DComponentMaterialBridge.h
+++ b/Engine/Plugins/VirtualProduction/Avalanche/Source/AvalancheText/Public/MaterialBridge/AvaText3DComponentMaterialBridge.h
@@ -3,19 +3,38 @@
 #pragma once
 
 #include "MaterialBridge/AvaMaterialBridge.h"
+#include "MaterialBridge/AvaMaterialContainerState.h"
+#include "AvaText3DComponentMaterialBridge.generated.h"
 
 #define UE_API AVALANCHETEXT_API
 
+/** Material state data for text3d components */
+USTRUCT()
+struct FAvaText3DComponentMaterialStateData : public FAvaMaterialContainerState
+{
+	GENERATED_BODY()
+
+	UPROPERTY()
+	FInstancedStruct MaterialExtensionStateData;
+};
+
 namespace UE::Ava
 {
 
 /** Material Bridge for Text3D Components */
 class FText3DComponentMaterialBridge : public FMaterialBridge
 {
+public:
+	using FContainerState = FAvaText3DComponentMaterialStateData;
+
 protected:
 	//~ Begin FMaterialBridge
 	UE_API virtual const UStruct* OnGetBridgedType() const override;
-	UE_API virtual EControlFlow OnAccessSlots(const FReadSlotContext& InContext, TFunctionRef<EControlFlow(const FReadSlotContext&, const FSlot&)> InFunc, const FReadSlotOptions& InOptions) const override;
+	UE_API virtual EControlFlow OnAccessSlots(const FReadSlotContext& InContext, TFunctionRef<EControlFlow(const FReadSlotContext&, const FReadSlot&)> InFunc, const FReadSlotOptions& InOptions) const override;
+	UE_API virtual EControlFlow OnAccessSlots(const FWriteSlotContext& InContext, TFunctionRef<EControlFlow(const FWriteSlotContext&, FWriteSlot&)> InFunc, const FWriteSlotOptions& InOptions) const override;
+	UE_API virtual TSubScriptStructOf<FAvaMaterialContainerState> OnGetContainerStateType() const override;
+	UE_API virtual void OnApplyState(const FApplyStateContext& InContext, FConstStructView InContainerState, const FApplyStateOptions& InOptions) const override;
+	UE_API virtual void OnStoreState(const FStoreStateContext& InContext, FStructView InContainerState, const FStoreStateOptions& InOptions) const override;
 	//~ End FMaterialBridge
 };
 
diff --git a/Engine/Plugins/VirtualProduction/Avalanche/Source/AvalancheText/Public/MaterialBridge/AvaText3DDefaultExtensionMaterialBridge.h b/Engine/Plugins/VirtualProduction/Avalanche/Source/AvalancheText/Public/MaterialBridge/AvaText3DDefaultExtensionMaterialBridge.h
new file mode 100644
--- /dev/null
+++ b/Engine/Plugins/VirtualProduction/Avalanche/Source/AvalancheText/Public/MaterialBridge/AvaText3DDefaultExtensionMaterialBridge.h
@@ -0,0 +1,54 @@
+// Copyright Epic Games, Inc. All Rights Reserved.
+
+#pragma once
+
+#include "AvaText3DExtensionMaterialBridge.h"
+#include "AvaText3DDefaultExtensionMaterialBridge.generated.h"
+
+#define UE_API AVALANCHETEXT_API
+
+/** Material state data struct for the Default Text3D Material Extension */
+USTRUCT()
+struct FAvaText3DDefaultExtensionMaterialContainerState : public FAvaText3DExtensionMaterialContainerState
+{
+	GENERATED_BODY()
+
+	UPROPERTY()
+	EText3DMaterialBlendMode BlendMode = EText3DMaterialBlendMode::Invalid;
+};
+
+namespace UE::Ava
+{
+
+/** Represents a slot for the Text3D default material extension for write */
+class FText3DDefaultMaterialWriteSlot : public FText3DMaterialWriteSlot
+{
+public:
+	using UE::Ava::FText3DMaterialWriteSlot::FText3DMaterialWriteSlot;
+
+protected:
+	//~ Begin FMaterialBridgeSlot
+	UE_API virtual void OnRequestFeatures(EMaterialBridgeRequestedFeatures InFeatures) override;
+	//~ End FMaterialBridgeSlot
+};
+
+/** Material Bridge for the Default Text3D Material Extension */
+class FText3DDefaultExtensionMaterialBridge : public FText3DExtensionMaterialBridge
+{
+public:
+	using Super = FText3DExtensionMaterialBridge;
+	using FContainerState = FAvaText3DDefaultExtensionMaterialContainerState;
+
+protected:
+	//~ Begin FMaterialBridge
+	UE_API virtual const UStruct* OnGetBridgedType() const override;
+	UE_API virtual TSubScriptStructOf<FAvaMaterialContainerState> OnGetContainerStateType() const override;
+	UE_API virtual EControlFlow OnAccessSlots(const FWriteSlotContext& InContext, TFunctionRef<EControlFlow(const FWriteSlotContext&, FWriteSlot&)> InFunc, const FWriteSlotOptions& InOptions) const override;
+	UE_API virtual void OnApplyState(const FApplyStateContext& InContext, FConstStructView InContainerState, const FApplyStateOptions& InOptions) const override;
+	UE_API virtual void OnStoreState(const FStoreStateContext& InContext, FStructView InContainerState, const FStoreStateOptions& InOptions) const override;
+	//~ End FMaterialBridge
+};
+
+} // UE::Ava
+
+#undef UE_API
diff --git a/Engine/Plugins/VirtualProduction/Avalanche/Source/AvalancheText/Public/MaterialBridge/AvaText3DExtensionMaterialBridge.h b/Engine/Plugins/VirtualProduction/Avalanche/Source/AvalancheText/Public/MaterialBridge/AvaText3DExtensionMaterialBridge.h
new file mode 100644
--- /dev/null
+++ b/Engine/Plugins/VirtualProduction/Avalanche/Source/AvalancheText/Public/MaterialBridge/AvaText3DExtensionMaterialBridge.h
@@ -0,0 +1,84 @@
+// Copyright Epic Games, Inc. All Rights Reserved.
+
+#pragma once
+
+#include "MaterialBridge/AvaMaterialBridge.h"
+#include "Text3DTypes.h"
+#include "MaterialBridge/AvaMaterialBridgeWriteSlot.h"
+#include "MaterialBridge/AvaMaterialContainerState.h"
+#include "AvaText3DExtensionMaterialBridge.generated.h"
+
+#define UE_API AVALANCHETEXT_API
+
+class UText3DMaterialExtensionBase;
+
+/** State data for a material in Text3D material extensions */
+USTRUCT()
+struct FAvaText3DExtensionMaterialData
+{
+	GENERATED_BODY()
+
+	/** Tag used to identify a material section, none = default */
+	UPROPERTY()
+	FName Tag = NAME_None;
+
+	/** Group used to identify an entry in the material section */
+	UPROPERTY()
+	EText3DGroupType Group = EText3DGroupType::Front;
+
+	UPROPERTY()
+	TObjectPtr<UMaterialInterface> Material;
+};
+
+/** Base material container state struct for Text3D material extensions */
+USTRUCT()
+struct FAvaText3DExtensionMaterialContainerState : public FAvaMaterialContainerState
+{
+	GENERATED_BODY()
+
+	UPROPERTY()
+	TArray<FAvaText3DExtensionMaterialData> MaterialData;
+};
+
+namespace UE::Ava
+{
+
+/** Represents a slot in Text3D material extension when accessing the slot with write access */
+class FText3DMaterialWriteSlot : public FMaterialBridgeWriteSlot
+{
+public:
+	UE_API explicit FText3DMaterialWriteSlot(TNotNull<UText3DMaterialExtensionBase*> InMaterialExtension, const UE::Text3D::Material::FMaterialParameters& InMaterialParameters);
+
+protected:
+	//~ Begin FMaterialBridgeSlot
+	UE_API virtual void OnMaterialChanged() override;
+	//~ End FMaterialBridgeSlot
+
+	/** The material extension that owns the slot */
+	TNotNull<TObjectPtr<UText3DMaterialExtensionBase>> BaseMaterialExtension;
+
+	/** The material parameters identifying the slot */
+	UE::Text3D::Material::FMaterialParameters MaterialParameters;
+};
+
+/** Base Material Bridge for Text3D material extensions */
+class FText3DExtensionMaterialBridge : public FMaterialBridge
+{
+public:
+	using FContainerState = FAvaText3DExtensionMaterialContainerState;
+
+protected:
+	//~ Begin FMaterialBridge
+	UE_API virtual const UStruct* OnGetBridgedType() const override;
+	UE_API virtual EControlFlow OnAccessSlots(const FReadSlotContext& InContext, TFunctionRef<EControlFlow(const FReadSlotContext&, const FReadSlot&)> InFunc, const FReadSlotOptions& InOptions) const override;
+	UE_API virtual EControlFlow OnAccessSlots(const FWriteSlotContext& InContext, TFunctionRef<EControlFlow(const FWriteSlotContext&, FWriteSlot&)> InFunc, const FWriteSlotOptions& InOptions) const override;
+	UE_API virtual TSubScriptStructOf<FAvaMaterialContainerState> OnGetContainerStateType() const override;
+	UE_API virtual void OnApplyState(const FApplyStateContext& InContext, FConstStructView InContainerState, const FApplyStateOptions& InOptions) const override;
+	UE_API virtual void OnStoreState(const FStoreStateContext& InContext, FStructView InContainerState, const FStoreStateOptions& InOptions) const override;
+	//~ End FMaterialBridge
+
+};
+
+} // UE::Ava
+
+#undef UE_API
diff --git a/Engine/Plugins/VirtualProduction/Avalanche/Source/AvalancheText/Public/MaterialBridge/AvaText3DManagedComponentMaterialBridge.h b/Engine/Plugins/VirtualProduction/Avalanche/Source/AvalancheText/Public/MaterialBridge/AvaText3DManagedComponentMaterialBridge.h
new file mode 100644
--- /dev/null
+++ b/Engine/Plugins/VirtualProduction/Avalanche/Source/AvalancheText/Public/MaterialBridge/AvaText3DManagedComponentMaterialBridge.h
@@ -0,0 +1,27 @@
+// Copyright Epic Games, Inc. All Rights Reserved.
+
+#pragma once
+
+#include "MaterialBridge/AvaMaterialBridge.h"
+
+#define UE_API AVALANCHETEXT_API
+
+namespace UE::Ava
+{
+/**
+ * Material Bridge for components that are managed by the Text3D 
+ * Slot accessing and material container state are not implemented by design as the material slot access for these is managed by the Text3D Component and its extensions. 
+ * The Text3DComponent bridge should be used instead.
+ */
+class FText3DManagedComponentMaterialBridge : public FMaterialBridge
+{
+protected:
+	//~ Begin FMaterialBridge
+	UE_API virtual const UStruct* OnGetBridgedType() const override;
+	UE_API virtual bool OnIsMaterialContainerSupported(FConstDataView InMaterialContainer) const override;
+	//~ End FMaterialBridge
+};
+
+} // UE::Ava
+
+#undef UE_API
