Motion Design: remove adding mask read component in mask read (output) modifier as it was not contributing to anything.

#Synced-CL 48692301

diff --git a/Engine/Plugins/VirtualProduction/Avalanche/Source/AvalancheMask/Private/Mask2D/AvaMask2DBaseModifier.cpp b/Engine/Plugins/VirtualProduction/Avalanche/Source/AvalancheMask/Private/Mask2D/AvaMask2DBaseModifier.cpp
--- a/Engine/Plugins/VirtualProduction/Avalanche/Source/AvalancheMask/Private/Mask2D/AvaMask2DBaseModifier.cpp
+++ b/Engine/Plugins/VirtualProduction/Avalanche/Source/AvalancheMask/Private/Mask2D/AvaMask2DBaseModifier.cpp
@@ -678,7 +678,6 @@ void UAvaMask2DBaseModifier::RemoveFromActor(AActor* InActor)
 	MaterialCollectionHandleData.Remove(InActor);
 	MaterialCollectionHandles.Remove(InActor);
 	
-	UE::AvaMask::Internal::RemoveComponentByInterface<UGeometryMaskReadInterface>(InActor);
 	UE::AvaMask::Internal::RemoveComponentByInterface<UGeometryMaskWriteInterface>(InActor);
 	if (UAvaGizmoComponent* FoundComponent = InActor->FindComponentByClass<UAvaGizmoComponent>())
 	{
diff --git a/Engine/Plugins/VirtualProduction/Avalanche/Source/AvalancheMask/Private/Mask2D/AvaMask2DReadModifier.cpp b/Engine/Plugins/VirtualProduction/Avalanche/Source/AvalancheMask/Private/Mask2D/AvaMask2DReadModifier.cpp
--- a/Engine/Plugins/VirtualProduction/Avalanche/Source/AvalancheMask/Private/Mask2D/AvaMask2DReadModifier.cpp
+++ b/Engine/Plugins/VirtualProduction/Avalanche/Source/AvalancheMask/Private/Mask2D/AvaMask2DReadModifier.cpp
@@ -8,7 +8,6 @@
 #include "Engine/Engine.h"
 #include "GeometryMaskCanvas.h"
 #include "GeometryMaskCanvasResource.h"
-#include "GeometryMaskReadComponent.h"
 #include "Handling/AvaHandleUtilities.h"
 #include "Handling/AvaObjectHandleSubsystem.h"
 #include "Handling/IAvaMaskMaterialCollectionHandle.h"
@@ -161,16 +160,6 @@ bool UAvaMask2DReadModifier::ApplyRead(AActor* InActor, FAvaMask2DActorData& InA
 			MaterialCollectionHandle->ApplyModifiedState(ApplyParameters, *MaterialCollectionData);
 		}
 
-		if (UGeometryMaskReadComponent* ReadComponent = FindOrAddMaskComponent<UGeometryMaskReadComponent>(InActor))
-		{
-			FGeometryMaskReadParameters Parameters = ReadComponent->GetParameters();
-			Parameters.CanvasName = Channel;
-			Parameters.ColorChannel = Canvas->GetColorChannel();
-			Parameters.bInvert = bInverted;
-
-			ReadComponent->SetParameters(Parameters);
-		}
-
 		// Only check for valid materials here, after ApplyModifiedState has occured once
 		if (!MaterialCollectionHandle->ValidateMaterials(OutFailReason))
 		{
@@ -224,7 +213,9 @@ void UAvaMask2DReadModifier::SetupMaskComponent(UActorComponent* InComponent)
 
 	if (IGeometryMaskReadInterface* MaskReader = Cast<IGeometryMaskReadInterface>(InComponent))
 	{
+PRAGMA_DISABLE_DEPRECATION_WARNINGS
 		SetupMaskReadComponent(MaskReader);
+PRAGMA_ENABLE_DEPRECATION_WARNINGS
 	}
 	
 	UAvaMask2DBaseModifier::SetupMaskComponent(InComponent);
diff --git a/Engine/Plugins/VirtualProduction/Avalanche/Source/AvalancheMask/Public/Mask2D/AvaMask2DReadModifier.h b/Engine/Plugins/VirtualProduction/Avalanche/Source/AvalancheMask/Public/Mask2D/AvaMask2DReadModifier.h
--- a/Engine/Plugins/VirtualProduction/Avalanche/Source/AvalancheMask/Public/Mask2D/AvaMask2DReadModifier.h
+++ b/Engine/Plugins/VirtualProduction/Avalanche/Source/AvalancheMask/Public/Mask2D/AvaMask2DReadModifier.h
@@ -62,7 +62,10 @@ protected:
 
 	void OnBaseOpacityChanged();
 
+	UE_DEPRECATED(5.8, "SetupMaskReadComponent is deprecated. Mask Read Modifier no longer allocates a component")
 	virtual void SetupMaskComponent(UActorComponent* InComponent) override;
+
+	UE_DEPRECATED(5.8, "SetupMaskReadComponent is deprecated. Mask Read Modifier no longer allocates a component")
 	void SetupMaskReadComponent(IGeometryMaskReadInterface* InMaskReader);
 
 	UE_DEPRECATED(5.7, "Use ApplyRead that takes in a fail reason")
