Motion Design: changed material bridge to use TConstStructView<T> and TStructView<T> instead of FConstStructView and FStructView for clarity that these views are all based on the base material container state struct.

#Synced-CL 48692301

diff --git a/Engine/Plugins/VirtualProduction/Avalanche/Source/AvalancheMaterial/Private/MaterialBridge/AvaMaterialBridge.cpp b/Engine/Plugins/VirtualProduction/Avalanche/Source/AvalancheMaterial/Private/MaterialBridge/AvaMaterialBridge.cpp
--- a/Engine/Plugins/VirtualProduction/Avalanche/Source/AvalancheMaterial/Private/MaterialBridge/AvaMaterialBridge.cpp
+++ b/Engine/Plugins/VirtualProduction/Avalanche/Source/AvalancheMaterial/Private/MaterialBridge/AvaMaterialBridge.cpp
@@ -59,16 +59,16 @@ bool FMaterialBridge::CanCreateContainerState() const
 	return !!ContainerStateType;
 }
 
-FInstancedStruct FMaterialBridge::CreateContainerState() const
+TInstancedStruct<FAvaMaterialContainerState> FMaterialBridge::CreateContainerState() const
 {
 	if (CanCreateContainerState())
 	{
-		return FInstancedStruct(ContainerStateType);
+		return TInstancedStruct<FAvaMaterialContainerState>(ContainerStateType);
 	}
-	return FInstancedStruct();
+	return TInstancedStruct<FAvaMaterialContainerState>();
 }
 
-bool FMaterialBridge::ApplyState(const FApplyStateContext& InContext, FConstStructView InContainerState, const FApplyStateOptions& InOptions) const
+bool FMaterialBridge::ApplyState(const FApplyStateContext& InContext, TConstStructView<FAvaMaterialContainerState> InContainerState, const FApplyStateOptions& InOptions) const
 {
 	if (InContainerState.IsValid() && InContainerState.GetScriptStruct() == ContainerStateType)
 	{
@@ -78,11 +78,11 @@ bool FMaterialBridge::ApplyState(const FApplyStateContext& InContext, FConstStru
 	return false;
 }
 
-bool FMaterialBridge::StoreState(const FStoreStateContext& InContext, TNotNull<FInstancedStruct*> InOutContainerState, const FStoreStateOptions& InOptions) const
+bool FMaterialBridge::StoreState(const FStoreStateContext& InContext, TNotNull<TInstancedStruct<FAvaMaterialContainerState>*> InOutContainerState, const FStoreStateOptions& InOptions) const
 {
 	if (!InOutContainerState->IsValid() || InOutContainerState->GetScriptStruct() != ContainerStateType)
 	{
-		FInstancedStruct NewContainerState = CreateContainerState();
+		TInstancedStruct<FAvaMaterialContainerState> NewContainerState = CreateContainerState();
 		if (!NewContainerState.IsValid())
 		{
 			UE_LOG(LogAvaMaterialBridge, Verbose, TEXT("Material Container State (Type: '%s') for Material Bridge (Bridged Type: '%s') could not be created.")
@@ -97,7 +97,7 @@ bool FMaterialBridge::StoreState(const FStoreStateContext& InContext, TNotNull<F
 	return StoreState(InContext, *InOutContainerState, InOptions);
 }
 
-bool FMaterialBridge::StoreState(const FStoreStateContext& InContext, FStructView InContainerState, const FStoreStateOptions& InOptions) const
+bool FMaterialBridge::StoreState(const FStoreStateContext& InContext, TStructView<FAvaMaterialContainerState> InContainerState, const FStoreStateOptions& InOptions) const
 {
 	if (InContainerState.IsValid() && InContainerState.GetScriptStruct() == ContainerStateType)
 	{
diff --git a/Engine/Plugins/VirtualProduction/Avalanche/Source/AvalancheMaterial/Private/MaterialBridge/Common/AvaActorMaterialBridge.cpp b/Engine/Plugins/VirtualProduction/Avalanche/Source/AvalancheMaterial/Private/MaterialBridge/Common/AvaActorMaterialBridge.cpp
--- a/Engine/Plugins/VirtualProduction/Avalanche/Source/AvalancheMaterial/Private/MaterialBridge/Common/AvaActorMaterialBridge.cpp
+++ b/Engine/Plugins/VirtualProduction/Avalanche/Source/AvalancheMaterial/Private/MaterialBridge/Common/AvaActorMaterialBridge.cpp
@@ -107,7 +107,7 @@ TSubScriptStructOf<FAvaMaterialContainerState> FActorMaterialBridge::OnGetContai
 	return FContainerState::StaticStruct();
 }
 
-void FActorMaterialBridge::OnApplyState(const FApplyStateContext& InContext, FConstStructView InContainerState, const FApplyStateOptions& InOptions) const
+void FActorMaterialBridge::OnApplyState(const FApplyStateContext& InContext, TConstStructView<FAvaMaterialContainerState> InContainerState, const FApplyStateOptions& InOptions) const
 {
 	const AActor& Actor = InContext.MaterialContainer.Get<const AActor>();
 	const FContainerState& ContainerState = InContainerState.Get<FContainerState>();
@@ -132,7 +132,7 @@ void FActorMaterialBridge::OnApplyState(const FApplyStateContext& InContext, FCo
 	}
 }
 
-void FActorMaterialBridge::OnStoreState(const FStoreStateContext& InContext, FStructView InContainerState, const FStoreStateOptions& InOptions) const
+void FActorMaterialBridge::OnStoreState(const FStoreStateContext& InContext, TStructView<FAvaMaterialContainerState> InContainerState, const FStoreStateOptions& InOptions) const
 {
 	const AActor& Actor = InContext.MaterialContainer.Get<const AActor>();
 
diff --git a/Engine/Plugins/VirtualProduction/Avalanche/Source/AvalancheMaterial/Private/MaterialBridge/Common/AvaDecalComponentMaterialBridge.cpp b/Engine/Plugins/VirtualProduction/Avalanche/Source/AvalancheMaterial/Private/MaterialBridge/Common/AvaDecalComponentMaterialBridge.cpp
--- a/Engine/Plugins/VirtualProduction/Avalanche/Source/AvalancheMaterial/Private/MaterialBridge/Common/AvaDecalComponentMaterialBridge.cpp
+++ b/Engine/Plugins/VirtualProduction/Avalanche/Source/AvalancheMaterial/Private/MaterialBridge/Common/AvaDecalComponentMaterialBridge.cpp
@@ -46,7 +46,7 @@ TSubScriptStructOf<FAvaMaterialContainerState> FDecalComponentMaterialBridge::On
 	return FContainerState::StaticStruct();
 }
 
-void FDecalComponentMaterialBridge::OnApplyState(const FApplyStateContext& InContext, FConstStructView InContainerState, const FApplyStateOptions& InOptions) const
+void FDecalComponentMaterialBridge::OnApplyState(const FApplyStateContext& InContext, TConstStructView<FAvaMaterialContainerState> InContainerState, const FApplyStateOptions& InOptions) const
 {
 	UDecalComponent& Component = InContext.MaterialContainer.GetMutable<UDecalComponent>();
 
@@ -54,7 +54,7 @@ void FDecalComponentMaterialBridge::OnApplyState(const FApplyStateContext& InCon
 	Component.SetDecalMaterial(ContainerState.DecalMaterial);
 }
 
-void FDecalComponentMaterialBridge::OnStoreState(const FStoreStateContext& InContext, FStructView InContainerState, const FStoreStateOptions& InOptions) const
+void FDecalComponentMaterialBridge::OnStoreState(const FStoreStateContext& InContext, TStructView<FAvaMaterialContainerState> InContainerState, const FStoreStateOptions& InOptions) const
 {
 	const UDecalComponent& Component = InContext.MaterialContainer.Get<const UDecalComponent>();
 
diff --git a/Engine/Plugins/VirtualProduction/Avalanche/Source/AvalancheMaterial/Private/MaterialBridge/Common/AvaPrimitiveComponentMaterialBridge.cpp b/Engine/Plugins/VirtualProduction/Avalanche/Source/AvalancheMaterial/Private/MaterialBridge/Common/AvaPrimitiveComponentMaterialBridge.cpp
--- a/Engine/Plugins/VirtualProduction/Avalanche/Source/AvalancheMaterial/Private/MaterialBridge/Common/AvaPrimitiveComponentMaterialBridge.cpp
+++ b/Engine/Plugins/VirtualProduction/Avalanche/Source/AvalancheMaterial/Private/MaterialBridge/Common/AvaPrimitiveComponentMaterialBridge.cpp
@@ -60,7 +60,7 @@ TSubScriptStructOf<FAvaMaterialContainerState> FPrimitiveComponentMaterialBridge
 	return FContainerState::StaticStruct();
 }
 
-void FPrimitiveComponentMaterialBridge::OnApplyState(const FApplyStateContext& InContext, FConstStructView InContainerState, const FApplyStateOptions& InOptions) const
+void FPrimitiveComponentMaterialBridge::OnApplyState(const FApplyStateContext& InContext, TConstStructView<FAvaMaterialContainerState> InContainerState, const FApplyStateOptions& InOptions) const
 {
 	UPrimitiveComponent& Component = InContext.MaterialContainer.GetMutable<UPrimitiveComponent>();
 
@@ -72,7 +72,7 @@ void FPrimitiveComponentMaterialBridge::OnApplyState(const FApplyStateContext& I
 	}
 }
 
-void FPrimitiveComponentMaterialBridge::OnStoreState(const FStoreStateContext& InContext, FStructView InContainerState, const FStoreStateOptions& InOptions) const
+void FPrimitiveComponentMaterialBridge::OnStoreState(const FStoreStateContext& InContext, TStructView<FAvaMaterialContainerState> InContainerState, const FStoreStateOptions& InOptions) const
 {
 	const UPrimitiveComponent& Component = InContext.MaterialContainer.Get<const UPrimitiveComponent>();
 	FContainerState& ContainerState = InContainerState.Get<FContainerState>();
diff --git a/Engine/Plugins/VirtualProduction/Avalanche/Source/AvalancheMaterial/Public/MaterialBridge/AvaMaterialBridge.h b/Engine/Plugins/VirtualProduction/Avalanche/Source/AvalancheMaterial/Public/MaterialBridge/AvaMaterialBridge.h
--- a/Engine/Plugins/VirtualProduction/Avalanche/Source/AvalancheMaterial/Public/MaterialBridge/AvaMaterialBridge.h
+++ b/Engine/Plugins/VirtualProduction/Avalanche/Source/AvalancheMaterial/Public/MaterialBridge/AvaMaterialBridge.h
@@ -108,7 +108,7 @@ public:
 	UE_API bool CanCreateContainerState() const;
 
 	/** Creates a material container state instance that can be used for this material bridge */
-	UE_API FInstancedStruct CreateContainerState() const;
+	UE_API TInstancedStruct<FAvaMaterialContainerState> CreateContainerState() const;
 
 	/** 
 	 * Applies the material container state to the context object.
@@ -117,7 +117,7 @@ public:
 	 * @param InOptions additional options to how apply should take place. see FMaterialBridgeApplyStateOptions
 	 * @return true if the container state was valid and was processed
 	 */
-	UE_API bool ApplyState(const FApplyStateContext& InContext, FConstStructView InContainerState, const FApplyStateOptions& InOptions) const;
+	UE_API bool ApplyState(const FApplyStateContext& InContext, TConstStructView<FAvaMaterialContainerState> InContainerState, const FApplyStateOptions& InOptions) const;
 
 	/** 
 	 * Stores the state of the context object in the material container state
@@ -127,7 +127,7 @@ public:
 	 * @param InOptions additional options to how store should take place. see FMaterialBridgeStoreStateOptions
 	 * @return true if the container state was valid (or initialized to valid) and was processed
 	 */
-	UE_API bool StoreState(const FStoreStateContext& InContext, TNotNull<FInstancedStruct*> InOutContainerState, const FStoreStateOptions& InOptions) const;
+	UE_API bool StoreState(const FStoreStateContext& InContext, TNotNull<TInstancedStruct<FAvaMaterialContainerState>*> InOutContainerState, const FStoreStateOptions& InOptions) const;
 
 	/** 
 	 * Stores the state of the context object in the material container state
@@ -136,7 +136,7 @@ public:
 	 * @param InOptions additional options to how store should take place. see FMaterialBridgeStoreStateOptions
 	 * @return true if the container state was valid and was processed
 	 */
-	UE_API bool StoreState(const FStoreStateContext& InContext, FStructView InContainerState, const FStoreStateOptions& InOptions) const;
+	UE_API bool StoreState(const FStoreStateContext& InContext, TStructView<FAvaMaterialContainerState> InContainerState, const FStoreStateOptions& InOptions) const;
 
 protected:
 	/** Retrieves the type that this bridge supports. Only called once in Initialize() */
@@ -167,12 +167,12 @@ protected:
 	}
 
 	/** Applies the container state to the context object */
-	virtual void OnApplyState(const FApplyStateContext& InContext, FConstStructView InContainerState, const FApplyStateOptions& InOptions) const
+	virtual void OnApplyState(const FApplyStateContext& InContext, TConstStructView<FAvaMaterialContainerState> InContainerState, const FApplyStateOptions& InOptions) const
 	{
 	}
 
 	/** Stores the state of the context object in the container state */
-	virtual void OnStoreState(const FStoreStateContext& InContext, FStructView InContainerState, const FStoreStateOptions& InOptions) const
+	virtual void OnStoreState(const FStoreStateContext& InContext, TStructView<FAvaMaterialContainerState> InContainerState, const FStoreStateOptions& InOptions) const
 	{
 	}
 
diff --git a/Engine/Plugins/VirtualProduction/Avalanche/Source/AvalancheMaterial/Public/MaterialBridge/AvaMaterialContainerState.h b/Engine/Plugins/VirtualProduction/Avalanche/Source/AvalancheMaterial/Public/MaterialBridge/AvaMaterialContainerState.h
--- a/Engine/Plugins/VirtualProduction/Avalanche/Source/AvalancheMaterial/Public/MaterialBridge/AvaMaterialContainerState.h
+++ b/Engine/Plugins/VirtualProduction/Avalanche/Source/AvalancheMaterial/Public/MaterialBridge/AvaMaterialContainerState.h
@@ -4,21 +4,9 @@
 
 #include "AvaMaterialContainerState.generated.h"
 
-class UMaterialInterface;
-
 /** Base struct for all material container states */
 USTRUCT()
 struct FAvaMaterialContainerState
 {
 	GENERATED_BODY()
 };
-
-/** Material container state that holds the materials as an array-like api */
-USTRUCT()
-struct FAvaMaterialArrayState : public FAvaMaterialContainerState
-{
-	GENERATED_BODY()
-
-	UPROPERTY()
-	TArray<TObjectPtr<UMaterialInterface>> Materials;
-};
diff --git a/Engine/Plugins/VirtualProduction/Avalanche/Source/AvalancheMaterial/Public/MaterialBridge/Common/AvaActorMaterialBridge.h b/Engine/Plugins/VirtualProduction/Avalanche/Source/AvalancheMaterial/Public/MaterialBridge/Common/AvaActorMaterialBridge.h
--- a/Engine/Plugins/VirtualProduction/Avalanche/Source/AvalancheMaterial/Public/MaterialBridge/Common/AvaActorMaterialBridge.h
+++ b/Engine/Plugins/VirtualProduction/Avalanche/Source/AvalancheMaterial/Public/MaterialBridge/Common/AvaActorMaterialBridge.h
@@ -31,7 +31,7 @@ struct FAvaActorMaterialContainerStateComponent
 
 	/** The container state data for the component. The type varies as each component could have a different Material Bridge */
 	UPROPERTY()
-	FInstancedStruct ContainerState;
+	TInstancedStruct<FAvaMaterialContainerState> ContainerState;
 };
 
 /** Material container state for actors */
@@ -60,8 +60,8 @@ protected:
 	UE_API virtual EControlFlow OnAccessSlots(const FReadSlotContext& InContext, TFunctionRef<EControlFlow(const FReadSlotContext&, const FReadSlot&)> InFunc, const FReadSlotOptions& InOptions) const override;
 	UE_API virtual EControlFlow OnAccessSlots(const FWriteSlotContext& InContext, TFunctionRef<EControlFlow(const FWriteSlotContext&, FWriteSlot&)> InFunc, const FWriteSlotOptions& InOptions) const override;
 	UE_API virtual TSubScriptStructOf<FAvaMaterialContainerState> OnGetContainerStateType() const override;
-	UE_API virtual void OnApplyState(const FApplyStateContext& InContext, FConstStructView InContainerState, const FApplyStateOptions& InOptions) const override;
-	UE_API virtual void OnStoreState(const FStoreStateContext& InContext, FStructView InContainerState, const FStoreStateOptions& InOptions) const override;
+	UE_API virtual void OnApplyState(const FApplyStateContext& InContext, TConstStructView<FAvaMaterialContainerState> InContainerState, const FApplyStateOptions& InOptions) const override;
+	UE_API virtual void OnStoreState(const FStoreStateContext& InContext, TStructView<FAvaMaterialContainerState> InContainerState, const FStoreStateOptions& InOptions) const override;
 	//~ End FMaterialBridge
 };
 
diff --git a/Engine/Plugins/VirtualProduction/Avalanche/Source/AvalancheMaterial/Public/MaterialBridge/Common/AvaDecalComponentMaterialBridge.h b/Engine/Plugins/VirtualProduction/Avalanche/Source/AvalancheMaterial/Public/MaterialBridge/Common/AvaDecalComponentMaterialBridge.h
--- a/Engine/Plugins/VirtualProduction/Avalanche/Source/AvalancheMaterial/Public/MaterialBridge/Common/AvaDecalComponentMaterialBridge.h
+++ b/Engine/Plugins/VirtualProduction/Avalanche/Source/AvalancheMaterial/Public/MaterialBridge/Common/AvaDecalComponentMaterialBridge.h
@@ -8,6 +8,8 @@
 
 #define UE_API AVALANCHEMATERIAL_API
 
+class UMaterialInterface;
+
 /** Material container state for decal components */
 USTRUCT()
 struct FAvaDecalComponentMaterialContainerState : public FAvaMaterialContainerState
@@ -34,8 +36,8 @@ protected:
 	UE_API virtual EControlFlow OnAccessSlots(const FReadSlotContext& InContext, TFunctionRef<EControlFlow(const FReadSlotContext&, const FReadSlot&)> InFunc, const FReadSlotOptions& InOptions) const override;
 	UE_API virtual EControlFlow OnAccessSlots(const FWriteSlotContext& InContext, TFunctionRef<EControlFlow(const FWriteSlotContext&, FWriteSlot&)> InFunc, const FWriteSlotOptions& InOptions) const override;
 	UE_API virtual TSubScriptStructOf<FAvaMaterialContainerState> OnGetContainerStateType() const override;
-	UE_API virtual void OnApplyState(const FApplyStateContext& InContext, FConstStructView InContainerState, const FApplyStateOptions& InOptions) const override;
-	UE_API virtual void OnStoreState(const FStoreStateContext& InContext, FStructView InContainerState, const FStoreStateOptions& InOptions) const override;
+	UE_API virtual void OnApplyState(const FApplyStateContext& InContext, TConstStructView<FAvaMaterialContainerState> InContainerState, const FApplyStateOptions& InOptions) const override;
+	UE_API virtual void OnStoreState(const FStoreStateContext& InContext, TStructView<FAvaMaterialContainerState> InContainerState, const FStoreStateOptions& InOptions) const override;
 	//~ End FMaterialBridge
 };
 
diff --git a/Engine/Plugins/VirtualProduction/Avalanche/Source/AvalancheMaterial/Public/MaterialBridge/Common/AvaPrimitiveComponentMaterialBridge.h b/Engine/Plugins/VirtualProduction/Avalanche/Source/AvalancheMaterial/Public/MaterialBridge/Common/AvaPrimitiveComponentMaterialBridge.h
--- a/Engine/Plugins/VirtualProduction/Avalanche/Source/AvalancheMaterial/Public/MaterialBridge/Common/AvaPrimitiveComponentMaterialBridge.h
+++ b/Engine/Plugins/VirtualProduction/Avalanche/Source/AvalancheMaterial/Public/MaterialBridge/Common/AvaPrimitiveComponentMaterialBridge.h
@@ -8,6 +8,8 @@
 
 #define UE_API AVALANCHEMATERIAL_API
 
+class UMaterialInterface;
+
 /** Material container state for primitive components */
 USTRUCT()
 struct FAvaPrimitiveComponentMaterialContainerState : public FAvaMaterialContainerState
@@ -33,8 +35,8 @@ protected:
 	UE_API virtual EControlFlow OnAccessSlots(const FReadSlotContext& InContext, TFunctionRef<EControlFlow(const FReadSlotContext&, const FReadSlot&)> InFunc, const FReadSlotOptions& InOptions) const override;
 	UE_API virtual EControlFlow OnAccessSlots(const FWriteSlotContext& InContext, TFunctionRef<EControlFlow(const FWriteSlotContext&, FWriteSlot&)> InFunc, const FWriteSlotOptions& InOptions) const override;
 	UE_API virtual TSubScriptStructOf<FAvaMaterialContainerState> OnGetContainerStateType() const override;
-	UE_API virtual void OnApplyState(const FApplyStateContext& InContext, FConstStructView InContainerState, const FApplyStateOptions& InOptions) const override;
-	UE_API virtual void OnStoreState(const FStoreStateContext& InContext, FStructView InContainerState, const FStoreStateOptions& InOptions) const override;
+	UE_API virtual void OnApplyState(const FApplyStateContext& InContext, TConstStructView<FAvaMaterialContainerState> InContainerState, const FApplyStateOptions& InOptions) const override;
+	UE_API virtual void OnStoreState(const FStoreStateContext& InContext, TStructView<FAvaMaterialContainerState> InContainerState, const FStoreStateOptions& InOptions) const override;
 	//~ End FMaterialBridge
 };
 
diff --git a/Engine/Plugins/VirtualProduction/Avalanche/Source/AvalancheShapes/Private/AvaShapeMaterialBridge.cpp b/Engine/Plugins/VirtualProduction/Avalanche/Source/AvalancheShapes/Private/AvaShapeMaterialBridge.cpp
--- a/Engine/Plugins/VirtualProduction/Avalanche/Source/AvalancheShapes/Private/AvaShapeMaterialBridge.cpp
+++ b/Engine/Plugins/VirtualProduction/Avalanche/Source/AvalancheShapes/Private/AvaShapeMaterialBridge.cpp
@@ -97,7 +97,7 @@ TSubScriptStructOf<FAvaMaterialContainerState> FShapeMaterialBridge::OnGetContai
 	return FContainerState::StaticStruct();
 }
 
-void FShapeMaterialBridge::OnApplyState(const FApplyStateContext& InContext, FConstStructView InContainerState, const FApplyStateOptions& InOptions) const
+void FShapeMaterialBridge::OnApplyState(const FApplyStateContext& InContext, TConstStructView<FAvaMaterialContainerState> InContainerState, const FApplyStateOptions& InOptions) const
 {
 	UAvaShapeDynamicMeshBase& Shape = InContext.MaterialContainer.GetMutable<UAvaShapeDynamicMeshBase>();
 	const FContainerState& ContainerState = InContainerState.Get<FContainerState>();
@@ -123,7 +123,7 @@ void FShapeMaterialBridge::OnApplyState(const FApplyStateContext& InContext, FCo
 	Shape.MarkRenderStateDirty();
 }
 
-void FShapeMaterialBridge::OnStoreState(const FStoreStateContext& InContext, FStructView InContainerState, const FStoreStateOptions& InOptions) const
+void FShapeMaterialBridge::OnStoreState(const FStoreStateContext& InContext, TStructView<FAvaMaterialContainerState> InContainerState, const FStoreStateOptions& InOptions) const
 {
 	const UAvaShapeDynamicMeshBase& Shape = InContext.MaterialContainer.Get<const UAvaShapeDynamicMeshBase>();
 	FContainerState& StateData = InContainerState.Get<FContainerState>();
diff --git a/Engine/Plugins/VirtualProduction/Avalanche/Source/AvalancheShapes/Public/AvaShapeMaterialBridge.h b/Engine/Plugins/VirtualProduction/Avalanche/Source/AvalancheShapes/Public/AvaShapeMaterialBridge.h
--- a/Engine/Plugins/VirtualProduction/Avalanche/Source/AvalancheShapes/Public/AvaShapeMaterialBridge.h
+++ b/Engine/Plugins/VirtualProduction/Avalanche/Source/AvalancheShapes/Public/AvaShapeMaterialBridge.h
@@ -74,8 +74,8 @@ protected:
 	UE_API virtual EControlFlow OnAccessSlots(const FReadSlotContext& InContext, TFunctionRef<EControlFlow(const FReadSlotContext&, const FReadSlot&)> InFunc, const FReadSlotOptions& InOptions) const override;
 	UE_API virtual EControlFlow OnAccessSlots(const FWriteSlotContext& InContext, TFunctionRef<EControlFlow(const FWriteSlotContext&, FWriteSlot&)> InFunc, const FWriteSlotOptions& InOptions) const override;
 	UE_API virtual TSubScriptStructOf<FAvaMaterialContainerState> OnGetContainerStateType() const override;
-	UE_API virtual void OnApplyState(const FApplyStateContext& InContext, FConstStructView InContainerState, const FApplyStateOptions& InOptions) const override;
-	UE_API virtual void OnStoreState(const FStoreStateContext& InContext, FStructView InContainerState, const FStoreStateOptions& InOptions) const override;
+	UE_API virtual void OnApplyState(const FApplyStateContext& InContext, TConstStructView<FAvaMaterialContainerState> InContainerState, const FApplyStateOptions& InOptions) const override;
+	UE_API virtual void OnStoreState(const FStoreStateContext& InContext, TStructView<FAvaMaterialContainerState> InContainerState, const FStoreStateOptions& InOptions) const override;
 	//~ End FMaterialBridge
 };
 
diff --git a/Engine/Plugins/VirtualProduction/Avalanche/Source/AvalancheText/Private/MaterialBridge/AvaText3DComponentMaterialBridge.cpp b/Engine/Plugins/VirtualProduction/Avalanche/Source/AvalancheText/Private/MaterialBridge/AvaText3DComponentMaterialBridge.cpp
--- a/Engine/Plugins/VirtualProduction/Avalanche/Source/AvalancheText/Private/MaterialBridge/AvaText3DComponentMaterialBridge.cpp
+++ b/Engine/Plugins/VirtualProduction/Avalanche/Source/AvalancheText/Private/MaterialBridge/AvaText3DComponentMaterialBridge.cpp
@@ -53,7 +53,7 @@ TSubScriptStructOf<FAvaMaterialContainerState> FText3DComponentMaterialBridge::O
 	return FContainerState::StaticStruct();
 }
 
-void FText3DComponentMaterialBridge::OnApplyState(const FApplyStateContext& InContext, FConstStructView InContainerState, const FApplyStateOptions& InOptions) const
+void FText3DComponentMaterialBridge::OnApplyState(const FApplyStateContext& InContext, TConstStructView<FAvaMaterialContainerState> InContainerState, const FApplyStateOptions& InOptions) const
 {
 	UText3DComponent& Component = InContext.MaterialContainer.GetMutable<UText3DComponent>();
 	const FContainerState& ContainerState = InContainerState.Get<FContainerState>();
@@ -68,7 +68,7 @@ void FText3DComponentMaterialBridge::OnApplyState(const FApplyStateContext& InCo
 	}
 }
 
-void FText3DComponentMaterialBridge::OnStoreState(const FStoreStateContext& InContext, FStructView InContainerState, const FStoreStateOptions& InOptions) const
+void FText3DComponentMaterialBridge::OnStoreState(const FStoreStateContext& InContext, TStructView<FAvaMaterialContainerState> InContainerState, const FStoreStateOptions& InOptions) const
 {
 	const UText3DComponent& Component = InContext.MaterialContainer.Get<const UText3DComponent>();
 	FContainerState& StateData = InContainerState.Get<FContainerState>();
diff --git a/Engine/Plugins/VirtualProduction/Avalanche/Source/AvalancheText/Private/MaterialBridge/AvaText3DDefaultExtensionMaterialBridge.cpp b/Engine/Plugins/VirtualProduction/Avalanche/Source/AvalancheText/Private/MaterialBridge/AvaText3DDefaultExtensionMaterialBridge.cpp
--- a/Engine/Plugins/VirtualProduction/Avalanche/Source/AvalancheText/Private/MaterialBridge/AvaText3DDefaultExtensionMaterialBridge.cpp
+++ b/Engine/Plugins/VirtualProduction/Avalanche/Source/AvalancheText/Private/MaterialBridge/AvaText3DDefaultExtensionMaterialBridge.cpp
@@ -76,7 +76,7 @@ EControlFlow FText3DDefaultExtensionMaterialBridge::OnAccessSlots(const FWriteSl
 	return ControlFlow;
 }
 
-void FText3DDefaultExtensionMaterialBridge::OnApplyState(const FApplyStateContext& InContext, FConstStructView InContainerState, const FApplyStateOptions& InOptions) const
+void FText3DDefaultExtensionMaterialBridge::OnApplyState(const FApplyStateContext& InContext, TConstStructView<FAvaMaterialContainerState> InContainerState, const FApplyStateOptions& InOptions) const
 {
 	UText3DDefaultMaterialExtension& MaterialExtension = InContext.MaterialContainer.GetMutable<UText3DDefaultMaterialExtension>();
 	const FContainerState& ContainerState = InContainerState.Get<FContainerState>();
@@ -87,7 +87,7 @@ void FText3DDefaultExtensionMaterialBridge::OnApplyState(const FApplyStateContex
 	Super::OnApplyState(InContext, InContainerState, InOptions);
 }
 
-void FText3DDefaultExtensionMaterialBridge::OnStoreState(const FStoreStateContext& InContext, FStructView InContainerState, const FStoreStateOptions& InOptions) const
+void FText3DDefaultExtensionMaterialBridge::OnStoreState(const FStoreStateContext& InContext, TStructView<FAvaMaterialContainerState> InContainerState, const FStoreStateOptions& InOptions) const
 {
 	// Call parent implementation to store the materials
 	Super::OnStoreState(InContext, InContainerState, InOptions);
diff --git a/Engine/Plugins/VirtualProduction/Avalanche/Source/AvalancheText/Private/MaterialBridge/AvaText3DExtensionMaterialBridge.cpp b/Engine/Plugins/VirtualProduction/Avalanche/Source/AvalancheText/Private/MaterialBridge/AvaText3DExtensionMaterialBridge.cpp
--- a/Engine/Plugins/VirtualProduction/Avalanche/Source/AvalancheText/Private/MaterialBridge/AvaText3DExtensionMaterialBridge.cpp
+++ b/Engine/Plugins/VirtualProduction/Avalanche/Source/AvalancheText/Private/MaterialBridge/AvaText3DExtensionMaterialBridge.cpp
@@ -68,7 +68,7 @@ TSubScriptStructOf<FAvaMaterialContainerState> FText3DExtensionMaterialBridge::O
 	return FContainerState::StaticStruct();
 }
 
-void FText3DExtensionMaterialBridge::OnApplyState(const FApplyStateContext& InContext, FConstStructView InContainerState, const FApplyStateOptions& InOptions) const
+void FText3DExtensionMaterialBridge::OnApplyState(const FApplyStateContext& InContext, TConstStructView<FAvaMaterialContainerState> InContainerState, const FApplyStateOptions& InOptions) const
 {
 	UText3DMaterialExtensionBase& MaterialExtension = InContext.MaterialContainer.GetMutable<UText3DMaterialExtensionBase>();
 	const FContainerState& ContainerState = InContainerState.Get<const FContainerState>();
@@ -83,7 +83,7 @@ void FText3DExtensionMaterialBridge::OnApplyState(const FApplyStateContext& InCo
 	}
 }
 
-void FText3DExtensionMaterialBridge::OnStoreState(const FStoreStateContext& InContext, FStructView InContainerState, const FStoreStateOptions& InOptions) const
+void FText3DExtensionMaterialBridge::OnStoreState(const FStoreStateContext& InContext, TStructView<FAvaMaterialContainerState> InContainerState, const FStoreStateOptions& InOptions) const
 {
 	const UText3DMaterialExtensionBase& MaterialExtension = InContext.MaterialContainer.Get<const UText3DMaterialExtensionBase>();
 	FContainerState& ContainerState = InContainerState.Get<FContainerState>();
diff --git a/Engine/Plugins/VirtualProduction/Avalanche/Source/AvalancheText/Public/MaterialBridge/AvaText3DComponentMaterialBridge.h b/Engine/Plugins/VirtualProduction/Avalanche/Source/AvalancheText/Public/MaterialBridge/AvaText3DComponentMaterialBridge.h
--- a/Engine/Plugins/VirtualProduction/Avalanche/Source/AvalancheText/Public/MaterialBridge/AvaText3DComponentMaterialBridge.h
+++ b/Engine/Plugins/VirtualProduction/Avalanche/Source/AvalancheText/Public/MaterialBridge/AvaText3DComponentMaterialBridge.h
@@ -15,7 +15,7 @@ struct FAvaText3DComponentMaterialStateData : public FAvaMaterialContainerState
 	GENERATED_BODY()
 
 	UPROPERTY()
-	FInstancedStruct MaterialExtensionStateData;
+	TInstancedStruct<FAvaMaterialContainerState> MaterialExtensionStateData;
 };
 
 namespace UE::Ava
@@ -33,8 +33,8 @@ protected:
 	UE_API virtual EControlFlow OnAccessSlots(const FReadSlotContext& InContext, TFunctionRef<EControlFlow(const FReadSlotContext&, const FReadSlot&)> InFunc, const FReadSlotOptions& InOptions) const override;
 	UE_API virtual EControlFlow OnAccessSlots(const FWriteSlotContext& InContext, TFunctionRef<EControlFlow(const FWriteSlotContext&, FWriteSlot&)> InFunc, const FWriteSlotOptions& InOptions) const override;
 	UE_API virtual TSubScriptStructOf<FAvaMaterialContainerState> OnGetContainerStateType() const override;
-	UE_API virtual void OnApplyState(const FApplyStateContext& InContext, FConstStructView InContainerState, const FApplyStateOptions& InOptions) const override;
-	UE_API virtual void OnStoreState(const FStoreStateContext& InContext, FStructView InContainerState, const FStoreStateOptions& InOptions) const override;
+	UE_API virtual void OnApplyState(const FApplyStateContext& InContext, TConstStructView<FAvaMaterialContainerState> InContainerState, const FApplyStateOptions& InOptions) const override;
+	UE_API virtual void OnStoreState(const FStoreStateContext& InContext, TStructView<FAvaMaterialContainerState> InContainerState, const FStoreStateOptions& InOptions) const override;
 	//~ End FMaterialBridge
 };
 
diff --git a/Engine/Plugins/VirtualProduction/Avalanche/Source/AvalancheText/Public/MaterialBridge/AvaText3DDefaultExtensionMaterialBridge.h b/Engine/Plugins/VirtualProduction/Avalanche/Source/AvalancheText/Public/MaterialBridge/AvaText3DDefaultExtensionMaterialBridge.h
--- a/Engine/Plugins/VirtualProduction/Avalanche/Source/AvalancheText/Public/MaterialBridge/AvaText3DDefaultExtensionMaterialBridge.h
+++ b/Engine/Plugins/VirtualProduction/Avalanche/Source/AvalancheText/Public/MaterialBridge/AvaText3DDefaultExtensionMaterialBridge.h
@@ -44,8 +44,8 @@ protected:
 	UE_API virtual const UStruct* OnGetBridgedType() const override;
 	UE_API virtual TSubScriptStructOf<FAvaMaterialContainerState> OnGetContainerStateType() const override;
 	UE_API virtual EControlFlow OnAccessSlots(const FWriteSlotContext& InContext, TFunctionRef<EControlFlow(const FWriteSlotContext&, FWriteSlot&)> InFunc, const FWriteSlotOptions& InOptions) const override;
-	UE_API virtual void OnApplyState(const FApplyStateContext& InContext, FConstStructView InContainerState, const FApplyStateOptions& InOptions) const override;
-	UE_API virtual void OnStoreState(const FStoreStateContext& InContext, FStructView InContainerState, const FStoreStateOptions& InOptions) const override;
+	UE_API virtual void OnApplyState(const FApplyStateContext& InContext, TConstStructView<FAvaMaterialContainerState> InContainerState, const FApplyStateOptions& InOptions) const override;
+	UE_API virtual void OnStoreState(const FStoreStateContext& InContext, TStructView<FAvaMaterialContainerState> InContainerState, const FStoreStateOptions& InOptions) const override;
 	//~ End FMaterialBridge
 };
 
diff --git a/Engine/Plugins/VirtualProduction/Avalanche/Source/AvalancheText/Public/MaterialBridge/AvaText3DExtensionMaterialBridge.h b/Engine/Plugins/VirtualProduction/Avalanche/Source/AvalancheText/Public/MaterialBridge/AvaText3DExtensionMaterialBridge.h
--- a/Engine/Plugins/VirtualProduction/Avalanche/Source/AvalancheText/Public/MaterialBridge/AvaText3DExtensionMaterialBridge.h
+++ b/Engine/Plugins/VirtualProduction/Avalanche/Source/AvalancheText/Public/MaterialBridge/AvaText3DExtensionMaterialBridge.h
@@ -73,8 +73,8 @@ protected:
 	UE_API virtual EControlFlow OnAccessSlots(const FReadSlotContext& InContext, TFunctionRef<EControlFlow(const FReadSlotContext&, const FReadSlot&)> InFunc, const FReadSlotOptions& InOptions) const override;
 	UE_API virtual EControlFlow OnAccessSlots(const FWriteSlotContext& InContext, TFunctionRef<EControlFlow(const FWriteSlotContext&, FWriteSlot&)> InFunc, const FWriteSlotOptions& InOptions) const override;
 	UE_API virtual TSubScriptStructOf<FAvaMaterialContainerState> OnGetContainerStateType() const override;
-	UE_API virtual void OnApplyState(const FApplyStateContext& InContext, FConstStructView InContainerState, const FApplyStateOptions& InOptions) const override;
-	UE_API virtual void OnStoreState(const FStoreStateContext& InContext, FStructView InContainerState, const FStoreStateOptions& InOptions) const override;
+	UE_API virtual void OnApplyState(const FApplyStateContext& InContext, TConstStructView<FAvaMaterialContainerState> InContainerState, const FApplyStateOptions& InOptions) const override;
+	UE_API virtual void OnStoreState(const FStoreStateContext& InContext, TStructView<FAvaMaterialContainerState> InContainerState, const FStoreStateOptions& InOptions) const override;
 	//~ End FMaterialBridge
 
 };
