Avalanche - create AScreenshotFunctionalTestAvalanche for Motion Design screenshot testing.
New screenshot functional test class specifically for Avalanche (Motion Design) levels to work around two critical issues: CVar Priority Conflicts and PIE Screenshot Capture Failure

#Synced-CL 48812843

diff --git a/Engine/Plugins/VirtualProduction/Avalanche/Avalanche.uplugin b/Engine/Plugins/VirtualProduction/Avalanche/Avalanche.uplugin
--- a/Engine/Plugins/VirtualProduction/Avalanche/Avalanche.uplugin
+++ b/Engine/Plugins/VirtualProduction/Avalanche/Avalanche.uplugin
@@ -247,6 +247,12 @@
 			"Name": "AvalancheMaterial",
 			"Type": "Runtime",
 			"LoadingPhase": "Default"
+		},
+		{
+			"Name": "AvalancheFunctionalTest",
+			"Type": "Runtime",
+			"LoadingPhase": "Default",
+			"TargetDenyList": [ "Server" ]
 		}
 	],
 	"Plugins": [
diff --git a/Engine/Plugins/VirtualProduction/Avalanche/Source/AvalancheFunctionalTest/AvalancheFunctionalTest.Build.cs b/Engine/Plugins/VirtualProduction/Avalanche/Source/AvalancheFunctionalTest/AvalancheFunctionalTest.Build.cs
new file mode 100644
--- /dev/null
+++ b/Engine/Plugins/VirtualProduction/Avalanche/Source/AvalancheFunctionalTest/AvalancheFunctionalTest.Build.cs
@@ -0,0 +1,21 @@
+// Copyright Epic Games, Inc. All Rights Reserved.
+
+using UnrealBuildTool;
+
+public class AvalancheFunctionalTest : ModuleRules
+{
+	public AvalancheFunctionalTest(ReadOnlyTargetRules Target) : base(Target)
+	{
+		PCHUsage = ModuleRules.PCHUsageMode.UseExplicitOrSharedPCHs;
+
+		PublicDependencyModuleNames.AddRange(
+			new string[]
+			{
+				"Core",
+				"CoreUObject",
+				"Engine",
+				"FunctionalTesting",
+			}
+		);
+	}
+}
diff --git a/Engine/Plugins/VirtualProduction/Avalanche/Source/AvalancheFunctionalTest/Private/ScreenshotFunctionalTestAvalanche.cpp b/Engine/Plugins/VirtualProduction/Avalanche/Source/AvalancheFunctionalTest/Private/ScreenshotFunctionalTestAvalanche.cpp
new file mode 100644
--- /dev/null
+++ b/Engine/Plugins/VirtualProduction/Avalanche/Source/AvalancheFunctionalTest/Private/ScreenshotFunctionalTestAvalanche.cpp
@@ -0,0 +1,121 @@
+// Copyright Epic Games, Inc. All Rights Reserved.
+
+#include "ScreenshotFunctionalTestAvalanche.h"
+#include "DynamicResolutionState.h"
+#include "Engine/Engine.h"
+#include "Engine/GameViewportClient.h"
+#include "Tests/AutomationCommon.h"
+#include "UnrealClient.h"
+
+DEFINE_LOG_CATEGORY_STATIC(LogScreenshotAvalanche, Log, All);
+
+AScreenshotFunctionalTestAvalanche::AScreenshotFunctionalTestAvalanche(const FObjectInitializer& ObjectInitializer)
+	: Super(ObjectInitializer)
+	, AntiAliasingOverride(TEXT("r.AntiAliasingMethod"))
+	, AutoExposureOverride(TEXT("r.DefaultFeature.AutoExposure"))
+	, MotionBlurOverride(TEXT("r.DefaultFeature.MotionBlur"))
+	, MotionBlurQualityOverride(TEXT("r.MotionBlurQuality"))
+	, ContactShadowsOverride(TEXT("r.ContactShadows"))
+	, ScreenSpaceReflectionQualityOverride(TEXT("r.SSR.Quality"))
+	, EyeAdaptationQualityOverride(TEXT("r.EyeAdaptationQuality"))
+	, TonemapperGammaOverride(TEXT("r.TonemapperGamma"))
+	, ScreenPercentageOverride(TEXT("r.ScreenPercentage"))
+	, DynamicResTestScreenPercentageOverride(TEXT("r.DynamicRes.TestScreenPercentage"))
+	, DynamicResOperationModeOverride(TEXT("r.DynamicRes.OperationMode"))
+	, SecondaryScreenPercentageOverride(TEXT("r.SecondaryScreenPercentage.GameViewport"))
+{
+}
+
+void AScreenshotFunctionalTestAvalanche::PrepareTest()
+{
+	// Apply Avalanche-compatible CVar overrides BEFORE base class setup
+	// This prevents priority conflicts with Avalanche's custom viewport
+	ApplyAvalancheRenderingOverrides();
+
+	// Now call base class PrepareTest which sets up the entire screenshot test
+	Super::PrepareTest();
+}
+
+void AScreenshotFunctionalTestAvalanche::RequestScreenshot()
+{
+	// Call base class to register instance delegate callback and flush rendering
+	Super::RequestScreenshot();
+
+	// CRITICAL FIX for PIE: Also register to the static delegate used by FEditorViewportClient
+	// PIE uses FEditorViewportClient::ProcessScreenShots() which fires FScreenshotRequest::OnScreenshotCaptured() (static)
+	// instead of UGameViewportClient::OnScreenshotCaptured() (instance) that base class registers to
+	FScreenshotRequest::OnScreenshotCaptured().AddUObject(this, &AScreenshotFunctionalTestAvalanche::OnScreenShotCaptured);
+
+	// Request the actual screenshot (base class only registers callbacks)
+	FIntRect ViewRect = FIntRect(); // Empty ViewRect = entire render target
+	FScreenshotRequest::RequestScreenshot(TEXT(""), /*bInShowUI*/ false, /*bAddUniqueSuffix*/ true, /*bHdrScreenshot*/ false, ViewRect);
+}
+
+void AScreenshotFunctionalTestAvalanche::OnScreenShotCaptured(int32 InSizeX, int32 InSizeY, const TArray<FColor>& InImageData)
+{
+	// Remove callback from BOTH delegates (instance and static)
+	UGameViewportClient* GameViewportClient = AutomationCommon::GetAnyGameViewportClient();
+	if (GameViewportClient)
+	{
+		GameViewportClient->OnScreenshotCaptured().RemoveAll(this);
+	}
+	FScreenshotRequest::OnScreenshotCaptured().RemoveAll(this);
+
+	// Call parent implementation which handles the screenshot capture and comparison pipeline
+	Super::OnScreenShotCaptured(InSizeX, InSizeY, InImageData);
+}
+
+void AScreenshotFunctionalTestAvalanche::ApplyAvalancheRenderingOverrides()
+{
+	// Check if we should disable noisy rendering features (same behavior as base class)
+	if (ScreenshotOptions.bDisableNoisyRenderingFeatures)
+	{
+		// Disable anti-aliasing for deterministic rendering
+		AntiAliasingOverride.SetOverride(0);
+
+		// Disable auto-exposure
+		AutoExposureOverride.SetOverride(0);
+
+		// Disable motion blur
+		MotionBlurOverride.SetOverride(0);
+		MotionBlurQualityOverride.SetOverride(0);
+
+		// Disable screen space reflections
+		ScreenSpaceReflectionQualityOverride.SetOverride(0);
+
+		// Disable contact shadows
+		ContactShadowsOverride.SetOverride(0);
+
+		// Disable eye adaptation
+		EyeAdaptationQualityOverride.SetOverride(0);
+
+		// Set fixed tonemapper gamma
+		TonemapperGammaOverride.SetOverride(2.2f);
+	}
+	else if (ScreenshotOptions.bDisableTonemapping)
+	{
+		// Only disable eye adaptation and set gamma
+		EyeAdaptationQualityOverride.SetOverride(0);
+		TonemapperGammaOverride.SetOverride(2.2f);
+	}
+
+	// Force screen percentage to 100% for consistent results
+	{
+		// Completely disable dynamic resolution
+		DynamicResTestScreenPercentageOverride.SetOverride(0);
+		DynamicResOperationModeOverride.SetOverride(0);
+
+		// Emit dynamic resolution state change events
+		if (GEngine)
+		{
+			GEngine->EmitDynamicResolutionEvent(EDynamicResolutionStateEvent::EndFrame);
+			GEngine->EmitDynamicResolutionEvent(EDynamicResolutionStateEvent::BeginFrame);
+		}
+
+		// Set screen percentage to 100%
+		ScreenPercentageOverride.SetOverride(100.0f);
+	}
+
+	// Ignore High-DPI settings
+	SecondaryScreenPercentageOverride.SetOverride(100.0f);
+}
diff --git a/Engine/Plugins/VirtualProduction/Avalanche/Source/AvalancheFunctionalTest/Public/ScreenshotFunctionalTestAvalanche.h b/Engine/Plugins/VirtualProduction/Avalanche/Source/AvalancheFunctionalTest/Public/ScreenshotFunctionalTestAvalanche.h
new file mode 100644
--- /dev/null
+++ b/Engine/Plugins/VirtualProduction/Avalanche/Source/AvalancheFunctionalTest/Public/ScreenshotFunctionalTestAvalanche.h
@@ -0,0 +1,102 @@
+// Copyright Epic Games, Inc. All Rights Reserved.
+
+#pragma once
+
+#include "Containers/UnrealString.h"
+#include "HAL/IConsoleManager.h"
+#include "ScreenshotFunctionalTest.h"
+#include "ScreenshotFunctionalTestAvalanche.generated.h"
+
+/**
+ * Screenshot functional test optimized for Avalanche (Motion Design) levels.
+ *
+ * This class works around two critical issues that prevent regular screenshot tests
+ * from working correctly in Avalanche/Motion Design levels:
+ *
+ * 1. CVar Priority Conflicts:
+ *    - Avalanche's custom viewport client (UAvaGameViewportClient) sets rendering CVars
+ *      at Constructor/Code priority levels
+ *    - Base AScreenshotFunctionalTest uses SetWithCurrentPriority() which cannot override
+ *      Constructor-level CVars, causing 12 CVars to fail with priority errors
+ *    - This class uses ECVF_SetByConsole priority to successfully override during test
+ *
+ * 2. PIE Screenshot Delegate Mismatch:
+ *    - PIE uses FEditorViewportClient::ProcessScreenShots() which fires the static
+ *      FScreenshotRequest::OnScreenshotCaptured() delegate
+ *    - Base class only registers to the instance UGameViewportClient::OnScreenshotCaptured()
+ *    - This class registers to BOTH delegates to support PIE and standalone execution
+ *
+ * Usage:
+ *    Use this class instead of AScreenshotFunctionalTest when creating screenshot tests
+ *    for Avalanche/Motion Design levels. It's only available when the Avalanche plugin
+ *    is enabled.
+ *
+ * Note:
+ *    CVars are NOT restored after test completion to avoid priority stickiness issues.
+ *    Test-friendly values (AA=0, ScreenPercentage=100, MotionBlur=0) remain in place.
+ */
+UCLASS(Blueprintable, MinimalAPI, DisplayName = "Screenshot Functional Test Motion Design")
+class AScreenshotFunctionalTestAvalanche : public AScreenshotFunctionalTest
+{
+	GENERATED_BODY()
+
+public:
+	AVALANCHEFUNCTIONALTEST_API AScreenshotFunctionalTestAvalanche(const FObjectInitializer& ObjectInitializer);
+
+protected:
+	// Override to set up CVars with explicit priority before base class setup
+	AVALANCHEFUNCTIONALTEST_API virtual void PrepareTest() override;
+
+	// Override to register to both instance and static screenshot delegates for PIE support
+	AVALANCHEFUNCTIONALTEST_API virtual void RequestScreenshot() override;
+	AVALANCHEFUNCTIONALTEST_API virtual void OnScreenShotCaptured(int32 InSizeX, int32 InSizeY, const TArray<FColor>& InImageData) override;
+
+private:
+	/**
+	 * Helper struct to manage CVar overrides at Console priority.
+	 * Note: CVars are NOT restored after test to avoid priority stickiness issues.
+	 * Test-friendly values (AA=0, ScreenPercentage=100) are left in place.
+	 */
+	template<typename T>
+	struct FScreenshotTestAvalancheCVarOverride
+	{
+		FScreenshotTestAvalancheCVarOverride() = default;
+
+		FScreenshotTestAvalancheCVarOverride(const FString& InCVarName)
+			: CVarName(InCVarName)
+		{
+		}
+
+		void SetOverride(T NewValue)
+		{
+			IConsoleVariable* CVar = IConsoleManager::Get().FindConsoleVariable(*CVarName);
+			if (CVar)
+			{
+				// Use Console priority to override Scalability/Constructor/Code settings
+				// High enough to actually take effect, but won't be restored to avoid
+				// priority stickiness (Console priority would prevent lower priorities from working)
+				CVar->Set(NewValue, ECVF_SetByConsole);
+			}
+		}
+
+	private:
+		FString CVarName;
+	};
+
+	// CVar overrides managed with explicit priority
+	FScreenshotTestAvalancheCVarOverride<int32> AntiAliasingOverride;
+	FScreenshotTestAvalancheCVarOverride<int32> AutoExposureOverride;
+	FScreenshotTestAvalancheCVarOverride<int32> MotionBlurOverride;
+	FScreenshotTestAvalancheCVarOverride<int32> MotionBlurQualityOverride;
+	FScreenshotTestAvalancheCVarOverride<int32> ContactShadowsOverride;
+	FScreenshotTestAvalancheCVarOverride<int32> ScreenSpaceReflectionQualityOverride;
+	FScreenshotTestAvalancheCVarOverride<int32> EyeAdaptationQualityOverride;
+	FScreenshotTestAvalancheCVarOverride<float> TonemapperGammaOverride;
+	FScreenshotTestAvalancheCVarOverride<float> ScreenPercentageOverride;
+	FScreenshotTestAvalancheCVarOverride<int32> DynamicResTestScreenPercentageOverride;
+	FScreenshotTestAvalancheCVarOverride<int32> DynamicResOperationModeOverride;
+	FScreenshotTestAvalancheCVarOverride<float> SecondaryScreenPercentageOverride;
+
+	// Apply Avalanche-compatible rendering overrides
+	void ApplyAvalancheRenderingOverrides();
+};
