Motion Design: fixed issue with functional test module not being built fully and causing issues on startup

#Synced-CL 48854462

diff --git a/Engine/Plugins/VirtualProduction/Avalanche/Source/AvalancheFunctionalTest/Private/AvalancheFunctionalTestModule.cpp b/Engine/Plugins/VirtualProduction/Avalanche/Source/AvalancheFunctionalTest/Private/AvalancheFunctionalTestModule.cpp
new file mode 100644
--- /dev/null
+++ b/Engine/Plugins/VirtualProduction/Avalanche/Source/AvalancheFunctionalTest/Private/AvalancheFunctionalTestModule.cpp
@@ -0,0 +1,5 @@
+// Copyright Epic Games, Inc. All Rights Reserved.
+
+#include "Modules/ModuleManager.h"
+
+IMPLEMENT_MODULE(FDefaultModuleImpl, AvalancheFunctionalTest)
