Motion Design: added a base context class for common functionality

#Synced-CL 48854462

diff --git a/Engine/Plugins/VirtualProduction/Avalanche/Source/AvalancheMaterial/Private/MaterialBridge/Context/AvaMaterialBridgeApplyStateContext.cpp b/Engine/Plugins/VirtualProduction/Avalanche/Source/AvalancheMaterial/Private/MaterialBridge/Context/AvaMaterialBridgeApplyStateContext.cpp
deleted file mode 100644
diff --git a/Engine/Plugins/VirtualProduction/Avalanche/Source/AvalancheMaterial/Private/MaterialBridge/Context/AvaMaterialBridgeContext.cpp b/Engine/Plugins/VirtualProduction/Avalanche/Source/AvalancheMaterial/Private/MaterialBridge/Context/AvaMaterialBridgeContext.cpp
new file mode 100644
--- /dev/null
+++ b/Engine/Plugins/VirtualProduction/Avalanche/Source/AvalancheMaterial/Private/MaterialBridge/Context/AvaMaterialBridgeContext.cpp
@@ -0,0 +1,80 @@
+// Copyright Epic Games, Inc. All Rights Reserved.
+
+#include "MaterialBridge/Context/AvaMaterialBridgeContext.h"
+#include "MaterialBridge/AvaMaterialBridgeRegistry.h"
+#include "MaterialBridge/Context/AvaMaterialBridgeApplyStateContext.h"
+#include "MaterialBridge/Context/AvaMaterialBridgeReadSlotContext.h"
+#include "MaterialBridge/Context/AvaMaterialBridgeStoreStateContext.h"
+#include "MaterialBridge/Context/AvaMaterialBridgeWriteSlotContext.h"
+
+namespace UE::Ava
+{
+
+FMaterialBridgeBaseOptions::FMaterialBridgeBaseOptions()
+	: MaterialBridgeRegistry(&FMaterialBridgeRegistry::Get())
+{
+}
+
+template<typename InContextType>
+TMaterialBridgeContext<InContextType>::TMaterialBridgeContext(FDataViewType InMaterialContainer)
+	: MaterialContainer(InMaterialContainer)
+{
+}
+
+template<typename InContextType>
+TMaterialBridgeContext<InContextType>::TMaterialBridgeContext(UObjectType* InMaterialContainer)
+	: MaterialContainer(InMaterialContainer)
+{
+}
+
+template<typename InContextType>
+TMaterialBridgeContext<InContextType>::TMaterialBridgeContext(UObjectType* InMaterialContainer, TNotNull<const InContextType*> InParentContext)
+	: MaterialContainer(InMaterialContainer)
+	, ParentContext(InParentContext)
+{
+}
+
+template<typename InContextType>
+const InContextType& TMaterialBridgeContext<InContextType>::GetTopmostContext() const
+{
+	const InContextType* TopmostContext = static_cast<const InContextType*>(this);
+	while (TopmostContext->ParentContext)
+	{
+		TopmostContext = TopmostContext->ParentContext;
+	}
+	return *TopmostContext;
+}
+
+template<typename InContextType>
+TMaterialBridgeContext<InContextType>::UObjectType* TMaterialBridgeContext<InContextType>::GetMaterialContainerObject() const
+{
+	if constexpr (bWithConstMaterialContainer)
+	{
+		if (UObjectType* MaterialContainerObject = MaterialContainer.template GetPtr<UObjectType>())
+		{
+			return MaterialContainerObject;
+		}
+	}
+	else
+	{
+		if (UObjectType* MaterialContainerObject = MaterialContainer.template GetMutablePtr<UObjectType>())
+		{
+			return MaterialContainerObject;
+		}
+	}
+
+	if (ParentContext)
+	{
+		return ParentContext->GetMaterialContainerObject();
+	}
+	return nullptr;
+}
+
+// Explicit template instantiations
+template class TMaterialBridgeContext<FMaterialBridgeApplyStateContext>;
+template class TMaterialBridgeContext<FMaterialBridgeReadSlotContext>;
+template class TMaterialBridgeContext<FMaterialBridgeStoreStateContext>;
+template class TMaterialBridgeContext<FMaterialBridgeWriteSlotContext>;
+
+} // UE::Ava
+
diff --git a/Engine/Plugins/VirtualProduction/Avalanche/Source/AvalancheMaterial/Private/MaterialBridge/Context/AvaMaterialBridgeReadSlotContext.cpp b/Engine/Plugins/VirtualProduction/Avalanche/Source/AvalancheMaterial/Private/MaterialBridge/Context/AvaMaterialBridgeReadSlotContext.cpp
deleted file mode 100644
diff --git a/Engine/Plugins/VirtualProduction/Avalanche/Source/AvalancheMaterial/Private/MaterialBridge/Context/AvaMaterialBridgeStoreStateContext.cpp b/Engine/Plugins/VirtualProduction/Avalanche/Source/AvalancheMaterial/Private/MaterialBridge/Context/AvaMaterialBridgeStoreStateContext.cpp
deleted file mode 100644
diff --git a/Engine/Plugins/VirtualProduction/Avalanche/Source/AvalancheMaterial/Private/MaterialBridge/Context/AvaMaterialBridgeWriteSlotContext.cpp b/Engine/Plugins/VirtualProduction/Avalanche/Source/AvalancheMaterial/Private/MaterialBridge/Context/AvaMaterialBridgeWriteSlotContext.cpp
deleted file mode 100644
diff --git a/Engine/Plugins/VirtualProduction/Avalanche/Source/AvalancheMaterial/Public/MaterialBridge/AvaMaterialBridge.h b/Engine/Plugins/VirtualProduction/Avalanche/Source/AvalancheMaterial/Public/MaterialBridge/AvaMaterialBridge.h
--- a/Engine/Plugins/VirtualProduction/Avalanche/Source/AvalancheMaterial/Public/MaterialBridge/AvaMaterialBridge.h
+++ b/Engine/Plugins/VirtualProduction/Avalanche/Source/AvalancheMaterial/Public/MaterialBridge/AvaMaterialBridge.h
@@ -15,14 +15,14 @@ struct FAvaMaterialContainerState;
 
 namespace UE::Ava
 {
+	class FMaterialBridgeApplyStateContext;
 	class FMaterialBridgeReadSlot;
 	class FMaterialBridgeReadSlotContext;
+	class FMaterialBridgeStoreStateContext;
 	class FMaterialBridgeWriteSlot;
 	class FMaterialBridgeWriteSlotContext;
-	struct FMaterialBridgeApplyStateContext;
 	struct FMaterialBridgeApplyStateOptions;
 	struct FMaterialBridgeReadSlotOptions;
-	struct FMaterialBridgeStoreStateContext;
 	struct FMaterialBridgeStoreStateOptions;
 	struct FMaterialBridgeWriteSlotOptions;
 }
diff --git a/Engine/Plugins/VirtualProduction/Avalanche/Source/AvalancheMaterial/Public/MaterialBridge/Context/AvaMaterialBridgeApplyStateContext.h b/Engine/Plugins/VirtualProduction/Avalanche/Source/AvalancheMaterial/Public/MaterialBridge/Context/AvaMaterialBridgeApplyStateContext.h
--- a/Engine/Plugins/VirtualProduction/Avalanche/Source/AvalancheMaterial/Public/MaterialBridge/Context/AvaMaterialBridgeApplyStateContext.h
+++ b/Engine/Plugins/VirtualProduction/Avalanche/Source/AvalancheMaterial/Public/MaterialBridge/Context/AvaMaterialBridgeApplyStateContext.h
@@ -2,49 +2,35 @@
 
 #pragma once
 
-#include "AvaDataView.h"
-
-#define UE_API AVALANCHEMATERIAL_API
+#include "AvaMaterialBridgeContext.h"
 
 namespace UE::Ava
 {
-	class FMaterialBridgeRegistry;
-	struct FMaterialBridgeApplyStateContext;
+	class FMaterialBridgeApplyStateContext;
 }
 
 namespace UE::Ava
 {
 
 /** Additional options when applying state to a material container  */
-struct FMaterialBridgeApplyStateOptions
+struct FMaterialBridgeApplyStateOptions : public FMaterialBridgeBaseOptions
 {
-	UE_API FMaterialBridgeApplyStateOptions();
+};
 
-	/** The registry used to get material bridges. Defaults to the global registry. */
-	TNotNull<const FMaterialBridgeRegistry*> MaterialBridgeRegistry;
+template<>
+struct TMaterialBridgeContextFlags<FMaterialBridgeApplyStateContext> : public TMaterialBridgeContextFlagsBase<FMaterialBridgeApplyStateContext>
+{
+	enum
+	{
+		WithConstMaterialContainer = false,
+	};
 };
 
 /** Short-lived context information when applying state to a material container */
-struct FMaterialBridgeApplyStateContext
+class FMaterialBridgeApplyStateContext : public TMaterialBridgeContext<FMaterialBridgeApplyStateContext>
 {
-	FMaterialBridgeApplyStateContext() = default;
-
-	explicit UE_API FMaterialBridgeApplyStateContext(UObject* InMaterialContainer);
-	explicit UE_API FMaterialBridgeApplyStateContext(UObject* InMaterialContainer, TNotNull<const FMaterialBridgeApplyStateContext*> InParentContext);
-
-	/** Gets the first material container that is a UObject, or null if the data is not nested inside a UObject */
-	UE_API UObject* GetMaterialContainerObject() const;
-
-	/** The object/struct containing the material */
-	FDataView MaterialContainer;
-
-	/** Parent context. Non-null if in a nested call. Should always be valid checked */
-	const FMaterialBridgeApplyStateContext* ParentContext = nullptr;
-
-private:
-	friend struct FMaterialBridgeApplyStateChangeScope;
+public:
+	using TMaterialBridgeContext<FMaterialBridgeApplyStateContext>::TMaterialBridgeContext;
 };
 
 } // UE::Ava
-
-#undef UE_API
diff --git a/Engine/Plugins/VirtualProduction/Avalanche/Source/AvalancheMaterial/Public/MaterialBridge/Context/AvaMaterialBridgeContext.h b/Engine/Plugins/VirtualProduction/Avalanche/Source/AvalancheMaterial/Public/MaterialBridge/Context/AvaMaterialBridgeContext.h
new file mode 100644
--- /dev/null
+++ b/Engine/Plugins/VirtualProduction/Avalanche/Source/AvalancheMaterial/Public/MaterialBridge/Context/AvaMaterialBridgeContext.h
@@ -0,0 +1,73 @@
+// Copyright Epic Games, Inc. All Rights Reserved.
+
+#pragma once
+
+#include "AvaDataView.h"
+
+#define UE_API AVALANCHEMATERIAL_API
+
+class UObject;
+
+namespace UE::Ava
+{
+	class FMaterialBridgeRegistry;
+}
+
+namespace UE::Ava
+{
+
+/** Base for all options structs */
+struct FMaterialBridgeBaseOptions
+{
+	UE_API FMaterialBridgeBaseOptions();
+
+	/** The registry used to get material bridges. Defaults to the global registry. */
+	TNotNull<const FMaterialBridgeRegistry*> MaterialBridgeRegistry;
+};
+
+template<typename InContextType>
+struct TMaterialBridgeContextFlagsBase
+{
+	enum
+	{
+		WithConstMaterialContainer = true,
+	};
+};
+
+template<typename InContextType>
+struct TMaterialBridgeContextFlags : public TMaterialBridgeContextFlagsBase<InContextType>
+{
+};
+
+/**  Base context class containing common functionality shared among all material bridge contexts */
+template<typename InContextType>
+class TMaterialBridgeContext
+{
+public:
+	static constexpr bool bWithConstMaterialContainer = TMaterialBridgeContextFlags<InContextType>::WithConstMaterialContainer;
+
+	using UObjectType   = std::conditional_t<bWithConstMaterialContainer, const UObject, UObject>;
+	using FDataViewType = std::conditional_t<bWithConstMaterialContainer, FConstDataView, FDataView>;
+
+	TMaterialBridgeContext() = default;
+
+	UE_API explicit TMaterialBridgeContext(FDataViewType InMaterialContainer);
+	UE_API explicit TMaterialBridgeContext(UObjectType* InMaterialContainer);
+	UE_API explicit TMaterialBridgeContext(UObjectType* InMaterialContainer, TNotNull<const InContextType*> InParentContext);
+
+	/** Gets the topmost slot context in the context chain */
+	UE_API const InContextType& GetTopmostContext() const;
+
+	/** Gets the first material container that is a UObject */
+	UE_API UObjectType* GetMaterialContainerObject() const;
+
+	/** The object containing the material */
+	FDataViewType MaterialContainer;
+
+	/** Parent context. Non-null if in a nested call. Should always be valid checked */
+	const InContextType* ParentContext = nullptr;
+};
+
+} // UE::Ava
+
+#undef UE_API
diff --git a/Engine/Plugins/VirtualProduction/Avalanche/Source/AvalancheMaterial/Public/MaterialBridge/Context/AvaMaterialBridgeReadSlotContext.h b/Engine/Plugins/VirtualProduction/Avalanche/Source/AvalancheMaterial/Public/MaterialBridge/Context/AvaMaterialBridgeReadSlotContext.h
--- a/Engine/Plugins/VirtualProduction/Avalanche/Source/AvalancheMaterial/Public/MaterialBridge/Context/AvaMaterialBridgeReadSlotContext.h
+++ b/Engine/Plugins/VirtualProduction/Avalanche/Source/AvalancheMaterial/Public/MaterialBridge/Context/AvaMaterialBridgeReadSlotContext.h
@@ -2,14 +2,11 @@
 
 #pragma once
 
-#include "AvaDataView.h"
 #include "AvaEnums.h"
-
-#define UE_API AVALANCHEMATERIAL_API
+#include "AvaMaterialBridgeContext.h"
 
 namespace UE::Ava
 {
-	class FMaterialBridgeRegistry;
 	class FMaterialBridgeReadSlotContext;
 
 } // UE::Ava
@@ -18,13 +15,8 @@ namespace UE::Ava
 {
 
 /** Options to how read access should take place */
-struct FMaterialBridgeReadSlotOptions
+struct FMaterialBridgeReadSlotOptions : public FMaterialBridgeBaseOptions
 {
-	UE_API FMaterialBridgeReadSlotOptions();
-
-	/** The registry to use to visit any materials of sub objects. Defaults to the global registry. */
-	TNotNull<const FMaterialBridgeRegistry*> MaterialBridgeRegistry;
-
 	/** Delegate called when a container is not yet complete and cannot access its materials. Only called when bTrySkipWaitOnContainerCompletion is true. */
 	TDelegate<EControlFlow(const FMaterialBridgeReadSlotContext&)> OnContainerPendingCompletion;
 
@@ -36,21 +28,10 @@ struct FMaterialBridgeReadSlotOptions
 };
 
 /** Short-lived context information when accessing a material slot for read  */
-class FMaterialBridgeReadSlotContext
+class FMaterialBridgeReadSlotContext : public TMaterialBridgeContext<FMaterialBridgeReadSlotContext>
 {
 public:
-	FMaterialBridgeReadSlotContext() = default;
-
-	explicit UE_API FMaterialBridgeReadSlotContext(const UObject* InContainer);
-	explicit UE_API FMaterialBridgeReadSlotContext(const UObject* InContainer, TNotNull<const FMaterialBridgeReadSlotContext*> InParentContext);
-
-	/** Container of the visited material */
-	FConstDataView MaterialContainer;
-
-	/** Parent of this context for additional information (e.g. container of the container) */
-	const FMaterialBridgeReadSlotContext* ParentContext = nullptr;
+	using TMaterialBridgeContext<FMaterialBridgeReadSlotContext>::TMaterialBridgeContext;
 };
 
 } // UE::Ava
-
-#undef UE_API
diff --git a/Engine/Plugins/VirtualProduction/Avalanche/Source/AvalancheMaterial/Public/MaterialBridge/Context/AvaMaterialBridgeStoreStateContext.h b/Engine/Plugins/VirtualProduction/Avalanche/Source/AvalancheMaterial/Public/MaterialBridge/Context/AvaMaterialBridgeStoreStateContext.h
--- a/Engine/Plugins/VirtualProduction/Avalanche/Source/AvalancheMaterial/Public/MaterialBridge/Context/AvaMaterialBridgeStoreStateContext.h
+++ b/Engine/Plugins/VirtualProduction/Avalanche/Source/AvalancheMaterial/Public/MaterialBridge/Context/AvaMaterialBridgeStoreStateContext.h
@@ -2,46 +2,21 @@
 
 #pragma once
 
-#include "AvaDataView.h"
-
-#define UE_API AVALANCHEMATERIAL_API
-
-namespace UE::Ava
-{
-	class FMaterialBridgeRegistry;
-	struct FMaterialBridgeStoreStateContext;
-}
+#include "AvaMaterialBridgeContext.h"
 
 namespace UE::Ava
 {
 
 /** Additional options when applying/storing material container state  */
-struct FMaterialBridgeStoreStateOptions
+struct FMaterialBridgeStoreStateOptions : public FMaterialBridgeBaseOptions
 {
-	UE_API FMaterialBridgeStoreStateOptions();
-
-	/** The registry used to get material bridges. Defaults to the global registry. */
-	TNotNull<const FMaterialBridgeRegistry*> MaterialBridgeRegistry;
 };
 
 /** Short-lived context information when storing the state of a material container */
-struct FMaterialBridgeStoreStateContext
+class FMaterialBridgeStoreStateContext : public TMaterialBridgeContext<FMaterialBridgeStoreStateContext>
 {
-	FMaterialBridgeStoreStateContext() = default;
-
-	explicit UE_API FMaterialBridgeStoreStateContext(const UObject* InMaterialContainer);
-	explicit UE_API FMaterialBridgeStoreStateContext(const UObject* InMaterialContainer, TNotNull<const FMaterialBridgeStoreStateContext*> InParentContext);
-
-	/** Gets the first material container that is a UObject */
-	UE_API const UObject* GetMaterialContainerObject() const;
-
-	/** The object containing the material */
-	FConstDataView MaterialContainer;
-
-	/** Parent context. Non-null if in a nested call. Should always be valid checked */
-	const FMaterialBridgeStoreStateContext* ParentContext = nullptr;
+public:
+	using TMaterialBridgeContext<FMaterialBridgeStoreStateContext>::TMaterialBridgeContext;
 };
-	
-} // UE::Ava
 
-#undef UE_API
+} // UE::Ava
diff --git a/Engine/Plugins/VirtualProduction/Avalanche/Source/AvalancheMaterial/Public/MaterialBridge/Context/AvaMaterialBridgeWriteSlotContext.h b/Engine/Plugins/VirtualProduction/Avalanche/Source/AvalancheMaterial/Public/MaterialBridge/Context/AvaMaterialBridgeWriteSlotContext.h
--- a/Engine/Plugins/VirtualProduction/Avalanche/Source/AvalancheMaterial/Public/MaterialBridge/Context/AvaMaterialBridgeWriteSlotContext.h
+++ b/Engine/Plugins/VirtualProduction/Avalanche/Source/AvalancheMaterial/Public/MaterialBridge/Context/AvaMaterialBridgeWriteSlotContext.h
@@ -3,13 +3,10 @@
 #pragma once
 
 #include "AvaEnums.h"
-#include "AvaDataView.h"
-
-#define UE_API AVALANCHEMATERIAL_API
+#include "AvaMaterialBridgeContext.h"
 
 namespace UE::Ava
 {
-	class FMaterialBridgeRegistry;
 	class FMaterialBridgeWriteSlotContext;
 
 } // UE::Ava
@@ -18,13 +15,8 @@ namespace UE::Ava
 {
 
 /** Options used when accessing a slot for write */
-struct FMaterialBridgeWriteSlotOptions
+struct FMaterialBridgeWriteSlotOptions : public FMaterialBridgeBaseOptions
 {
-	UE_API FMaterialBridgeWriteSlotOptions();
-
-	/** The registry to use to visit any materials of sub objects. Defaults to the global registry. */
-	TNotNull<const FMaterialBridgeRegistry*> MaterialBridgeRegistry;
-
 	/** Delegate called when a container is not yet complete and cannot access its materials. Only called when bTrySkipWaitOnContainerCompletion is true. */
 	TDelegate<EControlFlow(const FMaterialBridgeWriteSlotContext&)> OnContainerPendingCompletion;
 
@@ -35,30 +27,20 @@ struct FMaterialBridgeWriteSlotOptions
 	bool bTrySkipWaitOnContainerCompletion = false;
 };
 
+template<>
+struct TMaterialBridgeContextFlags<FMaterialBridgeWriteSlotContext> : public TMaterialBridgeContextFlagsBase<FMaterialBridgeWriteSlotContext>
+{
+	enum
+	{
+		WithConstMaterialContainer = false,
+	};
+};
+
 /** Short-lived context information when accessing a material slot for write  */
-class FMaterialBridgeWriteSlotContext
+class FMaterialBridgeWriteSlotContext : public TMaterialBridgeContext<FMaterialBridgeWriteSlotContext>
 {
 public:
-	FMaterialBridgeWriteSlotContext() = default;
-
-	explicit UE_API FMaterialBridgeWriteSlotContext(FDataView InMaterialContainer);
-	explicit UE_API FMaterialBridgeWriteSlotContext(UObject* InMaterialContainer);
-	explicit UE_API FMaterialBridgeWriteSlotContext(UObject* InMaterialContainer, TNotNull<const FMaterialBridgeWriteSlotContext*> InParentContext);
-
-	/** Gets the first UObject container in the context chain */
-	UE_API UObject* GetMaterialContainerObject() const;
-
-	/** Container of the accessed material slot */
-	FDataView MaterialContainer;
-
-	/** Parent context. Non-null if in a nested call. Should always be valid checked */
-	const FMaterialBridgeWriteSlotContext* ParentContext = nullptr;
-
-private:
-	/** Keeps track of whether pre change has been called for this context */
-	mutable bool bPreChangeCalled = false;
+	using TMaterialBridgeContext<FMaterialBridgeWriteSlotContext>::TMaterialBridgeContext;
 };
 
 } // UE::Ava
-
-#undef UE_API
