Motion Design: moved material slot files to their own directory for clarity

#Synced-CL 48854462

diff --git a/Engine/Plugins/VirtualProduction/Avalanche/Source/AvalancheMaterial/Private/MaterialBridge/AvaMaterialBridgeReadSlot.cpp b/Engine/Plugins/VirtualProduction/Avalanche/Source/AvalancheMaterial/Private/MaterialBridge/AvaMaterialBridgeReadSlot.cpp
deleted file mode 100644
diff --git a/Engine/Plugins/VirtualProduction/Avalanche/Source/AvalancheMaterial/Private/MaterialBridge/AvaMaterialBridgeWriteSlot.cpp b/Engine/Plugins/VirtualProduction/Avalanche/Source/AvalancheMaterial/Private/MaterialBridge/AvaMaterialBridgeWriteSlot.cpp
deleted file mode 100644
diff --git a/Engine/Plugins/VirtualProduction/Avalanche/Source/AvalancheMaterial/Private/MaterialBridge/Common/AvaDecalComponentMaterialBridge.cpp b/Engine/Plugins/VirtualProduction/Avalanche/Source/AvalancheMaterial/Private/MaterialBridge/Common/AvaDecalComponentMaterialBridge.cpp
--- a/Engine/Plugins/VirtualProduction/Avalanche/Source/AvalancheMaterial/Private/MaterialBridge/Common/AvaDecalComponentMaterialBridge.cpp
+++ b/Engine/Plugins/VirtualProduction/Avalanche/Source/AvalancheMaterial/Private/MaterialBridge/Common/AvaDecalComponentMaterialBridge.cpp
@@ -3,12 +3,12 @@
 #include "MaterialBridge/Common/AvaDecalComponentMaterialBridge.h"
 #include "AvaDataView.h"
 #include "Components/DecalComponent.h"
-#include "MaterialBridge/AvaMaterialBridgeReadSlot.h"
-#include "MaterialBridge/AvaMaterialBridgeWriteSlot.h"
 #include "MaterialBridge/Context/AvaMaterialBridgeApplyStateContext.h"
 #include "MaterialBridge/Context/AvaMaterialBridgeReadSlotContext.h"
 #include "MaterialBridge/Context/AvaMaterialBridgeStoreStateContext.h"
 #include "MaterialBridge/Context/AvaMaterialBridgeWriteSlotContext.h"
+#include "MaterialBridge/Slot/AvaMaterialBridgeReadSlot.h"
+#include "MaterialBridge/Slot/AvaMaterialBridgeWriteSlot.h"
 
 namespace UE::Ava
 {
diff --git a/Engine/Plugins/VirtualProduction/Avalanche/Source/AvalancheMaterial/Private/MaterialBridge/Common/AvaPrimitiveComponentMaterialBridge.cpp b/Engine/Plugins/VirtualProduction/Avalanche/Source/AvalancheMaterial/Private/MaterialBridge/Common/AvaPrimitiveComponentMaterialBridge.cpp
--- a/Engine/Plugins/VirtualProduction/Avalanche/Source/AvalancheMaterial/Private/MaterialBridge/Common/AvaPrimitiveComponentMaterialBridge.cpp
+++ b/Engine/Plugins/VirtualProduction/Avalanche/Source/AvalancheMaterial/Private/MaterialBridge/Common/AvaPrimitiveComponentMaterialBridge.cpp
@@ -2,12 +2,12 @@
 
 #include "MaterialBridge/Common/AvaPrimitiveComponentMaterialBridge.h"
 #include "Components/PrimitiveComponent.h"
-#include "MaterialBridge/AvaMaterialBridgeReadSlot.h"
-#include "MaterialBridge/AvaMaterialBridgeWriteSlot.h"
 #include "MaterialBridge/Context/AvaMaterialBridgeApplyStateContext.h"
 #include "MaterialBridge/Context/AvaMaterialBridgeReadSlotContext.h"
 #include "MaterialBridge/Context/AvaMaterialBridgeStoreStateContext.h"
 #include "MaterialBridge/Context/AvaMaterialBridgeWriteSlotContext.h"
+#include "MaterialBridge/Slot/AvaMaterialBridgeReadSlot.h"
+#include "MaterialBridge/Slot/AvaMaterialBridgeWriteSlot.h"
 #include "Misc/EnumerateRange.h"
 
 namespace UE::Ava
diff --git a/Engine/Plugins/VirtualProduction/Avalanche/Source/AvalancheMaterial/Private/MaterialBridge/Common/AvaStaticMeshMaterialBridge.cpp b/Engine/Plugins/VirtualProduction/Avalanche/Source/AvalancheMaterial/Private/MaterialBridge/Common/AvaStaticMeshMaterialBridge.cpp
--- a/Engine/Plugins/VirtualProduction/Avalanche/Source/AvalancheMaterial/Private/MaterialBridge/Common/AvaStaticMeshMaterialBridge.cpp
+++ b/Engine/Plugins/VirtualProduction/Avalanche/Source/AvalancheMaterial/Private/MaterialBridge/Common/AvaStaticMeshMaterialBridge.cpp
@@ -3,10 +3,10 @@
 #include "MaterialBridge/Common/AvaStaticMeshMaterialBridge.h"
 #include "Engine/StaticMesh.h"
 #include "MaterialBridge/AvaMaterialBridgeLog.h"
-#include "MaterialBridge/AvaMaterialBridgeReadSlot.h"
-#include "MaterialBridge/AvaMaterialBridgeWriteSlot.h"
 #include "MaterialBridge/Context/AvaMaterialBridgeReadSlotContext.h"
 #include "MaterialBridge/Context/AvaMaterialBridgeWriteSlotContext.h"
+#include "MaterialBridge/Slot/AvaMaterialBridgeReadSlot.h"
+#include "MaterialBridge/Slot/AvaMaterialBridgeWriteSlot.h"
 
 namespace UE::Ava
 {
diff --git a/Engine/Plugins/VirtualProduction/Avalanche/Source/AvalancheMaterial/Private/MaterialBridge/Slot/AvaMaterialBridgeReadSlot.cpp b/Engine/Plugins/VirtualProduction/Avalanche/Source/AvalancheMaterial/Private/MaterialBridge/Slot/AvaMaterialBridgeReadSlot.cpp
new file mode 100644
--- /dev/null
+++ b/Engine/Plugins/VirtualProduction/Avalanche/Source/AvalancheMaterial/Private/MaterialBridge/Slot/AvaMaterialBridgeReadSlot.cpp
@@ -0,0 +1,13 @@
+// Copyright Epic Games, Inc. All Rights Reserved.
+
+#include "MaterialBridge/Slot/AvaMaterialBridgeReadSlot.h"
+
+namespace UE::Ava
+{
+	
+FMaterialBridgeReadSlot::FMaterialBridgeReadSlot(UMaterialInterface* InMaterial)
+	: Material(InMaterial)
+{
+}
+
+} // UE::Ava
diff --git a/Engine/Plugins/VirtualProduction/Avalanche/Source/AvalancheMaterial/Private/MaterialBridge/Slot/AvaMaterialBridgeWriteSlot.cpp b/Engine/Plugins/VirtualProduction/Avalanche/Source/AvalancheMaterial/Private/MaterialBridge/Slot/AvaMaterialBridgeWriteSlot.cpp
new file mode 100644
--- /dev/null
+++ b/Engine/Plugins/VirtualProduction/Avalanche/Source/AvalancheMaterial/Private/MaterialBridge/Slot/AvaMaterialBridgeWriteSlot.cpp
@@ -0,0 +1,35 @@
+// Copyright Epic Games, Inc. All Rights Reserved.
+
+#include "MaterialBridge/Slot/AvaMaterialBridgeWriteSlot.h"
+
+namespace UE::Ava
+{
+
+FMaterialBridgeWriteSlot::FMaterialBridgeWriteSlot(UMaterialInterface* InMaterial)
+	: Material(InMaterial)
+{
+}
+
+void FMaterialBridgeWriteSlot::SetMaterial(UMaterialInterface* InMaterial)
+{
+	if (InMaterial != Material)
+	{
+		SetMaterialInternal(InMaterial);
+		OnMaterialChanged();
+	}
+}
+
+void FMaterialBridgeWriteSlot::RequestFeatures(EMaterialBridgeRequestedFeatures InFeatures)
+{
+	if (InFeatures != EMaterialBridgeRequestedFeatures::None)
+	{
+		OnRequestFeatures(InFeatures);
+	}
+}
+
+void FMaterialBridgeWriteSlot::SetMaterialInternal(UMaterialInterface* InMaterial)
+{
+	Material = InMaterial;
+}
+
+} // UE::Ava
diff --git a/Engine/Plugins/VirtualProduction/Avalanche/Source/AvalancheMaterial/Private/MaterialCache/AvaMaterialCacheHelper.cpp b/Engine/Plugins/VirtualProduction/Avalanche/Source/AvalancheMaterial/Private/MaterialCache/AvaMaterialCacheHelper.cpp
--- a/Engine/Plugins/VirtualProduction/Avalanche/Source/AvalancheMaterial/Private/MaterialCache/AvaMaterialCacheHelper.cpp
+++ b/Engine/Plugins/VirtualProduction/Avalanche/Source/AvalancheMaterial/Private/MaterialCache/AvaMaterialCacheHelper.cpp
@@ -3,9 +3,9 @@
 #include "MaterialCache/AvaMaterialCacheHelper.h"
 #include "AvaDataView.h"
 #include "MaterialBridge/AvaMaterialBridge.h"
-#include "MaterialBridge/AvaMaterialBridgeReadSlot.h"
 #include "MaterialBridge/AvaMaterialBridgeRegistry.h"
 #include "MaterialBridge/Context/AvaMaterialBridgeReadSlotContext.h"
+#include "MaterialBridge/Slot/AvaMaterialBridgeReadSlot.h"
 #include "MaterialCache/AvaMaterialCacheLog.h"
 #include "MaterialCache/AvaMaterialCacheSettings.h"
 #include "MaterialShared.h"
diff --git a/Engine/Plugins/VirtualProduction/Avalanche/Source/AvalancheMaterial/Public/MaterialBridge/AvaMaterialBridgeReadSlot.h b/Engine/Plugins/VirtualProduction/Avalanche/Source/AvalancheMaterial/Public/MaterialBridge/AvaMaterialBridgeReadSlot.h
deleted file mode 100644
diff --git a/Engine/Plugins/VirtualProduction/Avalanche/Source/AvalancheMaterial/Public/MaterialBridge/AvaMaterialBridgeWriteSlot.h b/Engine/Plugins/VirtualProduction/Avalanche/Source/AvalancheMaterial/Public/MaterialBridge/AvaMaterialBridgeWriteSlot.h
deleted file mode 100644
diff --git a/Engine/Plugins/VirtualProduction/Avalanche/Source/AvalancheMaterial/Public/MaterialBridge/Slot/AvaMaterialBridgeReadSlot.h b/Engine/Plugins/VirtualProduction/Avalanche/Source/AvalancheMaterial/Public/MaterialBridge/Slot/AvaMaterialBridgeReadSlot.h
new file mode 100644
--- /dev/null
+++ b/Engine/Plugins/VirtualProduction/Avalanche/Source/AvalancheMaterial/Public/MaterialBridge/Slot/AvaMaterialBridgeReadSlot.h
@@ -0,0 +1,33 @@
+// Copyright Epic Games, Inc. All Rights Reserved.
+
+#pragma once
+
+#include "UObject/ObjectPtr.h"
+
+#define UE_API AVALANCHEMATERIAL_API
+
+class UMaterialInterface;
+
+namespace UE::Ava
+{
+
+/** Short-lived representation of an accessed material slot for read */
+class FMaterialBridgeReadSlot
+{
+public:
+	UE_API explicit FMaterialBridgeReadSlot(UMaterialInterface* InMaterial);
+
+	/** Gets the material for this slot */
+	UMaterialInterface* GetMaterial() const
+	{
+		return Material;
+	}
+
+protected:
+	/** Material in this slot */
+	TObjectPtr<UMaterialInterface> Material;
+};
+
+} // UE::Ava
+
+#undef UE_API
diff --git a/Engine/Plugins/VirtualProduction/Avalanche/Source/AvalancheMaterial/Public/MaterialBridge/Slot/AvaMaterialBridgeWriteSlot.h b/Engine/Plugins/VirtualProduction/Avalanche/Source/AvalancheMaterial/Public/MaterialBridge/Slot/AvaMaterialBridgeWriteSlot.h
new file mode 100644
--- /dev/null
+++ b/Engine/Plugins/VirtualProduction/Avalanche/Source/AvalancheMaterial/Public/MaterialBridge/Slot/AvaMaterialBridgeWriteSlot.h
@@ -0,0 +1,60 @@
+// Copyright Epic Games, Inc. All Rights Reserved.
+
+#pragma once
+
+#include "MaterialBridge/AvaMaterialBridgeEnums.h"
+#include "UObject/ObjectPtr.h"
+
+#define UE_API AVALANCHEMATERIAL_API
+
+class UMaterialInterface;
+
+namespace UE::Ava
+{
+
+/** Base struct for short-lived representation of an accessed material slot for write */
+class FMaterialBridgeWriteSlot
+{
+public:
+	UE_API explicit FMaterialBridgeWriteSlot(UMaterialInterface* InMaterial);
+
+	virtual ~FMaterialBridgeWriteSlot() = default;
+
+	/** Gets the material for this slot */
+	UMaterialInterface* GetMaterial() const
+	{
+		return Material;
+	}
+
+	/** Sets the new material for the slot */
+	UE_API void SetMaterial(UMaterialInterface* InMaterial);
+
+	/**
+	 * Requests features for this material slot
+	 * NOTE: this is not meant to modify any property of the material, 
+	 * but rather meant to modify the state of the material container to adapt to the requested features.
+	 */
+	UE_API void RequestFeatures(EMaterialBridgeRequestedFeatures InFeatures);
+
+protected:
+	/** Sets the material property directly without calling the change notification */
+	UE_API void SetMaterialInternal(UMaterialInterface* InMaterial);
+
+	/** Notifies when a new material has been set for the slot */
+	virtual void OnMaterialChanged()
+	{
+	}
+
+	/** Requests features for this material slot */
+	virtual void OnRequestFeatures(EMaterialBridgeRequestedFeatures InFeatures)
+	{
+	}
+
+private:
+	/** Material in this slot */
+	TObjectPtr<UMaterialInterface> Material;
+};
+
+} // UE::Ava
+
+#undef UE_API
diff --git a/Engine/Plugins/VirtualProduction/Avalanche/Source/AvalancheShapes/Private/AvaShapeMaterialBridge.cpp b/Engine/Plugins/VirtualProduction/Avalanche/Source/AvalancheShapes/Private/AvaShapeMaterialBridge.cpp
--- a/Engine/Plugins/VirtualProduction/Avalanche/Source/AvalancheShapes/Private/AvaShapeMaterialBridge.cpp
+++ b/Engine/Plugins/VirtualProduction/Avalanche/Source/AvalancheShapes/Private/AvaShapeMaterialBridge.cpp
@@ -2,11 +2,11 @@
 
 #include "AvaShapeMaterialBridge.h"
 #include "DynamicMeshes/AvaShapeDynMeshBase.h"
-#include "MaterialBridge/AvaMaterialBridgeReadSlot.h"
 #include "MaterialBridge/Context/AvaMaterialBridgeApplyStateContext.h"
 #include "MaterialBridge/Context/AvaMaterialBridgeReadSlotContext.h"
 #include "MaterialBridge/Context/AvaMaterialBridgeStoreStateContext.h"
 #include "MaterialBridge/Context/AvaMaterialBridgeWriteSlotContext.h"
+#include "MaterialBridge/Slot/AvaMaterialBridgeReadSlot.h"
 
 namespace UE::Ava
 {
diff --git a/Engine/Plugins/VirtualProduction/Avalanche/Source/AvalancheShapes/Public/AvaShapeMaterialBridge.h b/Engine/Plugins/VirtualProduction/Avalanche/Source/AvalancheShapes/Public/AvaShapeMaterialBridge.h
--- a/Engine/Plugins/VirtualProduction/Avalanche/Source/AvalancheShapes/Public/AvaShapeMaterialBridge.h
+++ b/Engine/Plugins/VirtualProduction/Avalanche/Source/AvalancheShapes/Public/AvaShapeMaterialBridge.h
@@ -5,8 +5,8 @@
 #include "AvaShapeParametricMaterial.h"
 #include "AvaShapeUVParameters.h"
 #include "MaterialBridge/AvaMaterialBridge.h"
-#include "MaterialBridge/AvaMaterialBridgeWriteSlot.h"
 #include "MaterialBridge/AvaMaterialContainerState.h"
+#include "MaterialBridge/Slot/AvaMaterialBridgeWriteSlot.h"
 #include "AvaShapeMaterialBridge.generated.h"
 
 #define UE_API AVALANCHESHAPES_API
diff --git a/Engine/Plugins/VirtualProduction/Avalanche/Source/AvalancheText/Private/MaterialBridge/AvaText3DExtensionMaterialBridge.cpp b/Engine/Plugins/VirtualProduction/Avalanche/Source/AvalancheText/Private/MaterialBridge/AvaText3DExtensionMaterialBridge.cpp
--- a/Engine/Plugins/VirtualProduction/Avalanche/Source/AvalancheText/Private/MaterialBridge/AvaText3DExtensionMaterialBridge.cpp
+++ b/Engine/Plugins/VirtualProduction/Avalanche/Source/AvalancheText/Private/MaterialBridge/AvaText3DExtensionMaterialBridge.cpp
@@ -2,11 +2,11 @@
 
 #include "MaterialBridge/AvaText3DExtensionMaterialBridge.h"
 #include "Extensions/Text3DMaterialExtensionBase.h"
-#include "MaterialBridge/AvaMaterialBridgeReadSlot.h"
 #include "MaterialBridge/Context/AvaMaterialBridgeApplyStateContext.h"
 #include "MaterialBridge/Context/AvaMaterialBridgeReadSlotContext.h"
 #include "MaterialBridge/Context/AvaMaterialBridgeStoreStateContext.h"
 #include "MaterialBridge/Context/AvaMaterialBridgeWriteSlotContext.h"
+#include "MaterialBridge/Slot/AvaMaterialBridgeReadSlot.h"
 #include "Text3DComponent.h"
 
 namespace UE::Ava
diff --git a/Engine/Plugins/VirtualProduction/Avalanche/Source/AvalancheText/Public/MaterialBridge/AvaText3DExtensionMaterialBridge.h b/Engine/Plugins/VirtualProduction/Avalanche/Source/AvalancheText/Public/MaterialBridge/AvaText3DExtensionMaterialBridge.h
--- a/Engine/Plugins/VirtualProduction/Avalanche/Source/AvalancheText/Public/MaterialBridge/AvaText3DExtensionMaterialBridge.h
+++ b/Engine/Plugins/VirtualProduction/Avalanche/Source/AvalancheText/Public/MaterialBridge/AvaText3DExtensionMaterialBridge.h
@@ -3,9 +3,9 @@
 #pragma once
 
 #include "MaterialBridge/AvaMaterialBridge.h"
-#include "Text3DTypes.h"
-#include "MaterialBridge/AvaMaterialBridgeWriteSlot.h"
 #include "MaterialBridge/AvaMaterialContainerState.h"
+#include "MaterialBridge/Slot/AvaMaterialBridgeWriteSlot.h"
+#include "Text3DTypes.h"
 #include "AvaText3DExtensionMaterialBridge.generated.h"
 
 #define UE_API AVALANCHETEXT_API
