[AvalancheTests] Automate 'Sequence - Assigning sequence tags to Sequences'

#rnx

#Synced-CL 45465245

diff --git a/Engine/Plugins/VirtualProduction/Avalanche/Source/AvalancheTagEditor/Private/Customization/SAvaTagCollectionPicker.h b/Engine/Plugins/VirtualProduction/Avalanche/Source/AvalancheTagEditor/Private/Customization/SAvaTagCollectionPicker.h
--- a/Engine/Plugins/VirtualProduction/Avalanche/Source/AvalancheTagEditor/Private/Customization/SAvaTagCollectionPicker.h
+++ b/Engine/Plugins/VirtualProduction/Avalanche/Source/AvalancheTagEditor/Private/Customization/SAvaTagCollectionPicker.h
@@ -33,7 +33,7 @@ private:
 
 	void OnTagCollectionMenuOpenChanged(bool bInIsOpened);
 
-	void OnTagCollectionSelected(const FAssetData& InAssetData);
+	DLLEXPORT void OnTagCollectionSelected(const FAssetData& InAssetData);
 
 	void CloseTagCollectionPicker();
 
@@ -42,4 +42,6 @@ private:
 	TSharedPtr<SComboButton> TagCollectionPickerButton;
 
 	FSimpleDelegate OnTagCollectionChanged;
+
+	friend class UAvaTagCollectionHelper;
 };
diff --git a/Engine/Plugins/VirtualProduction/Avalanche/Source/AvalancheTagEditor/Private/Customization/SAvaTagPicker.h b/Engine/Plugins/VirtualProduction/Avalanche/Source/AvalancheTagEditor/Private/Customization/SAvaTagPicker.h
--- a/Engine/Plugins/VirtualProduction/Avalanche/Source/AvalancheTagEditor/Private/Customization/SAvaTagPicker.h
+++ b/Engine/Plugins/VirtualProduction/Avalanche/Source/AvalancheTagEditor/Private/Customization/SAvaTagPicker.h
@@ -46,7 +46,7 @@ private:
 
 	bool IsTagHandleSelected(FAvaTagHandle InTagHandle) const;
 
-	void OnTagHandleSelectionChanged(const FAvaTagHandle& InTagHandle, bool bInIsSelected);
+	DLLEXPORT void OnTagHandleSelectionChanged(const FAvaTagHandle& InTagHandle, bool bInIsSelected);
 
 	FText GetValueDisplayText() const;
 
@@ -70,4 +70,6 @@ private:
 	TArray<TSharedPtr<FAvaTagHandle>> TagOptions;
 
 	bool bRequestOpenTagCollectionPicker = false;
+
+	friend class UAvaTagCollectionHelper;
 };
