Motion Design: changed material bridge features to be struct based instead of enum based allowing more flexibiltiy on features as well as features to have their own payload via their own structs.

#Synced-CL 48854462

diff --git a/Engine/Plugins/VirtualProduction/Avalanche/Source/AvalancheMaterial/Private/MaterialBridge/Slot/AvaMaterialBridgeWriteSlot.cpp b/Engine/Plugins/VirtualProduction/Avalanche/Source/AvalancheMaterial/Private/MaterialBridge/Slot/AvaMaterialBridgeWriteSlot.cpp
--- a/Engine/Plugins/VirtualProduction/Avalanche/Source/AvalancheMaterial/Private/MaterialBridge/Slot/AvaMaterialBridgeWriteSlot.cpp
+++ b/Engine/Plugins/VirtualProduction/Avalanche/Source/AvalancheMaterial/Private/MaterialBridge/Slot/AvaMaterialBridgeWriteSlot.cpp
@@ -19,12 +19,13 @@ void FMaterialBridgeWriteSlot::SetMaterial(UMaterialInterface* InMaterial)
 	}
 }
 
-void FMaterialBridgeWriteSlot::RequestFeatures(EMaterialBridgeRequestedFeatures InFeatures)
+bool FMaterialBridgeWriteSlot::RequestFeature(TConstStructView<FAvaMaterialBridgeFeature> InFeature)
 {
-	if (InFeatures != EMaterialBridgeRequestedFeatures::None)
+	if (InFeature.IsValid())
 	{
-		OnRequestFeatures(InFeatures);
+		return OnFeatureRequest(InFeature);
 	}
+	return false;
 }
 
 void FMaterialBridgeWriteSlot::SetMaterialInternal(UMaterialInterface* InMaterial)
diff --git a/Engine/Plugins/VirtualProduction/Avalanche/Source/AvalancheMaterial/Public/MaterialBridge/AvaMaterialBridgeEnums.h b/Engine/Plugins/VirtualProduction/Avalanche/Source/AvalancheMaterial/Public/MaterialBridge/AvaMaterialBridgeEnums.h
deleted file mode 100644
diff --git a/Engine/Plugins/VirtualProduction/Avalanche/Source/AvalancheMaterial/Public/MaterialBridge/Slot/AvaMaterialBridgeFeature.h b/Engine/Plugins/VirtualProduction/Avalanche/Source/AvalancheMaterial/Public/MaterialBridge/Slot/AvaMaterialBridgeFeature.h
new file mode 100644
--- /dev/null
+++ b/Engine/Plugins/VirtualProduction/Avalanche/Source/AvalancheMaterial/Public/MaterialBridge/Slot/AvaMaterialBridgeFeature.h
@@ -0,0 +1,12 @@
+// Copyright Epic Games, Inc. All Rights Reserved.
+
+#pragma once
+
+#include "AvaMaterialBridgeFeature.generated.h"
+
+/** Base struct for all slot features */
+USTRUCT()
+struct FAvaMaterialBridgeFeature
+{
+	GENERATED_BODY()
+};
diff --git a/Engine/Plugins/VirtualProduction/Avalanche/Source/AvalancheMaterial/Public/MaterialBridge/Slot/AvaMaterialBridgeWriteSlot.h b/Engine/Plugins/VirtualProduction/Avalanche/Source/AvalancheMaterial/Public/MaterialBridge/Slot/AvaMaterialBridgeWriteSlot.h
--- a/Engine/Plugins/VirtualProduction/Avalanche/Source/AvalancheMaterial/Public/MaterialBridge/Slot/AvaMaterialBridgeWriteSlot.h
+++ b/Engine/Plugins/VirtualProduction/Avalanche/Source/AvalancheMaterial/Public/MaterialBridge/Slot/AvaMaterialBridgeWriteSlot.h
@@ -2,12 +2,13 @@
 
 #pragma once
 
-#include "MaterialBridge/AvaMaterialBridgeEnums.h"
+#include "StructUtils/StructView.h"
 #include "UObject/ObjectPtr.h"
 
 #define UE_API AVALANCHEMATERIAL_API
 
 class UMaterialInterface;
+struct FAvaMaterialBridgeFeature;
 
 namespace UE::Ava
 {
@@ -29,12 +30,12 @@ public:
 	/** Sets the new material for the slot */
 	UE_API void SetMaterial(UMaterialInterface* InMaterial);
 
-	/**
-	 * Requests features for this material slot
-	 * NOTE: this is not meant to modify any property of the material, 
-	 * but rather meant to modify the state of the material container to adapt to the requested features.
+	/** 
+	 * Requests a feature for this material slot 
+	 * @param InFeature the feature to request
+	 * @return true if the feature was recognized by the slot
 	 */
-	UE_API void RequestFeatures(EMaterialBridgeRequestedFeatures InFeatures);
+	UE_API bool RequestFeature(TConstStructView<FAvaMaterialBridgeFeature> InFeature);
 
 protected:
 	/** Sets the material property directly without calling the change notification */
@@ -45,9 +46,14 @@ protected:
 	{
 	}
 
-	/** Requests features for this material slot */
-	virtual void OnRequestFeatures(EMaterialBridgeRequestedFeatures InFeatures)
+	/**
+	 * Requests a feature for this material slot
+	 * @param InFeature the feature to request
+	 * @return true if the feature was recognized by the slot
+	 */
+	virtual bool OnFeatureRequest(TConstStructView<FAvaMaterialBridgeFeature> InFeature)
 	{
+		return false;
 	}
 
 private:
diff --git a/Engine/Plugins/VirtualProduction/Avalanche/Source/AvalancheMaterial/Public/MaterialBridge/Slot/CommonFeatures/AvaMaterialBridgeBlendModeFeature.h b/Engine/Plugins/VirtualProduction/Avalanche/Source/AvalancheMaterial/Public/MaterialBridge/Slot/CommonFeatures/AvaMaterialBridgeBlendModeFeature.h
new file mode 100644
--- /dev/null
+++ b/Engine/Plugins/VirtualProduction/Avalanche/Source/AvalancheMaterial/Public/MaterialBridge/Slot/CommonFeatures/AvaMaterialBridgeBlendModeFeature.h
@@ -0,0 +1,18 @@
+// Copyright Epic Games, Inc. All Rights Reserved.
+
+#pragma once
+
+#include "Engine/EngineTypes.h"
+#include "MaterialBridge/Slot/AvaMaterialBridgeFeature.h"
+#include "AvaMaterialBridgeBlendModeFeature.generated.h"
+
+/** Feature representing the blend mode to configure for the slot */
+USTRUCT()
+struct FAvaMaterialBridgeBlendModeFeature : public FAvaMaterialBridgeFeature
+{
+	GENERATED_BODY()
+
+	/** The desired blend mode for the slot */
+	UPROPERTY()
+	TEnumAsByte<EBlendMode> BlendMode = EBlendMode::BLEND_Opaque;
+};
diff --git a/Engine/Plugins/VirtualProduction/Avalanche/Source/AvalancheShapes/Private/AvaShapeMaterialBridge.cpp b/Engine/Plugins/VirtualProduction/Avalanche/Source/AvalancheShapes/Private/AvaShapeMaterialBridge.cpp
--- a/Engine/Plugins/VirtualProduction/Avalanche/Source/AvalancheShapes/Private/AvaShapeMaterialBridge.cpp
+++ b/Engine/Plugins/VirtualProduction/Avalanche/Source/AvalancheShapes/Private/AvaShapeMaterialBridge.cpp
@@ -7,6 +7,7 @@
 #include "MaterialBridge/Context/AvaMaterialBridgeStoreStateContext.h"
 #include "MaterialBridge/Context/AvaMaterialBridgeWriteSlotContext.h"
 #include "MaterialBridge/Slot/AvaMaterialBridgeReadSlot.h"
+#include "MaterialBridge/Slot/CommonFeatures/AvaMaterialBridgeBlendModeFeature.h"
 
 namespace UE::Ava
 {
@@ -14,13 +15,17 @@ namespace UE::Ava
 namespace Private
 {
 
-TOptional<EAvaShapeParametricMaterialTranslucency> GetTranslucency(EMaterialBridgeRequestedFeatures InRequestedFeatures)
+EAvaShapeParametricMaterialTranslucency GetTargetTranslucency(EBlendMode InBlendMode)
 {
-	if (EnumHasAnyFlags(InRequestedFeatures, EMaterialBridgeRequestedFeatures::Translucency))
+	switch (InBlendMode)
 	{
+	case BLEND_Masked:
+	case BLEND_Translucent:
 		return EAvaShapeParametricMaterialTranslucency::Enabled;
+
+	default:
+		return EAvaShapeParametricMaterialTranslucency::Auto;
 	}
-	return TOptional<EAvaShapeParametricMaterialTranslucency>();
 }
 
 } // UE::Ava::Private
@@ -37,17 +42,23 @@ void FShapeMaterialWriteSlot::OnMaterialChanged()
 	Shape->SetMaterial(MeshIndex, GetMaterial());
 }
 
-void FShapeMaterialWriteSlot::OnRequestFeatures(EMaterialBridgeRequestedFeatures InFeatures)
+bool FShapeMaterialWriteSlot::OnFeatureRequest(TConstStructView<FAvaMaterialBridgeFeature> InFeature)
 {
-	if (EnumHasAnyFlags(InFeatures, EMaterialBridgeRequestedFeatures::Translucency))
+	// Blend mode feature
+	if (const FAvaMaterialBridgeBlendModeFeature* BlendModeFeature = InFeature.GetPtr<FAvaMaterialBridgeBlendModeFeature>())
 	{
+		const EAvaShapeParametricMaterialTranslucency TargetTranslucency = Private::GetTargetTranslucency(BlendModeFeature->BlendMode);
+
 		FAvaShapeParametricMaterial* const ParametricMaterial = Shape->GetParametricMaterialPtr(MeshIndex);
-		if (ParametricMaterial && Shape->GetMaterialType(MeshIndex) == EMaterialType::Parametric)
+		if (ParametricMaterial && ParametricMaterial->GetTranslucency() != TargetTranslucency && Shape->GetMaterialType(MeshIndex) == EMaterialType::Parametric)
 		{
-			ParametricMaterial->SetTranslucency(EAvaShapeParametricMaterialTranslucency::Enabled);
+			ParametricMaterial->SetTranslucency(TargetTranslucency);
 			SetMaterialInternal(Shape->GetMaterial(MeshIndex));
 		}
+		return true;
 	}
+
+	return false;
 }
 
 const UStruct* FShapeMaterialBridge::OnGetBridgedType() const
diff --git a/Engine/Plugins/VirtualProduction/Avalanche/Source/AvalancheShapes/Public/AvaShapeMaterialBridge.h b/Engine/Plugins/VirtualProduction/Avalanche/Source/AvalancheShapes/Public/AvaShapeMaterialBridge.h
--- a/Engine/Plugins/VirtualProduction/Avalanche/Source/AvalancheShapes/Public/AvaShapeMaterialBridge.h
+++ b/Engine/Plugins/VirtualProduction/Avalanche/Source/AvalancheShapes/Public/AvaShapeMaterialBridge.h
@@ -52,7 +52,7 @@ public:
 protected:
 	//~ Begin FMaterialBridgeSlot
 	UE_API virtual void OnMaterialChanged() override;
-	UE_API virtual void OnRequestFeatures(EMaterialBridgeRequestedFeatures InFeatures) override;
+	UE_API virtual bool OnFeatureRequest(TConstStructView<FAvaMaterialBridgeFeature> InFeature) override;
 	//~ End FMaterialBridgeSlot
 
 	/** The dynamic mesh shape component owning the slot */
diff --git a/Engine/Plugins/VirtualProduction/Avalanche/Source/AvalancheText/Private/MaterialBridge/AvaText3DDefaultExtensionMaterialBridge.cpp b/Engine/Plugins/VirtualProduction/Avalanche/Source/AvalancheText/Private/MaterialBridge/AvaText3DDefaultExtensionMaterialBridge.cpp
--- a/Engine/Plugins/VirtualProduction/Avalanche/Source/AvalancheText/Private/MaterialBridge/AvaText3DDefaultExtensionMaterialBridge.cpp
+++ b/Engine/Plugins/VirtualProduction/Avalanche/Source/AvalancheText/Private/MaterialBridge/AvaText3DDefaultExtensionMaterialBridge.cpp
@@ -5,6 +5,7 @@
 #include "MaterialBridge/Context/AvaMaterialBridgeApplyStateContext.h"
 #include "MaterialBridge/Context/AvaMaterialBridgeStoreStateContext.h"
 #include "MaterialBridge/Context/AvaMaterialBridgeWriteSlotContext.h"
+#include "MaterialBridge/Slot/CommonFeatures/AvaMaterialBridgeBlendModeFeature.h"
 #include "Text3DComponent.h"
 
 namespace UE::Ava
@@ -12,6 +13,20 @@ namespace UE::Ava
 
 namespace Private
 {
+// Gets the target Text3D blend mode for a given input blend mode
+EText3DMaterialBlendMode GetText3DTargetBlendMode(EBlendMode InBlendMode)
+{
+	switch (InBlendMode)
+	{
+	case BLEND_Masked:
+	case BLEND_Translucent:
+		return EText3DMaterialBlendMode::Translucent;
+
+	default:
+		return EText3DMaterialBlendMode::Opaque;
+	}
+}
+
 // Sets the material extension's blend mode and updating Text3D Component immediately to update any material properties
 // returns true if the blend mode was changed, or false if the blend mode was already the requested one.
 bool SetText3DBlendMode(TNotNull<UText3DDefaultMaterialExtension*> InMaterialExtension, EText3DMaterialBlendMode InBlendMode)
@@ -32,21 +47,25 @@ bool SetText3DBlendMode(TNotNull<UText3DDefaultMaterialExtension*> InMaterialExt
 	
 } // UE::Ava::Private
 
-void FText3DDefaultMaterialWriteSlot::OnRequestFeatures(EMaterialBridgeRequestedFeatures InFeatures)
+bool FText3DDefaultMaterialWriteSlot::OnFeatureRequest(TConstStructView<FAvaMaterialBridgeFeature> InFeature)
 {
 	UText3DDefaultMaterialExtension* MaterialExtension = Cast<UText3DDefaultMaterialExtension>(BaseMaterialExtension);
 	if (!ensure(MaterialExtension))
 	{
-		return;
+		return false;
 	}
 
-	if (EnumHasAnyFlags(InFeatures, EMaterialBridgeRequestedFeatures::Translucency))
+	// Blend mode feature
+	if (const FAvaMaterialBridgeBlendModeFeature* BlendModeFeature = InFeature.GetPtr<FAvaMaterialBridgeBlendModeFeature>())
 	{
-		if (Private::SetText3DBlendMode(MaterialExtension, EText3DMaterialBlendMode::Translucent))
+		if (Private::SetText3DBlendMode(MaterialExtension, Private::GetText3DTargetBlendMode(BlendModeFeature->BlendMode)))
 		{
 			SetMaterialInternal(MaterialExtension->GetMaterial(MaterialParameters));
 		}
+		return true;
 	}
+
+	return false;
 }
 
 const UStruct* FText3DDefaultExtensionMaterialBridge::OnGetBridgedType() const
diff --git a/Engine/Plugins/VirtualProduction/Avalanche/Source/AvalancheText/Public/MaterialBridge/AvaText3DDefaultExtensionMaterialBridge.h b/Engine/Plugins/VirtualProduction/Avalanche/Source/AvalancheText/Public/MaterialBridge/AvaText3DDefaultExtensionMaterialBridge.h
--- a/Engine/Plugins/VirtualProduction/Avalanche/Source/AvalancheText/Public/MaterialBridge/AvaText3DDefaultExtensionMaterialBridge.h
+++ b/Engine/Plugins/VirtualProduction/Avalanche/Source/AvalancheText/Public/MaterialBridge/AvaText3DDefaultExtensionMaterialBridge.h
@@ -28,7 +28,7 @@ public:
 
 protected:
 	//~ Begin FMaterialBridgeSlot
-	UE_API virtual void OnRequestFeatures(EMaterialBridgeRequestedFeatures InFeatures) override;
+	UE_API virtual bool OnFeatureRequest(TConstStructView<FAvaMaterialBridgeFeature> InFeature) override;
 	//~ End FMaterialBridgeSlot
 };
 
