Motion Design: changed mask modifiers to use the material bridge api. The material bridge is a more opened system that is more flexible in dealing with different types of objects than the existing internal material handle system.

#Synced-CL 48854462

diff --git a/Engine/Plugins/VirtualProduction/Avalanche/Source/AvalancheMask/AvalancheMask.Build.cs b/Engine/Plugins/VirtualProduction/Avalanche/Source/AvalancheMask/AvalancheMask.Build.cs
--- a/Engine/Plugins/VirtualProduction/Avalanche/Source/AvalancheMask/AvalancheMask.Build.cs
+++ b/Engine/Plugins/VirtualProduction/Avalanche/Source/AvalancheMask/AvalancheMask.Build.cs
@@ -10,6 +10,7 @@ public class AvalancheMask : ModuleRules
 
 		PublicDependencyModuleNames.AddRange(
 			new string[] {
+				"AvalancheMaterial",
 				"AvalancheModifiers",
 				"AvalancheShapes",
 				"Core",
diff --git a/Engine/Plugins/VirtualProduction/Avalanche/Source/AvalancheMask/Private/AvaMaskSettings.cpp b/Engine/Plugins/VirtualProduction/Avalanche/Source/AvalancheMask/Private/AvaMaskSettings.cpp
--- a/Engine/Plugins/VirtualProduction/Avalanche/Source/AvalancheMask/Private/AvaMaskSettings.cpp
+++ b/Engine/Plugins/VirtualProduction/Avalanche/Source/AvalancheMask/Private/AvaMaskSettings.cpp
@@ -15,33 +15,20 @@ UAvaMaskSettings::UAvaMaskSettings()
 {
 	CategoryName = TEXT("Motion Design");
 	SectionName = TEXT("Mask");
+#if WITH_EDITORONLY_DATA
+	MaterialFunction = FSoftObjectPath(TEXT("/Script/Engine.MaterialFunction'/GeometryMask/GeometryMask/MF_ApplyMask2D_Single.MF_ApplyMask2D_Single'"));
+#endif
 }
 
 #if WITH_EDITORONLY_DATA
-UMaterialFunctionInterface* UAvaMaskSettings::GetMaterialFunction()
+UMaterialFunctionInterface* UAvaMaskSettings::GetMaterialFunction() const
 {
-	if (UMaterialFunctionInterface* StrongMaterialFunction = MaterialFunction.Get())
+	if (ResolvedMaterialFunction)
 	{
-		return StrongMaterialFunction;		
+		return ResolvedMaterialFunction;		
 	}
 
-	UMaterialFunctionInterface* DefaultMaterialFunction = GetDefaultMaterialFunction();
-	MaterialFunction = DefaultMaterialFunction;
-	
-	return DefaultMaterialFunction;
-}
-
-UMaterialFunctionInterface* UAvaMaskSettings::GetDefaultMaterialFunction()
-{
-	const FString AssetPath = TEXT("/Script/Engine.MaterialFunction'/GeometryMask/GeometryMask/MF_ApplyMask2D_Single.MF_ApplyMask2D_Single'");
-	const FAssetRegistryModule& AssetRegistryModule = FModuleManager::LoadModuleChecked<FAssetRegistryModule>("AssetRegistry");
-	if (const FAssetData AssetData = AssetRegistryModule.Get().GetAssetByObjectPath(FSoftObjectPath(AssetPath));
-		!AssetData.IsValid())
-	{
-		UE_LOG(LogAvaMask, Error, TEXT("Default MaterialFunction asset is invalid or not found: '%s'"), *AssetPath);
-		return nullptr;
-	}
-	
-	return Cast<UMaterialFunctionInterface>(StaticLoadObject(UMaterialFunctionInterface::StaticClass(), this, *AssetPath));
+	ResolvedMaterialFunction = MaterialFunction.LoadSynchronous();
+	return ResolvedMaterialFunction;
 }
 #endif
diff --git a/Engine/Plugins/VirtualProduction/Avalanche/Source/AvalancheMask/Private/AvaMaskSettings.h b/Engine/Plugins/VirtualProduction/Avalanche/Source/AvalancheMask/Private/AvaMaskSettings.h
--- a/Engine/Plugins/VirtualProduction/Avalanche/Source/AvalancheMask/Private/AvaMaskSettings.h
+++ b/Engine/Plugins/VirtualProduction/Avalanche/Source/AvalancheMask/Private/AvaMaskSettings.h
@@ -24,16 +24,15 @@ public:
     UAvaMaskSettings();
 
 #if WITH_EDITORONLY_DATA
-	/** Material Function to use to expect or add to a material. */
-	UPROPERTY(Config, EditAnywhere, BlueprintReadOnly, NoClear, Category = "Material")
-	TSoftObjectPtr<UMaterialFunctionInterface> MaterialFunction;
-
 	/** Get the user specified or default material function. */
-	UMaterialFunctionInterface* GetMaterialFunction();
-#endif
+	UMaterialFunctionInterface* GetMaterialFunction() const;
 
 private:
-#if WITH_EDITORONLY_DATA
-	UMaterialFunctionInterface* GetDefaultMaterialFunction();
+	/** Material Function to use to expect or add to a material. */
+	UPROPERTY(Config, EditAnywhere, NoClear, Category = "Material", meta=(ConfigRestartRequired=true))
+	TSoftObjectPtr<UMaterialFunctionInterface> MaterialFunction;
+
+	UPROPERTY(Transient)
+	mutable TObjectPtr<UMaterialFunctionInterface> ResolvedMaterialFunction;
 #endif
 };
diff --git a/Engine/Plugins/VirtualProduction/Avalanche/Source/AvalancheMask/Private/AvaMaskSubsystem.cpp b/Engine/Plugins/VirtualProduction/Avalanche/Source/AvalancheMask/Private/AvaMaskSubsystem.cpp
--- a/Engine/Plugins/VirtualProduction/Avalanche/Source/AvalancheMask/Private/AvaMaskSubsystem.cpp
+++ b/Engine/Plugins/VirtualProduction/Avalanche/Source/AvalancheMask/Private/AvaMaskSubsystem.cpp
@@ -1,12 +1,46 @@
 ﻿// Copyright Epic Games, Inc. All Rights Reserved.
 
 #include "AvaMaskSubsystem.h"
+#include "Engine/Engine.h"
+#include "Materials/MaterialInterface.h"
 
-#if WITH_EDITOR
+UMaterialInterface* UAvaMaskSubsystem::StaticGetDefaultMaskMaterial()
+{
+	UAvaMaskSubsystem* Subsystem = GEngine ? GEngine->GetEngineSubsystem<UAvaMaskSubsystem>() : nullptr;
+	return Subsystem ? Subsystem->GetDefaultMaskMaterial() : nullptr;
+}
+
+UMaterialInterface* UAvaMaskSubsystem::GetDefaultMaskMaterial() const
+{
+	if (DefaultMaskMaterial)
+	{
+		return DefaultMaskMaterial;
+	}
+
+	// fallback if the material soft object is not loaded
+	DefaultMaskMaterial = DefaultMaskMaterialSoft.LoadSynchronous();
+	return DefaultMaskMaterial;
+}
 
+#if WITH_EDITOR
 void UAvaMaskSubsystem::SetLastSpecifiedChannelName(const FName InName)
 {
 	LastSpecifiedChannelName = InName;
 }
-
 #endif
+
+void UAvaMaskSubsystem::Initialize(FSubsystemCollectionBase& InCollection)
+{
+	Super::Initialize(InCollection);
+
+	DefaultMaskMaterialSoft = FSoftObjectPath(TEXT("/Avalanche/MaskResources/M_AvalancheMaskDefaultTranslucent.M_AvalancheMaskDefaultTranslucent"));
+	DefaultMaskMaterialSoft.LoadAsync(FLoadSoftObjectPathAsyncDelegate::CreateWeakLambda(this, 
+		[this](const FSoftObjectPath&, UObject* InObject)
+		{
+			if (UMaterialInterface* Material = Cast<UMaterialInterface>(InObject))
+			{
+				DefaultMaskMaterial = Material;
+			}
+		}));
+}
+
diff --git a/Engine/Plugins/VirtualProduction/Avalanche/Source/AvalancheMask/Private/AvaMaskSubsystem.h b/Engine/Plugins/VirtualProduction/Avalanche/Source/AvalancheMask/Private/AvaMaskSubsystem.h
--- a/Engine/Plugins/VirtualProduction/Avalanche/Source/AvalancheMask/Private/AvaMaskSubsystem.h
+++ b/Engine/Plugins/VirtualProduction/Avalanche/Source/AvalancheMask/Private/AvaMaskSubsystem.h
@@ -3,22 +3,40 @@
 #pragma once
 
 #include "Subsystems/EngineSubsystem.h"
-
 #include "AvaMaskSubsystem.generated.h"
 
+class UMaterialInterface;
+
 UCLASS()
-class UAvaMaskSubsystem
-	: public UEngineSubsystem
+class UAvaMaskSubsystem : public UEngineSubsystem
 {
 	GENERATED_BODY()
 
 public:
+	/** Gets the material to use for masking when there's no material set */
+	static UMaterialInterface* StaticGetDefaultMaskMaterial();
+
+	/** Gets the material to use for masking when there's no material set */
+	UMaterialInterface* GetDefaultMaskMaterial() const;
+
 #if WITH_EDITORONLY_DATA
 	/** Stores the last specified (user or auto) channel name to use for subsequently added modifiers in Target mode. */
 	void SetLastSpecifiedChannelName(const FName InName);
 #endif
-	
+
+	//~ Begin USubsystem
+	virtual void Initialize(FSubsystemCollectionBase& InCollection) override;
+	//~ End USubsystem
+
 private:
+	/** Soft object ptr to the default mask material to use when there's no material set */
+	UPROPERTY()
+	TSoftObjectPtr<UMaterialInterface> DefaultMaskMaterialSoft;
+
+	/** Default mask material to use when there's no material set */
+	UPROPERTY()
+	mutable TObjectPtr<UMaterialInterface> DefaultMaskMaterial;
+
 #if WITH_EDITORONLY_DATA
 	/** Stores the last specified (user or auto) channel name to use for subsequently added modifiers in Target mode. */
 	UPROPERTY(Transient, Getter = "Auto", Setter)
diff --git a/Engine/Plugins/VirtualProduction/Avalanche/Source/AvalancheMask/Private/AvaMaskTypes.cpp b/Engine/Plugins/VirtualProduction/Avalanche/Source/AvalancheMask/Private/AvaMaskTypes.cpp
--- a/Engine/Plugins/VirtualProduction/Avalanche/Source/AvalancheMask/Private/AvaMaskTypes.cpp
+++ b/Engine/Plugins/VirtualProduction/Avalanche/Source/AvalancheMask/Private/AvaMaskTypes.cpp
@@ -2,9 +2,58 @@
 
 #include "AvaMaskTypes.h"
 
+#include "AvaMaskLog.h"
+#include "AvaMaskUtilities.h"
 #include "Engine/Texture.h"
+#include "Materials/MaterialInstanceDynamic.h"
 #include "Materials/MaterialInterface.h"
 
+bool FAvaMask2DMaterialParameters::HasSameParameters(const FAvaMask2DMaterialParameters& InOther) const
+{
+	constexpr float Tolerance = UE_SMALL_NUMBER;
+
+	return BlendMode == InOther.BlendMode
+		&& Texture == InOther.Texture
+		&& FMath::IsNearlyEqual(BaseOpacity, InOther.BaseOpacity, Tolerance)
+		&& ChannelAsVector.Equals(InOther.ChannelAsVector, Tolerance)
+		&& bInvert == InOther.bInvert
+		&& Padding.Equals(InOther.Padding, Tolerance)
+		&& bApplyFeathering == InOther.bApplyFeathering
+		&& FMath::IsNearlyEqual(InnerFeatherRadius, InOther.InnerFeatherRadius, Tolerance)
+		&& FMath::IsNearlyEqual(OuterFeatherRadius, InOther.OuterFeatherRadius, Tolerance);
+}
+
+bool FAvaMask2DMaterialParameters::ApplyToMID(UMaterialInstanceDynamic* InMaterial) const
+{
+	if (!ensureAlways(InMaterial))
+	{
+		return false;
+	}
+
+	InMaterial->BlendMode = BlendMode;
+	InMaterial->BasePropertyOverrides.bOverride_BlendMode = true;
+	InMaterial->BasePropertyOverrides.BlendMode = BlendMode;
+
+	InMaterial->SetTextureParameterValueByInfo(UE::AvaMask::Internal::TextureParameterInfo, Texture);
+	InMaterial->SetScalarParameterValueByInfo(UE::AvaMask::Internal::BaseOpacityParameterInfo, FMath::Clamp(BaseOpacity, 0.0f, 1.0f));
+	InMaterial->SetVectorParameterValueByInfo(UE::AvaMask::Internal::ChannelParameterInfo, ChannelAsVector);
+	InMaterial->SetScalarParameterValueByInfo(UE::AvaMask::Internal::InvertParameterInfo, bInvert ? 1.f : 0.f);
+	InMaterial->SetVectorParameterValueByInfo(UE::AvaMask::Internal::PaddingParameterInfo, FLinearColor(FMath::Max(0, Padding.X), FMath::Max(0, Padding.Y), 0, 0));
+	InMaterial->SetVectorParameterValueByInfo(UE::AvaMask::Internal::FeatherParameterInfo, FLinearColor(bApplyFeathering ? 1.0f : 0.0f, FMath::Max(0, OuterFeatherRadius), FMath::Max(0, InnerFeatherRadius), FMath::Max(0, FMath::Max(OuterFeatherRadius, InnerFeatherRadius))));
+	return true;
+}
+
+bool FAvaMask2DMaterialParameters::StoreFromMaterial(const UMaterialInterface* InMaterial)
+{
+	return SetFromMaterial(InMaterial
+		, UE::AvaMask::Internal::TextureParameterInfo
+		, UE::AvaMask::Internal::ChannelParameterInfo
+		, UE::AvaMask::Internal::InvertParameterInfo
+		, UE::AvaMask::Internal::BaseOpacityParameterInfo
+		, UE::AvaMask::Internal::PaddingParameterInfo
+		, UE::AvaMask::Internal::FeatherParameterInfo);
+}
+
 bool FAvaMask2DMaterialParameters::SetFromMaterial(
 	const UMaterialInterface* InMaterial
 	, const FMaterialParameterInfo& InTextureParameter
@@ -19,11 +68,19 @@ bool FAvaMask2DMaterialParameters::SetFromMaterial(
 		return false;
 	}
 
+	BlendMode = InMaterial->GetBlendMode();
+
+	auto LogMissingParameter = [InMaterial](const FMaterialParameterInfo& InParameter)
+		{
+			UE_LOG(LogAvaMask, Warning, TEXT("Could not find parameter '%s' in material '%s'"), *InParameter.ToString(), *InMaterial->GetFullName());;
+		};
+
 	// Texture
 	{
 		UTexture* FoundTexture = nullptr;
 		if (!InMaterial->GetTextureParameterValue(InTextureParameter, FoundTexture))
 		{
+			LogMissingParameter(InTextureParameter);
 			return false;
 		}
 		Texture = FoundTexture;
@@ -34,6 +91,7 @@ bool FAvaMask2DMaterialParameters::SetFromMaterial(
 		float FoundValue = 0.0f;
 		if (!InMaterial->GetScalarParameterValue(InBaseOpacityParameter, FoundValue))
 		{
+			LogMissingParameter(InTextureParameter);
 			return false;
 		}
 		BaseOpacity = FoundValue;
@@ -44,6 +102,7 @@ bool FAvaMask2DMaterialParameters::SetFromMaterial(
 		FLinearColor FoundValue;
 		if (!InMaterial->GetVectorParameterValue(InChannelParameter, FoundValue))
 		{
+			LogMissingParameter(InTextureParameter);
 			return false;
 		}
 		ChannelAsVector = FoundValue;
@@ -55,6 +114,7 @@ bool FAvaMask2DMaterialParameters::SetFromMaterial(
 		float FoundValue = 0.0f;
 		if (!InMaterial->GetScalarParameterValue(InInvertParameter, FoundValue))
 		{
+			LogMissingParameter(InTextureParameter);
 			return false;
 		}
 		bInvert = FoundValue > 0.1;
@@ -65,6 +125,7 @@ bool FAvaMask2DMaterialParameters::SetFromMaterial(
 		FLinearColor FoundValue;
 		if (!InMaterial->GetVectorParameterValue(InPaddingParameter, FoundValue))
 		{
+			LogMissingParameter(InTextureParameter);
 			return false;			
 		}
 		Padding = FVector2f(FoundValue.R, FoundValue.G);
@@ -75,6 +136,7 @@ bool FAvaMask2DMaterialParameters::SetFromMaterial(
 		FLinearColor FoundValue;
 		if (!InMaterial->GetVectorParameterValue(InFeatherParameter, FoundValue))
 		{
+			LogMissingParameter(InTextureParameter);
 			return false;
 		}
 		bApplyFeathering = !FMath::IsNearlyZero(FoundValue.R);
diff --git a/Engine/Plugins/VirtualProduction/Avalanche/Source/AvalancheMask/Private/AvaMaskUtilities.cpp b/Engine/Plugins/VirtualProduction/Avalanche/Source/AvalancheMask/Private/AvaMaskUtilities.cpp
--- a/Engine/Plugins/VirtualProduction/Avalanche/Source/AvalancheMask/Private/AvaMaskUtilities.cpp
+++ b/Engine/Plugins/VirtualProduction/Avalanche/Source/AvalancheMask/Private/AvaMaskUtilities.cpp
@@ -20,6 +20,29 @@
 
 namespace UE::AvaMask::Internal
 {
+#if WITH_EDITOR
+	UDMMaterialProperty* GetMaterialPropertyForBlendMode(TNotNull<UMaterialInterface*> InMaterial, EBlendMode InBlendMode)
+	{
+		UDynamicMaterialInstance* const Material = Cast<UDynamicMaterialInstance>(InMaterial);
+		if (!Material)
+		{
+			return nullptr;
+		}
+
+		UDynamicMaterialModelEditorOnlyData* const EditorOnlyData = UDynamicMaterialModelEditorOnlyData::Get(Material->GetMaterialModel());
+		if (!EditorOnlyData)
+		{
+			return nullptr;
+		}
+
+		const EDMMaterialPropertyType PropertyType = InBlendMode == EBlendMode::BLEND_Masked
+			? EDMMaterialPropertyType::OpacityMask
+			: EDMMaterialPropertyType::Opacity;
+
+		return EditorOnlyData->GetMaterialProperty(PropertyType);
+	}
+#endif
+
 	FSoftComponentReference MakeComponentReference(
 		const AActor* InOwner
 		, const UActorComponent* InComponent)
@@ -182,4 +205,30 @@ namespace UE::AvaMask::Internal
 
 		return ParentAsset;
 	}
+
+#if WITH_EDITOR
+	void SetOutputProcessor(TNotNull<UMaterialInterface*> InMaterial, UMaterialFunctionInterface* InOutputProcessor, EBlendMode InBlendMode)
+	{
+		UDMMaterialProperty* MaterialProperty = GetMaterialPropertyForBlendMode(InMaterial, InBlendMode);
+		if (!MaterialProperty)
+		{
+			return;
+		}
+
+		const UMaterialFunctionInterface* CurrentOutputProcessor = MaterialProperty->GetOutputProcessor();
+		if (CurrentOutputProcessor != InOutputProcessor)
+		{
+			MaterialProperty->SetOutputProcessor(InOutputProcessor);
+		}
+	}
+
+	UMaterialFunctionInterface* GetOutputProcessor(TNotNull<UMaterialInterface*> InMaterial, EBlendMode InBlendMode)
+	{
+		if (UDMMaterialProperty* MaterialProperty = GetMaterialPropertyForBlendMode(InMaterial, InBlendMode))
+		{
+			return MaterialProperty->GetOutputProcessor();
+		}
+		return nullptr;
+	}
+#endif
 }
diff --git a/Engine/Plugins/VirtualProduction/Avalanche/Source/AvalancheMask/Private/AvaMaskUtilities.h b/Engine/Plugins/VirtualProduction/Avalanche/Source/AvalancheMask/Private/AvaMaskUtilities.h
--- a/Engine/Plugins/VirtualProduction/Avalanche/Source/AvalancheMask/Private/AvaMaskUtilities.h
+++ b/Engine/Plugins/VirtualProduction/Avalanche/Source/AvalancheMask/Private/AvaMaskUtilities.h
@@ -59,14 +59,6 @@ namespace UE::AvaMask::Internal
 
 	void RemoveComponentByInterface(const TSubclassOf<UInterface>& InInterfaceType, const AActor* InActor);
 
-	template<
-		typename InInterfaceType
-		UE_REQUIRES(std::is_base_of_v<UInterface, InInterfaceType>)>
-	void RemoveComponentByInterface(AActor* InActor)
-	{
-		RemoveComponentByInterface(InInterfaceType::StaticClass(), InActor);
-	}
-
 	FString GetGeneratedMaterialPath();
 
 	FString GetBlendModeString(const EBlendMode InBlendMode);
@@ -74,33 +66,13 @@ namespace UE::AvaMask::Internal
 	/** Returns whichever blend mode satisfies the requirement AND user specified blend mode */
 	EBlendMode GetTargetBlendMode(const EBlendMode InFromMaterial, const EBlendMode InRequired);
 
-	template <typename ElementType, typename KeyFuncs, typename AllocatorType, typename LookupType, typename LambdaType>
-	ElementType& FindOrAddByLambda(TSet<ElementType, KeyFuncs, AllocatorType>& Set, const LookupType& Key, LambdaType Lambda)
-	{
-		uint32 KeyHash = KeyFuncs::GetKeyHash(Key);
-		if (ElementType* Found = Set.FindByHash(Key))
-		{
-			return *Found;
-		}
-
-		ElementType NewElement = Invoke(Lambda, Key);
-		check(KeyFuncs::GetKeyHash(NewElement) == KeyHash);
-		FSetElementId NewElementId = Set.AddByHash(KeyHash, MoveTemp(NewElement));
-		return Set[NewElementId];
-	}
+	UMaterialInterface* GetNonTransientParentMaterial(const UMaterialInstance* InMaterialInstance);
 
-	template <typename KeyType, typename ValueType, typename AllocatorType, typename KeyFuncs, typename InitKeyType, typename LambdaType>
-	ValueType& FindOrAddByLambda(TMap<KeyType, ValueType, AllocatorType, KeyFuncs>& Map, InitKeyType&& Key, LambdaType Lambda)
-	{
-		uint32 KeyHash = KeyFuncs::GetKeyHash(Key);
-		if (ValueType* Found = Map.FindByHash(KeyHash, Key))
-		{
-			return *Found;
-		}
-
-		ValueType NewElement = Invoke(Lambda);
-		return Map.AddByHash(KeyHash, Forward<InitKeyType>(Key), MoveTemp(NewElement));
-	}
+#if WITH_EDITOR
+	/** Sets the output processor for Material Designer  materials */
+	void SetOutputProcessor(TNotNull<UMaterialInterface*> InMaterial, UMaterialFunctionInterface* InOutputProcessor, EBlendMode InBlendMode);
 
-	UMaterialInterface* GetNonTransientParentMaterial(const UMaterialInstance* InMaterialInstance);
+	/** Gets the output process from Material Designer materials */
+	UMaterialFunctionInterface* GetOutputProcessor(TNotNull<UMaterialInterface*> InMaterial, EBlendMode InBlendMode);
+#endif
 }
diff --git a/Engine/Plugins/VirtualProduction/Avalanche/Source/AvalancheMask/Private/Mask2D/AvaMask2DBaseModifier.cpp b/Engine/Plugins/VirtualProduction/Avalanche/Source/AvalancheMask/Private/Mask2D/AvaMask2DBaseModifier.cpp
--- a/Engine/Plugins/VirtualProduction/Avalanche/Source/AvalancheMask/Private/Mask2D/AvaMask2DBaseModifier.cpp
+++ b/Engine/Plugins/VirtualProduction/Avalanche/Source/AvalancheMask/Private/Mask2D/AvaMask2DBaseModifier.cpp
@@ -33,26 +33,6 @@
 
 namespace UE::AvaMask::Private
 {
-	template <typename KeyType, typename ValueType>
-	ValueType& FindInPreviousOrAddToNew(
-		const TMap<KeyType, ValueType>& InPrevious
-		, TMap<KeyType, ValueType>& InNew
-		, const KeyType& InKey
-		, TIdentity_T<TUniqueFunction<ValueType()>>&& InValueFactory)
-	{
-		if (const ValueType* FoundValue = InPrevious.Find(InKey))
-		{
-			return InNew.Emplace(InKey, MoveTempIfPossible(*FoundValue));
-		}
-
-		if (ValueType* FoundValue = InNew.Find(InKey))
-		{
-			return *FoundValue;
-		}
-
-		return InNew.Emplace(InKey, InValueFactory());
-	}
-
 	// @note: It's crucial these remain in sync with the defaults in UGeometryMaskCanvas
 	namespace CanvasPropertyDefaults
 	{
@@ -192,11 +172,13 @@ void UAvaMask2DBaseModifier::UpdateCanvas()
 {
 	// Reset/invalidate cached
 	{
-		for (TPair<TWeakObjectPtr<AActor>, FAvaMask2DActorData>& ActorDataPair : ActorData)
+PRAGMA_DISABLE_DEPRECATION_WARNINGS
+		for (TPair<TWeakObjectPtr<AActor>, FAvaMask2DActorData>& ActorDataPair : ActorData_DEPRECATED)
 		{
 			ActorDataPair.Value.CanvasTextureWeak.Reset();
 		}
-		
+PRAGMA_ENABLE_DEPRECATION_WARNINGS
+
 		LastResolvedCanvasName = NAME_None;
 		CanvasWeak.Reset();
 	}
@@ -302,7 +284,9 @@ void UAvaMask2DBaseModifier::OnModifierAdded(EActorModifierCoreEnableReason InRe
 	if (InReason == EActorModifierCoreEnableReason::Load)
 	{
 		LocalParamsToCanvas();
+PRAGMA_DISABLE_DEPRECATION_WARNINGS
 		RestoreActors();
+PRAGMA_ENABLE_DEPRECATION_WARNINGS
 	}
 	else
 	{
@@ -322,8 +306,10 @@ void UAvaMask2DBaseModifier::OnModifierRemoved(EActorModifierCoreDisableReason I
 	}
 
 	TArray<AActor*> ActorDataActors;
-	ActorDataActors.Reserve(ActorData.Num());	
-	Algo::Transform(ActorData, ActorDataActors, [](const TPair<TWeakObjectPtr<AActor>, FAvaMask2DActorData>& InActorDataPair)
+PRAGMA_DISABLE_DEPRECATION_WARNINGS
+	ActorDataActors.Reserve(ActorData_DEPRECATED.Num());	
+	Algo::Transform(ActorData_DEPRECATED, ActorDataActors, [](const TPair<TWeakObjectPtr<AActor>, FAvaMask2DActorData>& InActorDataPair)
+PRAGMA_ENABLE_DEPRECATION_WARNINGS
 	{
 		return InActorDataPair.Key.Get();
 	});
@@ -344,44 +330,13 @@ void UAvaMask2DBaseModifier::OnModifiedActorTransformed()
 	// This is needed to override parent behavior (to disable it)
 }
 
-void UAvaMask2DBaseModifier::SavePreState()
-{
-	Super::SavePreState();
-
-	ScannedPrimitiveComponents.Empty();
-	TSet<AActor*> KeepActors;
-
-	ForEachActor<AActor>(
-		[this, &KeepActors](AActor* InActor)
-		{
-			const UActorModifierCoreBase* MaskModifier = UE::ActorModifierCore::Utilities::FindFirstActorModifierByClass(InActor, UAvaMask2DBaseModifier::StaticClass());
-
-			// If first modifier in the hierarchy is not this modifier, then skip actor
-			if (MaskModifier == this)
-			{
-				SaveActorPreState(InActor, ActorData.FindOrAdd(InActor));
-				KeepActors.Add(InActor);
-			}
-
-			return true;
-		}
-		, EActorModifierCoreLookup::SelfAndAllChildren);
-
-	// Cleanup
-	for (TMap<TWeakObjectPtr<AActor>, FAvaMask2DActorData>::TIterator It(ActorData); It; ++It)
-	{
-		if (!KeepActors.Contains(It->Key.Get()))
-		{
-			It.RemoveCurrent();
-		}
-	}
-}
-
 void UAvaMask2DBaseModifier::RestorePreState()
 {
 	Super::RestorePreState();
 
+PRAGMA_DISABLE_DEPRECATION_WARNINGS
 	RestoreActors();
+PRAGMA_ENABLE_DEPRECATION_WARNINGS
 }
 
 void UAvaMask2DBaseModifier::Apply()
@@ -391,37 +346,11 @@ void UAvaMask2DBaseModifier::Apply()
 
 void UAvaMask2DBaseModifier::SaveActorPreState(AActor* InActor, FAvaMask2DActorData& InActorData)
 {
-	// Get MaterialCollectionHandle
-	const TSharedPtr<IAvaMaskMaterialCollectionHandle> MaterialCollectionHandle =
-		UE::Ava::Internal::FindOrAddHandleByLambda(
-			MaterialCollectionHandles
-			, InActor
-			, [this, InActor]()
-			{
-				TSharedPtr<IAvaMaskMaterialCollectionHandle> Handle = GetObjectHandleSubsystem()->MakeCollectionHandle<IAvaMaskMaterialCollectionHandle>(InActor, UE::AvaMask::Internal::HandleTag);
-
-				if (!Handle->OnSourceMaterialsChanged().IsBoundToObject(this))
-				{
-					Handle->OnSourceMaterialsChanged().BindUObject(this, &UAvaMask2DBaseModifier::OnMaterialsChanged);
-				}
-
-				return Handle;
-			}
-		);
-
-	FInstancedStruct& HandleData = MaterialCollectionHandleData.FindOrAdd(InActor, MaterialCollectionHandle->MakeDataStruct());
-	MaterialCollectionHandle->SaveOriginalState(HandleData);
-
-	TArray<UPrimitiveComponent*> PrimitiveComponents;
-	InActor->GetComponents(PrimitiveComponents);
-	Algo::Transform(PrimitiveComponents, ScannedPrimitiveComponents, [](UPrimitiveComponent* InComponent)
-	{
-		return InComponent;
-	});
 }
 
 void UAvaMask2DBaseModifier::RestoreActorPreState(AActor* InActor, FAvaMask2DActorData& InActorData)
 {
+PRAGMA_DISABLE_DEPRECATION_WARNINGS
 	// Get MaterialCollectionHandle
     const TSharedPtr<IAvaMaskMaterialCollectionHandle> MaterialCollectionHandle =
     	UE::Ava::Internal::FindOrAddHandleByLambda(
@@ -437,14 +366,16 @@ void UAvaMask2DBaseModifier::RestoreActorPreState(AActor* InActor, FAvaMask2DAct
     		}
     	);
 
-    FInstancedStruct& MaterialCollectionData = MaterialCollectionHandleData.FindOrAdd(InActor, MaterialCollectionHandle->MakeDataStruct());
-    MaterialCollectionHandle->ApplyOriginalState(MaterialCollectionData);
+    FInstancedStruct* MaterialCollectionData = MaterialCollectionHandleData_DEPRECATED.Find(InActor);
+PRAGMA_ENABLE_DEPRECATION_WARNINGS
+
+	if (MaterialCollectionData)
+	{
+		MaterialCollectionHandle->ApplyOriginalState(*MaterialCollectionData);
+	}
 }
 
-void UAvaMask2DBaseModifier::OnSceneTreeTrackedActorParentChanged(
-	int32 InIdx
-	, const TArray<TWeakObjectPtr<AActor>>& InPreviousParentActor
-	, const TArray<TWeakObjectPtr<AActor>>& InNewParentActor)
+void UAvaMask2DBaseModifier::OnSceneTreeTrackedActorParentChanged(int32 InIdx, const TArray<TWeakObjectPtr<AActor>>& InPreviousParentActor, const TArray<TWeakObjectPtr<AActor>>& InNewParentActor)
 {
 	Super::OnSceneTreeTrackedActorParentChanged(InIdx, InPreviousParentActor, InNewParentActor);
 
@@ -475,53 +406,6 @@ void UAvaMask2DBaseModifier::OnSceneTreeTrackedActorChildrenChanged(int32 InIdx,
 void UAvaMask2DBaseModifier::OnRenderStateUpdated(AActor* InActor, UActorComponent* InComponent)
 {
 	Super::OnRenderStateUpdated(InActor, InComponent);
-
-	if (!IsValid(InActor) || !IsValid(InComponent))
-	{
-		return;
-	}
-
-	const AActor* ActorModified = GetModifiedActor();
-
-	if (!(ActorModified == InActor || InActor->IsAttachedTo(ActorModified)))
-	{
-		return;
-	}
-
-	if (!ActorSupportsMaskReadWrite(InActor))
-	{
-		return;
-	}
-
-	const UActorModifierCoreBase* MaskModifier = UE::ActorModifierCore::Utilities::FindFirstActorModifierByClass(InActor, UAvaMask2DBaseModifier::StaticClass());
-
-	// Only trigger update if first modifier found in the hierarchy above us is this modifier
-	if (MaskModifier != this)
-	{
-		return;
-	}
-
-	const UPrimitiveComponent* PrimitiveComponent = Cast<UPrimitiveComponent>(InComponent);
-	if (!PrimitiveComponent)
-	{
-		return;
-	}
-
-	const int32 MaterialCount = PrimitiveComponent->GetNumMaterials();
-	if (MaterialCount != 0)
-	{
-		if (const FAvaMask2DActorData* MaskData = ActorData.Find(InActor))
-		{
-			if (CanvasWeak.IsValid()
-				&& MaskData->CanvasTextureWeak.IsValid()
-				&& ScannedPrimitiveComponents.Contains(PrimitiveComponent))
-			{
-				return;
-			}
-		}
-
-		MarkModifierDirty();
-	}
 }
 
 FName UAvaMask2DBaseModifier::GenerateUniqueMaskName() const
@@ -653,32 +537,16 @@ bool UAvaMask2DBaseModifier::TryResolveParentChannel()
 
 void UAvaMask2DBaseModifier::SetupMaskComponent(UActorComponent* InComponent)
 {
-	if (!InComponent)
-	{
-		return;
-	}
-
-	if (const UGeometryMaskCanvas* Canvas = GetCurrentCanvas())
-	{
-		if (FAvaMask2DActorData* StoredActorData = ActorData.Find(InComponent->GetOwner()))
-		{
-			StoredActorData->CanvasTextureWeak = Canvas->GetTexture();
-		}
-	}
-	else
-	{
-		UE_LOG(LogAvaMask, Warning, TEXT("Expected canvas to be valid"));
-	}
 }
 
 void UAvaMask2DBaseModifier::RemoveFromActor(AActor* InActor)
 {
-	ActorData.Remove(InActor);
-	
-	MaterialCollectionHandleData.Remove(InActor);
+PRAGMA_DISABLE_DEPRECATION_WARNINGS
+	ActorData_DEPRECATED.Remove(InActor);
+	MaterialCollectionHandleData_DEPRECATED.Remove(InActor);
 	MaterialCollectionHandles.Remove(InActor);
-	
-	UE::AvaMask::Internal::RemoveComponentByInterface<UGeometryMaskWriteInterface>(InActor);
+PRAGMA_ENABLE_DEPRECATION_WARNINGS
+
 	if (UAvaGizmoComponent* FoundComponent = InActor->FindComponentByClass<UAvaGizmoComponent>())
 	{
 		FoundComponent->DestroyComponent();
@@ -741,7 +609,9 @@ void UAvaMask2DBaseModifier::OnMaskSetCanvas(const UGeometryMaskCanvas* InCanvas
 		return;
 	}
 
-	if (FAvaMask2DActorData* StoredActorData = ActorData.Find(InActor))
+PRAGMA_DISABLE_DEPRECATION_WARNINGS
+	if (FAvaMask2DActorData* StoredActorData = ActorData_DEPRECATED.Find(InActor))
+PRAGMA_ENABLE_DEPRECATION_WARNINGS
 	{
 		// Cache canvas texture
 		StoredActorData->CanvasTextureWeak = InCanvas->GetTexture();
@@ -812,6 +682,7 @@ UTexture* UAvaMask2DBaseModifier::TryResolveCanvasTexture(AActor* InActor, FAvaM
 
 UAvaObjectHandleSubsystem* UAvaMask2DBaseModifier::GetObjectHandleSubsystem()
 {
+PRAGMA_DISABLE_DEPRECATION_WARNINGS
 	if (UAvaObjectHandleSubsystem* Subsystem = ObjectHandleSubsystem.Get())
 	{
 		return Subsystem;
@@ -819,6 +690,7 @@ UAvaObjectHandleSubsystem* UAvaMask2DBaseModifier::GetObjectHandleSubsystem()
 
 	ObjectHandleSubsystem = GEngine->GetEngineSubsystem<UAvaObjectHandleSubsystem>();
 	return ObjectHandleSubsystem.Get();
+PRAGMA_ENABLE_DEPRECATION_WARNINGS
 }
 
 UAvaMaskMaterialInstanceSubsystem* UAvaMask2DBaseModifier::GetMaterialInstanceSubsystem()
@@ -855,16 +727,17 @@ UGeometryMaskCanvas* UAvaMask2DBaseModifier::GetCurrentCanvas()
 
 void UAvaMask2DBaseModifier::OnMaterialsChanged(UObject* InMaterialOwner, const TArray<TSharedPtr<IAvaMaskMaterialHandle>>& InMaterialHandles)
 {
-	if (!InMaterialHandles.IsEmpty())
+PRAGMA_DISABLE_DEPRECATION_WARNINGS
+	if (!InMaterialHandles.IsEmpty() && !ActorData_DEPRECATED.IsEmpty())
 	{
 		if (AActor* OwningActor = InMaterialOwner->GetTypedOuter<AActor>())
 		{
 			const UActorModifierCoreBase* MaskModifier = UE::ActorModifierCore::Utilities::FindFirstActorModifierByClass(OwningActor, UAvaMask2DBaseModifier::StaticClass());
 
 			// Only trigger update if first modifier found in the hierarchy above us is this modifier
-			if (MaskModifier == this && ActorData.Contains(OwningActor))
+			if (MaskModifier == this && ActorData_DEPRECATED.Contains(OwningActor))
 			{
-				FInstancedStruct* HandleData = MaterialCollectionHandleData.Find(OwningActor);
+				FInstancedStruct* HandleData = MaterialCollectionHandleData_DEPRECATED.Find(OwningActor);
 				const TSharedPtr<IAvaMaskMaterialCollectionHandle>* CollectionHandle = MaterialCollectionHandles.Find(OwningActor);
 
 				if (HandleData && CollectionHandle && CollectionHandle->IsValid())
@@ -876,21 +749,21 @@ void UAvaMask2DBaseModifier::OnMaterialsChanged(UObject* InMaterialOwner, const
 			}
 		}
 	}
+PRAGMA_ENABLE_DEPRECATION_WARNINGS
 }
 
 void UAvaMask2DBaseModifier::RestoreActors()
 {
-	TMap<TWeakObjectPtr<AActor>, FAvaMask2DActorData> ActorDataCopy = ActorData;
+PRAGMA_DISABLE_DEPRECATION_WARNINGS
+	TMap<TWeakObjectPtr<AActor>, FAvaMask2DActorData> ActorDataCopy = ActorData_DEPRECATED;
 	for (TPair<TWeakObjectPtr<AActor>, FAvaMask2DActorData>& ActorDataPair : ActorDataCopy)
 	{
 		if (AActor* Actor = ActorDataPair.Key.Get())
 		{
 			RestoreActorPreState(Actor, ActorDataPair.Value);
-			RemoveFromActor(Actor);
 		}
 	}
-
-	ScannedPrimitiveComponents.Empty();
+PRAGMA_ENABLE_DEPRECATION_WARNINGS
 }
 
 #undef LOCTEXT_NAMESPACE
diff --git a/Engine/Plugins/VirtualProduction/Avalanche/Source/AvalancheMask/Private/Mask2D/AvaMask2DMaskState.cpp b/Engine/Plugins/VirtualProduction/Avalanche/Source/AvalancheMask/Private/Mask2D/AvaMask2DMaskState.cpp
new file mode 100644
--- /dev/null
+++ b/Engine/Plugins/VirtualProduction/Avalanche/Source/AvalancheMask/Private/Mask2D/AvaMask2DMaskState.cpp
@@ -0,0 +1,82 @@
+// Copyright Epic Games, Inc. All Rights Reserved.
+
+#include "Mask2D/AvaMask2DMaskState.h"
+#include "AvaMaskSettings.h"
+#include "AvaMaskUtilities.h"
+#include "Materials/MaterialInstanceDynamic.h"
+
+#if WITH_EDITOR
+#include "Components/DMMaterialProperty.h"
+#include "DMDefs.h"
+#include "Material/DynamicMaterialInstance.h"
+#include "Model/DynamicMaterialModelEditorOnlyData.h"
+#endif
+
+void FAvaMask2DMaskState::Apply(UMaterialInstanceDynamic* InMaterial) const
+{
+	if (!InMaterial)
+	{
+		return;
+	}
+
+	InMaterial->BlendMode = UE::AvaMask::Internal::GetTargetBlendMode(MaterialParameters.BlendMode, InMaterial->BlendMode);
+
+#if WITH_EDITOR
+	if (UDynamicMaterialInstance* MaterialDesigner = Cast<UDynamicMaterialInstance>(InMaterial))
+	{
+		UE::AvaMask::Internal::SetOutputProcessor(MaterialDesigner, OutputProcessor, InMaterial->BlendMode);
+	}
+#endif
+
+	MaterialParameters.ApplyToMID(InMaterial);
+}
+
+void FAvaMask2DMaskState::Store(UMaterialInstanceDynamic* InMaterial)
+{
+	if (!InMaterial)
+	{
+		return;
+	}
+
+	MaterialParameters.StoreFromMaterial(InMaterial);
+
+#if WITH_EDITOR
+	if (UDynamicMaterialInstance* MaterialDesigner = Cast<UDynamicMaterialInstance>(InMaterial))
+	{
+		OutputProcessor = UE::AvaMask::Internal::GetOutputProcessor(MaterialDesigner, InMaterial->BlendMode);
+	}
+#endif
+}
+
+FAvaMask2DMaterialState::FAvaMask2DMaterialState(UMaterialInstanceDynamic* InMaterial)
+	: MaterialWeak(InMaterial)
+{
+	MaskState.Store(InMaterial);
+}
+
+bool FAvaMask2DMaterialState::operator==(const UMaterialInstanceDynamic* InMaterial) const
+{
+	return MaterialWeak == InMaterial;
+}
+
+void FAvaMask2DMaterialState::Apply() const
+{
+	if (UMaterialInstanceDynamic* Material = MaterialWeak.Get())
+	{
+		MaskState.Apply(Material);
+	}
+}
+
+void FAvaMask2DMaterialState::Store()
+{
+	if (UMaterialInstanceDynamic* Material = MaterialWeak.Get())
+	{
+		MaskState.Store(Material);
+	}
+}
+
+bool FAvaMask2DMaterialState::IsDependentOf(UMaterialInterface* InMaterial) const
+{
+	UMaterialInstanceDynamic* const Material = MaterialWeak.Get();
+	return Material && Material->IsDependent(InMaterial);
+}
diff --git a/Engine/Plugins/VirtualProduction/Avalanche/Source/AvalancheMask/Private/Mask2D/AvaMask2DReadModifier.cpp b/Engine/Plugins/VirtualProduction/Avalanche/Source/AvalancheMask/Private/Mask2D/AvaMask2DReadModifier.cpp
--- a/Engine/Plugins/VirtualProduction/Avalanche/Source/AvalancheMask/Private/Mask2D/AvaMask2DReadModifier.cpp
+++ b/Engine/Plugins/VirtualProduction/Avalanche/Source/AvalancheMask/Private/Mask2D/AvaMask2DReadModifier.cpp
@@ -1,20 +1,28 @@
 // Copyright Epic Games, Inc. All Rights Reserved.
 
 #include "Mask2D/AvaMask2DReadModifier.h"
-
+#include "AvaDataView.h"
+#include "AvaMaskLog.h"
+#include "AvaMaskSettings.h"
+#include "AvaMaskSubsystem.h"
 #include "AvaMaskUtilities.h"
-#include "Components/PrimitiveComponent.h"
 #include "Engine/CanvasRenderTarget2D.h"
 #include "Engine/Engine.h"
 #include "GeometryMaskCanvas.h"
 #include "GeometryMaskCanvasResource.h"
-#include "Handling/AvaHandleUtilities.h"
-#include "Handling/AvaObjectHandleSubsystem.h"
-#include "Handling/IAvaMaskMaterialCollectionHandle.h"
-#include "Handling/IAvaMaskMaterialHandle.h"
+#include "MaterialBridge/AvaMaterialBridge.h"
+#include "MaterialBridge/AvaMaterialBridgeRegistry.h"
+#include "MaterialBridge/Context/AvaMaterialBridgeApplyStateContext.h"
+#include "MaterialBridge/Context/AvaMaterialBridgeReadSlotContext.h"
+#include "MaterialBridge/Context/AvaMaterialBridgeStoreStateContext.h"
+#include "MaterialBridge/Context/AvaMaterialBridgeWriteSlotContext.h"
+#include "MaterialBridge/Slot/AvaMaterialBridgeReadSlot.h"
+#include "MaterialBridge/Slot/AvaMaterialBridgeWriteSlot.h"
+#include "MaterialBridge/Slot/CommonFeatures/AvaMaterialBridgeBlendModeFeature.h"
 #include "Materials/Material.h"
 #include "Materials/MaterialInstanceDynamic.h"
 #include "Modifiers/ActorModifierCoreStack.h"
+#include "Modifiers/Utilities/ActorModifierCoreUtilities.h"
 
 #if WITH_EDITOR
 #include "Editor.h"
@@ -49,7 +57,18 @@ void UAvaMask2DReadModifier::PostEditChangeProperty(FPropertyChangedEvent& InPro
 
 void UAvaMask2DReadModifier::OnMaterialCompiled(UMaterialInterface* InMaterial)
 {
-	if (InMaterial && InMaterial->GetTypedOuter<AActor>() == GetModifiedActor())
+	if (!InMaterial)
+	{
+		return;
+	}
+
+	const int32 RemovedMaterialPreStateCount = MaterialPreStates.RemoveAll(
+		[InMaterial](const FAvaMask2DMaterialState& InMaterialState)
+		{
+			return InMaterialState.IsDependentOf(InMaterial);
+		});
+
+	if (RemovedMaterialPreStateCount > 0)
 	{
 		MarkModifierDirty();
 	}
@@ -66,107 +85,185 @@ void UAvaMask2DReadModifier::Apply()
 {
 	Super::Apply();
 
-	TArray<FText> FailReasons;
-	if (const UGeometryMaskCanvas* Canvas = GetCurrentCanvas())
+	if (UGeometryMaskCanvas* Canvas = GetCurrentCanvas())
 	{
-		for (TPair<TWeakObjectPtr<AActor>, FAvaMask2DActorData>& ActorDataPair : ActorData)
+		for (TPair<TObjectPtr<UObject>, TInstancedStruct<FAvaMaterialContainerState>>& Pair : MaterialContainerStates)
 		{
-			if (AActor* Actor = ActorDataPair.Key.Get())
+			if (Pair.Key)
 			{
-				FText FailReason;
-				if (!ApplyRead(Actor, ActorDataPair.Value, FailReason))
-				{
-					FailReasons.Add(MoveTemp(FailReason));
-				}
+				ApplyRead(Pair.Key, Canvas);
 			}
 		}
 	}
 
-	if (FailReasons.IsEmpty())
+	Next();
+}
+
+void UAvaMask2DReadModifier::OnRenderStateUpdated(AActor* InActor, UActorComponent* InComponent)
+{
+	Super::OnRenderStateUpdated(InActor, InComponent);
+
+	if (!CanMarkModifierDirty())
+	{
+		return;
+	}
+
+	const UActorModifierCoreBase* MaskModifier = UE::ActorModifierCore::Utilities::FindFirstActorModifierByClass(InActor, UAvaMask2DBaseModifier::StaticClass());
+
+	// Only trigger update if first modifier found in the hierarchy above us is this modifier
+	if (MaskModifier != this)
 	{
-		Next();
+		return;
 	}
-	else
+
+	const UE::Ava::FMaterialBridge* const MaterialBridge = UE::Ava::FMaterialBridgeRegistry::Get().GetMaterialBridge(UE::Ava::FConstDataView(InComponent));
+	if (!MaterialBridge)
 	{
-		Fail(FText::Join(FText::FromStringView(TEXT(", ")), FailReasons));
+		return;
+	}
+
+	bool bHasMaterials = false;
+
+	MaterialBridge->AccessSlots(UE::Ava::FMaterialBridgeReadSlotContext(InComponent), 
+		[&bHasMaterials](const UE::Ava::FMaterialBridgeReadSlotContext&, const UE::Ava::FMaterialBridgeReadSlot& InSlot)->UE::Ava::EControlFlow
+		{
+			bHasMaterials = true;
+			return UE::Ava::EControlFlow::Break;
+		}
+		, /*Options*/{});
+
+	if (bHasMaterials)
+	{
+		RemoveFromActor(InActor);
+		MarkModifierDirty();
 	}
 }
 
-EBlendMode UAvaMask2DReadModifier::GetBlendMode() const
+void UAvaMask2DReadModifier::RemoveFromActor(AActor* InActor)
 {
-	return bUseBlur || bUseFeathering ? EBlendMode::BLEND_Translucent : EBlendMode::BLEND_Masked;
+	Super::RemoveFromActor(InActor);
+
+	MaterialContainerStates.Remove(InActor);
+
+	for (FAvaMask2DMaterialState& PreState : MaterialPreStates)
+	{
+		PreState.TopmostMaterialContainers.Remove(InActor);
+	}
+
+	MaterialPreStates.RemoveAll(
+		[](const FAvaMask2DMaterialState& InState)
+		{
+			return InState.TopmostMaterialContainers.IsEmpty();
+		});
 }
 
-bool UAvaMask2DReadModifier::ApplyRead(AActor* InActor, FAvaMask2DActorData& InActorData)
+EBlendMode UAvaMask2DReadModifier::GetBlendMode() const
 {
-	FText FailReason;
-	return ApplyRead(InActor, InActorData, FailReason);
+	return bUseBlur || bUseFeathering ? EBlendMode::BLEND_Translucent : EBlendMode::BLEND_Masked;
 }
 
-bool UAvaMask2DReadModifier::ApplyRead(AActor* InActor, FAvaMask2DActorData& InActorData, FText& OutFailReason)
+bool UAvaMask2DReadModifier::ApplyRead(TNotNull<UObject*> InMaterialContainer, TNotNull<const UGeometryMaskCanvas*> InCurrentCanvas)
 {
-	// Only add read/write component to actors with primitives
-	if (!ActorSupportsMaskReadWrite(InActor))
+	const UE::Ava::FMaterialBridge* const MaterialBridge = UE::Ava::FMaterialBridgeRegistry::Get().GetMaterialBridge(UE::Ava::FConstDataView(InMaterialContainer));
+	if (!MaterialBridge)
 	{
 		// Ok to continue, just nothing done for this actor
 		return true;
 	}
 
-	if (const UGeometryMaskCanvas* Canvas = GetCurrentCanvas())
-	{
-		// Get MaterialCollectionHandle
-		const TSharedPtr<IAvaMaskMaterialCollectionHandle> MaterialCollectionHandle =
-			UE::Ava::Internal::FindOrAddHandleByLambda(
-				MaterialCollectionHandles
-				, InActor
-				, [this, InActor]()
-				{
-					return GetObjectHandleSubsystem()->MakeCollectionHandle<IAvaMaskMaterialCollectionHandle>(InActor, UE::AvaMask::Internal::HandleTag);
-				});
+	FGeometryMaskDrawingContext DrawingContext(InCurrentCanvas->GetCanvasId().Level);
+	const FIntPoint ViewportSize = InCurrentCanvas->GetResource()->GetMaxViewportSize();
+	const FVector2f ViewportPadding = FVector2f::One() - FVector2f(ViewportSize + InCurrentCanvas->GetResource()->GetViewportPadding(DrawingContext)) / ViewportSize;
 
-		const UTexture* CanvasTexture = TryResolveCanvasTexture(InActor, InActorData);
-		if (!CanvasTexture)
+	FAvaMask2DMaskState MaskedState;
+
+#if WITH_EDITOR
+	MaskedState.OutputProcessor = GetDefault<UAvaMaskSettings>()->GetMaterialFunction();
+#endif
+	MaskedState.MaterialParameters.BlendMode = GetBlendMode();		
+	MaskedState.MaterialParameters.CanvasName = InCurrentCanvas->GetCanvasName();
+	MaskedState.MaterialParameters.Texture = InCurrentCanvas->GetTexture();
+	MaskedState.MaterialParameters.BaseOpacity = BaseOpacity;
+	MaskedState.MaterialParameters.Channel = InCurrentCanvas->GetColorChannel();
+	MaskedState.MaterialParameters.ChannelAsVector = UE::GeometryMask::MaskChannelEnumToVector[InCurrentCanvas->GetColorChannel()];
+	MaskedState.MaterialParameters.bInvert = bInverted;
+	MaskedState.MaterialParameters.Padding = ViewportPadding;
+	MaskedState.MaterialParameters.bApplyFeathering = bUseFeathering;
+	MaskedState.MaterialParameters.OuterFeatherRadius = OuterFeatherRadius;
+	MaskedState.MaterialParameters.InnerFeatherRadius = InnerFeatherRadius;
+
+	MaterialBridge->AccessSlots(UE::Ava::FMaterialBridgeWriteSlotContext(InMaterialContainer)
+		, [This=this, &MaskedState](const UE::Ava::FMaterialBridgeWriteSlotContext& InContext, UE::Ava::FMaterialBridgeWriteSlot& InSlot)->UE::Ava::EControlFlow
 		{
-			return true;
+			This->ApplyMaskState(InContext, InSlot, MaskedState);
+			return UE::Ava::EControlFlow::Continue;
 		}
+		, /*Options*/{});
 
-		{
-			FAvaMask2DSubjectParameters ApplyParameters;
-			ApplyParameters.MaterialParameters.BlendMode = GetBlendMode();		
-			ApplyParameters.MaterialParameters.CanvasName = Canvas->GetCanvasName();
-			ApplyParameters.MaterialParameters.Texture = Canvas->GetTexture();
-			ApplyParameters.MaterialParameters.BaseOpacity = BaseOpacity;
-			ApplyParameters.MaterialParameters.Channel = Canvas->GetColorChannel();
-			ApplyParameters.MaterialParameters.ChannelAsVector = UE::GeometryMask::MaskChannelEnumToVector[ApplyParameters.MaterialParameters.Channel];
-			ApplyParameters.MaterialParameters.bInvert = bInverted;
-
-			FGeometryMaskDrawingContext DrawingContext(Canvas->GetCanvasId().Level);
-
-			const FIntPoint ViewportSize = Canvas->GetResource()->GetMaxViewportSize();
-			const FVector2f ViewportPadding = FVector2f::One() - FVector2f(ViewportSize + Canvas->GetResource()->GetViewportPadding(DrawingContext)) / ViewportSize;
-			ApplyParameters.MaterialParameters.Padding = ViewportPadding;
-
-			ApplyParameters.MaterialParameters.bApplyFeathering = bUseFeathering;
-			ApplyParameters.MaterialParameters.OuterFeatherRadius = OuterFeatherRadius;
-			ApplyParameters.MaterialParameters.InnerFeatherRadius = InnerFeatherRadius;
-
-			FInstancedStruct* MaterialCollectionData = MaterialCollectionHandleData.Find(InActor);
-			if (!MaterialCollectionData)
-			{
-				MaterialCollectionData = &MaterialCollectionHandleData.Add(InActor, MaterialCollectionHandle->MakeDataStruct());
-			}
-			check(MaterialCollectionData);
+	return true;
+}
+
+void UAvaMask2DReadModifier::ApplyMaskState(const UE::Ava::FMaterialBridgeWriteSlotContext& InContext, UE::Ava::FMaterialBridgeWriteSlot& InSlot, const FAvaMask2DMaskState& InMaskState)
+{
+	// Request the container to support the current blend mode.
+	// Do this first as the slot material could change from this request.
+	FAvaMaterialBridgeBlendModeFeature BlendModeFeature;
+	BlendModeFeature.BlendMode = GetBlendMode();
+	InSlot.RequestFeature(BlendModeFeature);
+
+	UMaterialInterface* const SlotMaterial = InSlot.GetMaterial() ? InSlot.GetMaterial() : UAvaMaskSubsystem::StaticGetDefaultMaskMaterial();
+	if (!SlotMaterial)
+	{
+		return;
+	}
 
-			MaterialCollectionHandle->ApplyModifiedState(ApplyParameters, *MaterialCollectionData);
+	// Re-use slot's existing MID
+	if (UMaterialInstanceDynamic* MID = Cast<UMaterialInstanceDynamic>(SlotMaterial))
+	{
+		// Only store the pre-state if it doesn't exist already.
+		FAvaMask2DMaterialState* MaterialPreState = MaterialPreStates.FindByKey(MID);
+		if (!MaterialPreState)
+		{
+			MaterialPreState = &MaterialPreStates.Emplace_GetRef(MID);
 		}
 
-		// Only check for valid materials here, after ApplyModifiedState has occured once
-		if (!MaterialCollectionHandle->ValidateMaterials(OutFailReason))
+		if (UObject* TopmostMaterialContainer = InContext.GetTopmostContext().GetMaterialContainerObject())
 		{
-			return false;
+			MaterialPreState->TopmostMaterialContainers.Add(TopmostMaterialContainer);
 		}
+		InMaskState.Apply(MID);
+		return;
 	}
 
+	// Slot's material is an MIC or a UMaterial. Create an MID out of these
+	if (UObject* MaterialContainer = InContext.GetMaterialContainerObject())
+	{
+		UMaterialInstanceDynamic* NewMID = UMaterialInstanceDynamic::Create(SlotMaterial, MaterialContainer);
+		InSlot.SetMaterial(NewMID);
+
+		// no need to store material prestate as it's a new material. The Material Bridge Slot state will restore to the older material.
+		InMaskState.Apply(NewMID);
+	}
+}
+
+bool UAvaMask2DReadModifier::ApplyRead(AActor* InActor, FAvaMask2DActorData& InActorData)
+{
+	FText FailReason;
+PRAGMA_DISABLE_DEPRECATION_WARNINGS
+	return ApplyRead(InActor, InActorData, FailReason);
+PRAGMA_ENABLE_DEPRECATION_WARNINGS
+}
+
+bool UAvaMask2DReadModifier::ApplyRead(AActor* InActor, FAvaMask2DActorData& InActorData, FText& OutFailReason)
+{
+	if (InActor)
+	{
+		if (UGeometryMaskCanvas* Canvas = GetCurrentCanvas())
+		{
+			return ApplyRead(InActor, Canvas);
+		}
+	}
+	// Ok to continue, just nothing done for this actor
 	return true;
 }
 
@@ -204,34 +301,71 @@ void UAvaMask2DReadModifier::OnModifierDisabled(EActorModifierCoreDisableReason
 #endif
 }
 
-void UAvaMask2DReadModifier::SetupMaskComponent(UActorComponent* InComponent)
+void UAvaMask2DReadModifier::SavePreState()
 {
-	if (!InComponent)
-	{
-		return;
-	}
+	Super::SavePreState();
 
-	if (IGeometryMaskReadInterface* MaskReader = Cast<IGeometryMaskReadInterface>(InComponent))
-	{
-PRAGMA_DISABLE_DEPRECATION_WARNINGS
-		SetupMaskReadComponent(MaskReader);
-PRAGMA_ENABLE_DEPRECATION_WARNINGS
-	}
-	
-	UAvaMask2DBaseModifier::SetupMaskComponent(InComponent);
+	auto StoreActorState = [&MaterialContainerStateData = MaterialContainerStates](AActor* InActor)
+		{
+			TInstancedStruct<FAvaMaterialContainerState>* StateData = MaterialContainerStateData.Find(InActor);
+			if (StateData)
+			{
+				// Skip elements that already have pre-state
+				return;
+			}
+
+			const UE::Ava::FMaterialBridge* const MaterialBridge = UE::Ava::FMaterialBridgeRegistry::Get().GetMaterialBridge(UE::Ava::FConstDataView(InActor));
+
+			// Actors have a built-in bridge so there should always be a valid material bridge found
+			if (!ensureMsgf(MaterialBridge, TEXT("Could not find a valid material bridge for actor of type '%s'!"), *GetNameSafe(InActor->GetClass())))
+			{
+				return;
+			}
+
+			StateData = &MaterialContainerStateData.Add(InActor);
+			MaterialBridge->StoreState(UE::Ava::FMaterialBridgeStoreStateContext(InActor), StateData, /*Options*/{});
+		};
+
+	ForEachActor<AActor>(
+		[this, &StoreActorState](AActor* InActor)
+		{
+			const UActorModifierCoreBase* MaskModifier = UE::ActorModifierCore::Utilities::FindFirstActorModifierByClass(InActor, UAvaMask2DBaseModifier::StaticClass());
+
+			// If first modifier in the hierarchy is not this modifier, then skip actor
+			if (MaskModifier == this)
+			{
+				StoreActorState(InActor);
+			}
+			return true;
+		}
+		, EActorModifierCoreLookup::SelfAndAllChildren);
 }
 
-void UAvaMask2DReadModifier::SetupMaskReadComponent(IGeometryMaskReadInterface* InMaskReader)
+void UAvaMask2DReadModifier::RestorePreState()
 {
-	if (const UGeometryMaskCanvas* Canvas = GetCurrentCanvas())
+	Super::RestorePreState();
+
+	// Restore material container state
+	TMap<TObjectPtr<UObject>, TInstancedStruct<FAvaMaterialContainerState>> OldContainerStateData = MoveTemp(MaterialContainerStates);
+	MaterialContainerStates.Reset();
+
+	for (TPair<TObjectPtr<UObject>, TInstancedStruct<FAvaMaterialContainerState>>& Pair : OldContainerStateData)
 	{
-		// Set canvas name to read from
-		FGeometryMaskReadParameters ReadParameters = InMaskReader->GetParameters();
-		ReadParameters.CanvasName = Channel;
-		ReadParameters.ColorChannel = Canvas->GetColorChannel();
+		const UE::Ava::FMaterialBridge* const MaterialBridge = UE::Ava::FMaterialBridgeRegistry::Get().GetMaterialBridge(UE::Ava::FConstDataView(Pair.Key));
 
-		InMaskReader->SetParameters(ReadParameters);
+		// Actors have a built-in bridge so there should always be a valid material bridge found
+		if (ensureMsgf(MaterialBridge, TEXT("Could not find a valid material bridge for actor of type '%s'!"), *GetNameSafe(Pair.Key->GetClass())))
+		{
+			MaterialBridge->ApplyState(UE::Ava::FMaterialBridgeApplyStateContext(Pair.Key), Pair.Value, /*Options*/{});
+		}
 	}
+
+	// Restore material state
+	for (const FAvaMask2DMaterialState& PreState : MaterialPreStates)
+	{
+		PreState.Apply();
+	}
+	MaterialPreStates.Reset();
 }
 
 void UAvaMask2DReadModifier::OnBaseOpacityChanged()
@@ -239,4 +373,12 @@ void UAvaMask2DReadModifier::OnBaseOpacityChanged()
 	MarkModifierDirty();
 }
 
+void UAvaMask2DReadModifier::SetupMaskComponent(UActorComponent* InComponent)
+{
+}
+
+void UAvaMask2DReadModifier::SetupMaskReadComponent(IGeometryMaskReadInterface* InMaskReader)
+{
+}
+
 #undef LOCTEXT_NAMESPACE
diff --git a/Engine/Plugins/VirtualProduction/Avalanche/Source/AvalancheMask/Private/Mask2D/AvaMask2DWriteModifier.cpp b/Engine/Plugins/VirtualProduction/Avalanche/Source/AvalancheMask/Private/Mask2D/AvaMask2DWriteModifier.cpp
--- a/Engine/Plugins/VirtualProduction/Avalanche/Source/AvalancheMask/Private/Mask2D/AvaMask2DWriteModifier.cpp
+++ b/Engine/Plugins/VirtualProduction/Avalanche/Source/AvalancheMask/Private/Mask2D/AvaMask2DWriteModifier.cpp
@@ -12,6 +12,7 @@
 #include "Materials/Material.h"
 #include "Materials/MaterialInstanceDynamic.h"
 #include "Modifiers/ActorModifierCoreStack.h"
+#include "Modifiers/Utilities/ActorModifierCoreUtilities.h"
 
 #define LOCTEXT_NAMESPACE "AvaMask2DWriteModifier"
 
@@ -34,46 +35,58 @@ void UAvaMask2DWriteModifier::PostEditChangeProperty(FPropertyChangedEvent& InPr
 void UAvaMask2DWriteModifier::Apply()
 {
 	Super::Apply();
-	
-	if (GetCurrentCanvas())
+
+	if (UGeometryMaskCanvas* Canvas = GetCurrentCanvas())
 	{
-		for (TPair<TWeakObjectPtr<AActor>, FAvaMask2DActorData>& ActorDataPair : ActorData)
-		{
-			if (AActor* Actor = ActorDataPair.Key.Get())
+		ForEachActor<AActor>(
+			[This=this, Canvas](AActor* InActor)
 			{
-				if (!ApplyWrite(Actor, ActorDataPair.Value))
+				const UActorModifierCoreBase* MaskModifier = UE::ActorModifierCore::Utilities::FindFirstActorModifierByClass(InActor, UAvaMask2DBaseModifier::StaticClass());
+
+				// If first modifier in the hierarchy is not this modifier, then skip actor
+				if (MaskModifier == This)
 				{
-					return;
+					This->ApplyWrite(InActor, Canvas);
 				}
+				return true;
 			}
-		}
+			, EActorModifierCoreLookup::SelfAndAllChildren);
 	}
 
 	Next();
 }
 
-bool UAvaMask2DWriteModifier::ApplyWrite(AActor* InActor, FAvaMask2DActorData& InActorData)
+void UAvaMask2DWriteModifier::OnRenderStateUpdated(AActor* InActor, UActorComponent* InComponent)
+{
+	Super::OnRenderStateUpdated(InActor, InComponent);
+
+	if (!CanMarkModifierDirty())
+	{
+		return;
+	}
+
+	const UActorModifierCoreBase* MaskModifier = UE::ActorModifierCore::Utilities::FindFirstActorModifierByClass(InActor, UAvaMask2DBaseModifier::StaticClass());
+
+	// Only trigger update if first modifier found in the hierarchy above us is this modifier
+	if (MaskModifier == this)
+	{
+		RemoveFromActor(InActor);
+		MarkModifierDirty();
+	}
+}
+
+bool UAvaMask2DWriteModifier::ApplyWrite(TNotNull<AActor*> InActor, TNotNull<UGeometryMaskCanvas*> InCanvas)
 {
 	// Only add read/write component to actors with primitives
 	if (ActorSupportsMaskReadWrite(InActor))
 	{
-		if (const UGeometryMaskCanvas* GeometryCanvas = GetCurrentCanvas())
+		if (UGeometryMaskWriteMeshComponent* WriteComponent = FindOrAddMaskComponent<UGeometryMaskWriteMeshComponent>(InActor))
 		{
-			const UTexture* CanvasTexture = TryResolveCanvasTexture(InActor, InActorData);
-			if (!CanvasTexture)
-			{
-				return true;
-			}
-		
-			if (UGeometryMaskWriteMeshComponent* WriteComponent = FindOrAddMaskComponent<UGeometryMaskWriteMeshComponent>(InActor))
-			{
-				FGeometryMaskWriteParameters Parameters = WriteComponent->GetParameters();
-				Parameters.CanvasName = Channel;
-				Parameters.ColorChannel = GeometryCanvas->GetColorChannel();
-				Parameters.OperationType = WriteOperation;
-
-				WriteComponent->SetParameters(Parameters);
-			}
+			FGeometryMaskWriteParameters Parameters = WriteComponent->GetParameters();
+			Parameters.CanvasName = Channel;
+			Parameters.ColorChannel = InCanvas->GetColorChannel();
+			Parameters.OperationType = WriteOperation;
+			WriteComponent->SetParameters(Parameters);
 		}
 	}
 
@@ -94,6 +107,18 @@ bool UAvaMask2DWriteModifier::ApplyWrite(AActor* InActor, FAvaMask2DActorData& I
 	return true;
 }
 
+bool UAvaMask2DWriteModifier::ApplyWrite(AActor* InActor, FAvaMask2DActorData& InActorData)
+{
+	if (InActor)
+	{
+		if (UGeometryMaskCanvas* Canvas = GetCurrentCanvas())
+		{
+			return ApplyWrite(InActor, Canvas);
+		}
+	}
+	return true;
+}
+
 void UAvaMask2DWriteModifier::SetWriteOperation(const EGeometryMaskCompositeOperation InWriteOperation)
 {
 	if (WriteOperation != InWriteOperation)
diff --git a/Engine/Plugins/VirtualProduction/Avalanche/Source/AvalancheMask/Public/AvaMaskTypes.h b/Engine/Plugins/VirtualProduction/Avalanche/Source/AvalancheMask/Public/AvaMaskTypes.h
--- a/Engine/Plugins/VirtualProduction/Avalanche/Source/AvalancheMask/Public/AvaMaskTypes.h
+++ b/Engine/Plugins/VirtualProduction/Avalanche/Source/AvalancheMask/Public/AvaMaskTypes.h
@@ -4,9 +4,9 @@
 
 #include "GeometryMaskTypes.h"
 #include "Materials/MaterialParameters.h"
-
 #include "AvaMaskTypes.generated.h"
 
+class UMaterialInstanceDynamic;
 class UMaterialInterface;
 class UTexture;
 
@@ -40,6 +40,8 @@ struct FAvaMask2DMaterialParameters
 {
 	GENERATED_BODY()
 
+	AVALANCHEMASK_API bool HasSameParameters(const FAvaMask2DMaterialParameters& InOther) const;
+
 	/** Canvas/Channel name. */
 	UPROPERTY()
 	FName CanvasName;
@@ -80,6 +82,9 @@ struct FAvaMask2DMaterialParameters
 	UPROPERTY()
 	TEnumAsByte<EBlendMode> BlendMode = EBlendMode::BLEND_Opaque;
 
+	bool ApplyToMID(UMaterialInstanceDynamic* InMaterial) const;
+	bool StoreFromMaterial(const UMaterialInterface* InMaterial);
+
 	/** Set current values from the input material and named parameters, return true if successful. */
 	bool SetFromMaterial(const UMaterialInterface* InMaterial, const FMaterialParameterInfo& InTextureParameter, const FMaterialParameterInfo& InChannelParameter, const FMaterialParameterInfo& InInvertParameter, const FMaterialParameterInfo& InBaseOpacityParameter, const FMaterialParameterInfo& InPaddingParameter, const FMaterialParameterInfo& InFeatherParameter);
 };
diff --git a/Engine/Plugins/VirtualProduction/Avalanche/Source/AvalancheMask/Public/Mask2D/AvaMask2DBaseModifier.h b/Engine/Plugins/VirtualProduction/Avalanche/Source/AvalancheMask/Public/Mask2D/AvaMask2DBaseModifier.h
--- a/Engine/Plugins/VirtualProduction/Avalanche/Source/AvalancheMask/Public/Mask2D/AvaMask2DBaseModifier.h
+++ b/Engine/Plugins/VirtualProduction/Avalanche/Source/AvalancheMask/Public/Mask2D/AvaMask2DBaseModifier.h
@@ -6,29 +6,27 @@
 #include "AvaPropertyChangeDispatcher.h"
 #include "GeometryMaskTypes.h"
 #include "Modifiers/ActorModifierArrangeBaseModifier.h"
-#include "StructUtils/StructView.h"
+#include "StructUtils/InstancedStruct.h"
 #include "UObject/ObjectPtr.h"
-
+#include "UObject/WeakObjectPtr.h"
 #include "AvaMask2DBaseModifier.generated.h"
 
-class AActor;
-class IAvaActorHandle;
-class IAvaComponentHandle;
+#define UE_API AVALANCHEMASK_API
+
+namespace UE::Ava
+{
+	class FMaterialBridge;
+}
+
 class IAvaMaskMaterialCollectionHandle;
 class IAvaMaskMaterialHandle;
-class IGeometryMaskReadInterface;
-class IGeometryMaskWriteInterface;
-class UActorComponent;
 class UAvaMaskMaterialInstanceSubsystem;
 class UAvaObjectHandleSubsystem;
-class UCanvasRenderTarget2D;
 class UGeometryMaskCanvas;
-class UMaterialInstanceDynamic;
 class UMaterialInterface;
 class UPrimitiveComponent;
 class UTexture;
-struct FAvaShapeParametricMaterial;
-struct FMaterialParameterInfo;
+struct FAvaMaterialContainerState;
 
 UENUM(BlueprintType)
 enum class EAvaMask2DMode : uint8
@@ -53,9 +51,8 @@ struct FAvaMask2DActorData
 };
 
 /** Uses scene actors to create a mask texture and applies it to attached actors */
-UCLASS(Abstract)
-class AVALANCHEMASK_API UAvaMask2DBaseModifier
-	: public UActorModifierArrangeBaseModifier
+UCLASS(Abstract, MinimalAPI)
+class UAvaMask2DBaseModifier : public UActorModifierArrangeBaseModifier
 {
 	GENERATED_BODY()
 
@@ -64,116 +61,142 @@ class AVALANCHEMASK_API UAvaMask2DBaseModifier
 public:
 	//~ Begin UObject
 #if WITH_EDITOR
-	virtual void PostEditChangeProperty(FPropertyChangedEvent& InPropertyChangedEvent) override;
+	UE_API virtual void PostEditChangeProperty(FPropertyChangedEvent& InPropertyChangedEvent) override;
 #endif
 	//~ End UObject
 
-	bool UseParentChannel() const { return bUseParentChannel; }
-	void SetUseParentChannel(const bool bInUseParentChannel);
+	bool UseParentChannel() const
+	{
+		return bUseParentChannel;
+	}
+	UE_API void SetUseParentChannel(const bool bInUseParentChannel);
 
-	const FName GetChannel() const;
-	void SetChannel(FName InChannel);
+	UE_API const FName GetChannel() const;
+	UE_API void SetChannel(FName InChannel);
 
-	bool IsInverted() const { return bInverted; }
-	void SetIsInverted(const bool bInInvert);
+	bool IsInverted() const
+	{
+		return bInverted;
+	}
+	UE_API void SetIsInverted(const bool bInInvert);
 
-	bool UseBlur() const { return bUseBlur; }
-	void UseBlur(bool bInUseBlur);
+	bool UseBlur() const
+	{
+		return bUseBlur;
+	}
+	UE_API void UseBlur(bool bInUseBlur);
 
-	float GetBlurStrength() const { return BlurStrength; }
-	void SetBlurStrength(float InBlurStrength);
+	float GetBlurStrength() const
+	{
+		return BlurStrength;
+	}
+	UE_API void SetBlurStrength(float InBlurStrength);
 
-	bool UseFeathering() const { return bUseFeathering; }
-	void UseFeathering(bool bInUseFeathering);
+	bool UseFeathering() const
+	{
+		return bUseFeathering;
+	}
+	UE_API void UseFeathering(bool bInUseFeathering);
 
-	int32 GetOuterFeatherRadius() const { return OuterFeatherRadius; }
-	void SetOuterFeatherRadius(int32 InFeatherRadius);
+	int32 GetOuterFeatherRadius() const
+	{
+		return OuterFeatherRadius;
+	}
+	UE_API void SetOuterFeatherRadius(int32 InFeatherRadius);
 
-	int32 GetInnerFeatherRadius() const { return InnerFeatherRadius; }
-	void SetInnerFeatherRadius(int32 InFeatherRadius);
+	int32 GetInnerFeatherRadius() const
+	{
+		return InnerFeatherRadius;
+	}
+	UE_API void SetInnerFeatherRadius(int32 InFeatherRadius);
 
 	/** Utility to generate a unique mask name based on the associated Actor */
-	FName GenerateUniqueMaskName() const;
+	UE_API FName GenerateUniqueMaskName() const;
 
 	/** Utility to find an existing modifier on the provided actor or it's parent */
-	static UAvaMask2DBaseModifier* FindMaskModifierOnActor(const AActor* InActor);
+	UE_API static UAvaMask2DBaseModifier* FindMaskModifierOnActor(const AActor* InActor);
 
 protected:
 	//~ Begin UActorModifierCoreBase
-	virtual void OnModifierCDOSetup(FActorModifierCoreMetadata& InMetadata) override;
-	virtual void OnModifierAdded(EActorModifierCoreEnableReason InReason) override;
-	virtual void OnModifierRemoved(EActorModifierCoreDisableReason InReason) override;
-	virtual void OnModifiedActorTransformed() override;
-	virtual void SavePreState() override;
-	virtual void RestorePreState() override;
-	virtual void Apply() override;
+	UE_API virtual void OnModifierCDOSetup(FActorModifierCoreMetadata& InMetadata) override;
+	UE_API virtual void OnModifierAdded(EActorModifierCoreEnableReason InReason) override;
+	UE_API virtual void OnModifierRemoved(EActorModifierCoreDisableReason InReason) override;
+	UE_API virtual void OnModifiedActorTransformed() override;
+	UE_API virtual void RestorePreState() override;
+	UE_API virtual void Apply() override;
 	//~ End UActorModifierCoreBase
 
 	//~ Begin IAvaSceneTreeUpdateModifierExtension
-	virtual void OnSceneTreeTrackedActorParentChanged(int32 InIdx, const TArray<TWeakObjectPtr<AActor>>& InPreviousParentActor, const TArray<TWeakObjectPtr<AActor>>& InNewParentActor) override;
-	virtual void OnSceneTreeTrackedActorChildrenChanged(int32 InIdx, const TSet<TWeakObjectPtr<AActor>>& InPreviousChildrenActors, const TSet<TWeakObjectPtr<AActor>>& InNewChildrenActors) override;
+	UE_API virtual void OnSceneTreeTrackedActorParentChanged(int32 InIdx, const TArray<TWeakObjectPtr<AActor>>& InPreviousParentActor, const TArray<TWeakObjectPtr<AActor>>& InNewParentActor) override;
+	UE_API virtual void OnSceneTreeTrackedActorChildrenChanged(int32 InIdx, const TSet<TWeakObjectPtr<AActor>>& InPreviousChildrenActors, const TSet<TWeakObjectPtr<AActor>>& InNewChildrenActors) override;
 	//~ End IAvaSceneTreeUpdateModifierExtension
 
 	//~ Begin IAvaRenderStateUpdateExtension
-	virtual void OnRenderStateUpdated(AActor* InActor, UActorComponent* InComponent) override;
+	UE_API virtual void OnRenderStateUpdated(AActor* InActor, UActorComponent* InComponent) override;
 	//~ End IAvaRenderStateUpdateExtension
 
-	virtual void SaveActorPreState(AActor* InActor, FAvaMask2DActorData& InActorData);
-	virtual void RestoreActorPreState(AActor* InActor, FAvaMask2DActorData& InActorData);
+	UE_DEPRECATED(5.8, "Actor pre state logic has been moved to the Read modifier.")
+	UE_API virtual void SaveActorPreState(AActor* InActor, FAvaMask2DActorData& InActorData);
 
-	void OnUseParentChannelChanged();
-	void OnChannelChanged();
-	void OnInvertedChanged();
-	void OnBlurChanged();
-	void OnFeatherChanged();
-	void OnCanvasChanged();
+	UE_DEPRECATED(5.8, "Actor pre state logic has been moved to the Read modifier.")
+	UE_API virtual void RestoreActorPreState(AActor* InActor, FAvaMask2DActorData& InActorData);
 
-	void UpdateCanvas();
+	UE_API void OnUseParentChannelChanged();
+	UE_API void OnChannelChanged();
+	UE_API void OnInvertedChanged();
+	UE_API void OnBlurChanged();
+	UE_API void OnFeatherChanged();
+	UE_API void OnCanvasChanged();
+
+	UE_API void UpdateCanvas();
 
 	/** Stores specific properties on the canvas locally. */	
-	void CanvasParamsToLocal();
+	UE_API void CanvasParamsToLocal();
 
 	/** Apply locally stored parameters to the canvas. */
-	void LocalParamsToCanvas();
+	UE_API void LocalParamsToCanvas();
 
-	void SetupChannelName();
-	virtual void SetupMaskComponent(UActorComponent* InComponent);
-	virtual void RemoveFromActor(AActor* InActor);
+	UE_API void SetupChannelName();
+	UE_API virtual void SetupMaskComponent(UActorComponent* InComponent);
+	UE_API virtual void RemoveFromActor(AActor* InActor);
 
-	void OnMaskSetCanvas(const UGeometryMaskCanvas* InCanvas, AActor* InActor);
+	UE_DEPRECATED(5.8, "ActorData has been deprecated. Actor restore takes place in the Material Read Modifier directly.")
+	UE_API void OnMaskSetCanvas(const UGeometryMaskCanvas* InCanvas, AActor* InActor);
 
-	UActorComponent* FindOrAddMaskComponent(TSubclassOf<UActorComponent> InComponentClass, AActor* InActor);
+	UE_API UActorComponent* FindOrAddMaskComponent(TSubclassOf<UActorComponent> InComponentClass, AActor* InActor);
 
-	template<
-		typename InComponentClass
-		UE_REQUIRES(std::is_base_of_v<UActorComponent, InComponentClass>)>
+	template<typename InComponentClass UE_REQUIRES(std::is_base_of_v<UActorComponent, InComponentClass>)>
 	InComponentClass* FindOrAddMaskComponent(AActor* InActor)
 	{
 		return Cast<InComponentClass>(FindOrAddMaskComponent(InComponentClass::StaticClass(), InActor));
 	}
 
-	static bool ActorSupportsMaskReadWrite(const AActor* InActor);
+	UE_API static bool ActorSupportsMaskReadWrite(const AActor* InActor);
 
 	/** Returns true if parent channel was found */
-	bool TryResolveParentChannel();
+	UE_API bool TryResolveParentChannel();
+
+	UE_API void TryResolveCanvas();
+	UE_API UTexture* TryResolveCanvasTexture(AActor* InActor, FAvaMask2DActorData& InActorData);
 
-	void TryResolveCanvas();
-	UTexture* TryResolveCanvasTexture(AActor* InActor, FAvaMask2DActorData& InActorData);
+	UE_DEPRECATED(5.8, "Material Mask Handles have been deprecated. Material Bridge is used instead.")
+	UE_API UAvaObjectHandleSubsystem* GetObjectHandleSubsystem();
 
-	UAvaObjectHandleSubsystem* GetObjectHandleSubsystem();
-	UAvaMaskMaterialInstanceSubsystem* GetMaterialInstanceSubsystem();
+	UE_API UAvaMaskMaterialInstanceSubsystem* GetMaterialInstanceSubsystem();
 
 	/** Returns or resolves the currently referenced canvas. */
-	UGeometryMaskCanvas* GetCurrentCanvas();
+	UE_API UGeometryMaskCanvas* GetCurrentCanvas();
 	
-	virtual void OnMaterialsChanged(UObject* InMaterialOwner, const TArray<TSharedPtr<IAvaMaskMaterialHandle>>& InMaterialHandles);
+	UE_DEPRECATED(5.8, "Material Mask Handles have been deprecated. Material Bridge is used instead.")
+	UE_API virtual void OnMaterialsChanged(UObject* InMaterialOwner, const TArray<TSharedPtr<IAvaMaskMaterialHandle>>& InMaterialHandles);
 
 	/** Restores modified actors to their original state and clears cache data */
-	void RestoreActors();
+	UE_DEPRECATED(5.8, "ActorData has been deprecated. Actor restore takes place in the Material Read Modifier directly.")
+	UE_API void RestoreActors();
 
 #if WITH_EDITOR
 	UFUNCTION(CallInEditor, Category = "Mask2D")
-	void VisualizeMask();
+	UE_API void VisualizeMask();
 #endif
 
 	/** Whether to get the channel from the parent (first one that specifies a mask channel) */
@@ -207,9 +230,10 @@ protected:
 	UPROPERTY(BlueprintReadWrite, Getter, Setter, Category = "Mask2D|Shared", meta = (ClampMin = 0, AllowPrivateAccess = "true", EditCondition = "bUseFeathering", EditConditionHides))
 	int32 InnerFeatherRadius = 16;
 
+	UE_DEPRECATED(5.8, "Material Mask Handles have been deprecated. Material Bridge is used instead.")
 	UPROPERTY(Transient)
 	TWeakObjectPtr<UAvaObjectHandleSubsystem> ObjectHandleSubsystem;
-	
+
 	UPROPERTY(Transient)
 	TWeakObjectPtr<UAvaMaskMaterialInstanceSubsystem> MaterialInstanceSubsystem;
 
@@ -221,8 +245,9 @@ protected:
 	FName AutoChannelName;
 
 	/** Cached actor data to apply/restore */
+	UE_DEPRECATED(5.8, "ActorData has been deprecated. Material Container State is used instead.")
 	UPROPERTY(DuplicateTransient)
-	TMap<TWeakObjectPtr<AActor>, FAvaMask2DActorData> ActorData;
+	TMap<TWeakObjectPtr<AActor>, FAvaMask2DActorData> ActorData_DEPRECATED;
 
 	/** Reference to the last resolved canvas name, as stored in CanvasWeak */
 	UPROPERTY(Transient, DuplicateTransient)
@@ -232,20 +257,26 @@ protected:
 	UPROPERTY(EditAnywhere, Transient, DuplicateTransient, Category = "Mask2D|Shared", NoClear,	meta = (DisplayName = "Canvas", AllowPrivateAccess = "true", ShowInnerProperties, NoResetToDefault))
 	TWeakObjectPtr<UGeometryMaskCanvas> CanvasWeak;
 
+	UE_DEPRECATED(5.8, "Material Collection Handles have been deprecated. Material Container State is used instead.")
 	UPROPERTY()
-	TMap<TWeakObjectPtr<UObject>, FInstancedStruct> MaterialCollectionHandleData;
+	TMap<TWeakObjectPtr<UObject>, FInstancedStruct> MaterialCollectionHandleData_DEPRECATED;
 
+	UE_DEPRECATED(5.8, "Material Handles have been deprecated. Material Pre State is used instead.")
 	UPROPERTY(DuplicateTransient)
-	TMap<TWeakObjectPtr<UMaterialInterface>, FInstancedStruct> MaterialHandleData;
+	TMap<TWeakObjectPtr<UMaterialInterface>, FInstancedStruct> MaterialHandleData_DEPRECATED;
 
+	UE_DEPRECATED(5.8, "Material Collection Handles have been deprecated. Material Container State is used instead.")
 	TMap<TObjectKey<UObject>, TSharedPtr<IAvaMaskMaterialCollectionHandle>> MaterialCollectionHandles;
 
 	/** Last scanned primitive components */
+	UE_DEPRECATED(5.8, "ScannedPrimitiveComponents has been deprecated. Material Container State is used instead.")
 	TSet<TObjectKey<UPrimitiveComponent>> ScannedPrimitiveComponents;
 
 private:
 #if WITH_EDITOR
 	/** Used for PECP */
-	static const TAvaPropertyChangeDispatcher<UAvaMask2DBaseModifier> PropertyChangeDispatcher;
+	UE_API static const TAvaPropertyChangeDispatcher<UAvaMask2DBaseModifier> PropertyChangeDispatcher;
 #endif
 };
+
+#undef UE_API
diff --git a/Engine/Plugins/VirtualProduction/Avalanche/Source/AvalancheMask/Public/Mask2D/AvaMask2DMaskState.h b/Engine/Plugins/VirtualProduction/Avalanche/Source/AvalancheMask/Public/Mask2D/AvaMask2DMaskState.h
new file mode 100644
--- /dev/null
+++ b/Engine/Plugins/VirtualProduction/Avalanche/Source/AvalancheMask/Public/Mask2D/AvaMask2DMaskState.h
@@ -0,0 +1,69 @@
+// Copyright Epic Games, Inc. All Rights Reserved.
+
+#pragma once
+
+#include "AvaMaskTypes.h"
+#include "UObject/WeakObjectPtr.h"
+#include "AvaMask2DMaskState.generated.h"
+
+class UMaterialFunctionInterface;
+class UMaterialInstanceDynamic;
+
+/** Represents the state for a mask to apply to / restore from a material */
+USTRUCT()
+struct FAvaMask2DMaskState
+{
+	GENERATED_BODY()
+
+	FAvaMask2DMaskState() = default;
+
+	/** Applies this state to the given material */
+	void Apply(UMaterialInstanceDynamic* InMaterial) const;
+
+	/** Stores the given material's mask state */
+	void Store(UMaterialInstanceDynamic* InMaterial);
+
+	/** The material parameters to save / restore */
+	UPROPERTY()
+	FAvaMask2DMaterialParameters MaterialParameters;
+
+#if WITH_EDITORONLY_DATA
+	/** The output processor was added to the material if material was from material designer */
+	UPROPERTY()
+	TObjectPtr<UMaterialFunctionInterface> OutputProcessor;
+#endif
+};
+
+/** Represents a material and its mask state */
+USTRUCT()
+struct FAvaMask2DMaterialState
+{
+	GENERATED_BODY()
+
+	FAvaMask2DMaterialState() = default;
+
+	explicit FAvaMask2DMaterialState(UMaterialInstanceDynamic* InMaterial);
+
+	bool operator==(const UMaterialInstanceDynamic* InMaterial) const;
+
+	/** Applies the mask state to the material */
+	void Apply() const;
+
+	/** Stores the material's mask state */
+	void Store();
+
+	/** Determines whether this material is dependent of the given material interface */
+	bool IsDependentOf(UMaterialInterface* InMaterial) const;
+
+	/** The material whose parameter state to save / restore */
+	UPROPERTY()
+	TWeakObjectPtr<UMaterialInstanceDynamic> MaterialWeak;
+
+	/** The mask state of the material */
+	UPROPERTY()
+	FAvaMask2DMaskState MaskState;
+
+	/** All the topmost material containers that this material is used */
+	UPROPERTY()
+	TArray<TWeakObjectPtr<UObject>> TopmostMaterialContainers;
+};
diff --git a/Engine/Plugins/VirtualProduction/Avalanche/Source/AvalancheMask/Public/Mask2D/AvaMask2DReadModifier.h b/Engine/Plugins/VirtualProduction/Avalanche/Source/AvalancheMask/Public/Mask2D/AvaMask2DReadModifier.h
--- a/Engine/Plugins/VirtualProduction/Avalanche/Source/AvalancheMask/Public/Mask2D/AvaMask2DReadModifier.h
+++ b/Engine/Plugins/VirtualProduction/Avalanche/Source/AvalancheMask/Public/Mask2D/AvaMask2DReadModifier.h
@@ -3,14 +3,15 @@
 #pragma once
 
 #include "AvaMask2DBaseModifier.h"
+#include "AvaMask2DMaskState.h"
 #include "AvaPropertyChangeDispatcher.h"
 #include "Modifiers/ActorModifierArrangeBaseModifier.h"
 #include "UObject/ObjectPtr.h"
 #include "UObject/Package.h"
-#include "UObject/WeakObjectPtr.h"
-
 #include "AvaMask2DReadModifier.generated.h"
 
+#define UE_API AVALANCHEMASK_API
+
 class AActor;
 class IAvaActorHandle;
 class IAvaComponentHandle;
@@ -21,8 +22,8 @@ class IGeometryMaskWriteInterface;
 class UActorComponent;
 class UAvaMaskMaterialInstanceSubsystem;
 class UAvaObjectHandleSubsystem;
-class UGeometryMaskCanvas;
 class UCanvasRenderTarget2D;
+class UGeometryMaskCanvas;
 class UMaterialInstanceDynamic;
 class UMaterialInterface;
 class UPrimitiveComponent;
@@ -30,10 +31,15 @@ class UTexture;
 struct FAvaShapeParametricMaterial;
 struct FMaterialParameterInfo;
 
+namespace UE::Ava
+{
+	class FMaterialBridgeWriteSlot;
+	class FMaterialBridgeWriteSlotContext;
+}
+
 /** Uses scene actors to create a mask texture and applies it to attached actors */
-UCLASS(BlueprintType)
-class AVALANCHEMASK_API UAvaMask2DReadModifier
-	: public UAvaMask2DBaseModifier
+UCLASS(BlueprintType, MinimalAPI)
+class UAvaMask2DReadModifier : public UAvaMask2DBaseModifier
 {
 	GENERATED_BODY()
 
@@ -42,41 +48,69 @@ class AVALANCHEMASK_API UAvaMask2DReadModifier
 public:
 	//~ Begin UObject
 #if WITH_EDITOR
-	virtual void PostEditChangeProperty(FPropertyChangedEvent& InPropertyChangedEvent) override;
+	UE_API virtual void PostEditChangeProperty(FPropertyChangedEvent& InPropertyChangedEvent) override;
 #endif
 	//~ End UObject
 
-	float GetBaseOpacity() const { return BaseOpacity; }
-	void SetBaseOpacity(float InBaseOpacity);
+	float GetBaseOpacity() const
+	{
+		return BaseOpacity;
+	}
+	UE_API void SetBaseOpacity(float InBaseOpacity);
 
 protected:
 	//~ Begin UActorModifierCoreBase
-	virtual void OnModifierCDOSetup(FActorModifierCoreMetadata& InMetadata) override;
-	virtual void OnModifierEnabled(EActorModifierCoreEnableReason InReason) override;
-	virtual void OnModifierDisabled(EActorModifierCoreDisableReason InReason) override;
-	virtual void Apply() override;
+	UE_API virtual void OnModifierCDOSetup(FActorModifierCoreMetadata& InMetadata) override;
+	UE_API virtual void OnModifierEnabled(EActorModifierCoreEnableReason InReason) override;
+	UE_API virtual void OnModifierDisabled(EActorModifierCoreDisableReason InReason) override;
+	UE_API virtual void SavePreState() override;
+	UE_API virtual void RestorePreState() override;
+	UE_API virtual void Apply() override;
 	//~ End UActorModifierCoreBase
 
+	//~ Begin UAvaMask2DBaseModifier
+	UE_API virtual void OnRenderStateUpdated(AActor* InActor, UActorComponent* InComponent) override;
+	UE_API virtual void RemoveFromActor(AActor* InActor) override;
+	//~ End UAvaMask2DBaseModifier
+
 	/** Get target blend mode dependent on modifier settings. */
-	EBlendMode GetBlendMode() const;
+	UE_API EBlendMode GetBlendMode() const;
 
-	void OnBaseOpacityChanged();
+	UE_API void OnBaseOpacityChanged();
 
 	UE_DEPRECATED(5.8, "SetupMaskReadComponent is deprecated. Mask Read Modifier no longer allocates a component")
-	virtual void SetupMaskComponent(UActorComponent* InComponent) override;
+	UE_API virtual void SetupMaskComponent(UActorComponent* InComponent) override;
 
 	UE_DEPRECATED(5.8, "SetupMaskReadComponent is deprecated. Mask Read Modifier no longer allocates a component")
-	void SetupMaskReadComponent(IGeometryMaskReadInterface* InMaskReader);
+	UE_API void SetupMaskReadComponent(IGeometryMaskReadInterface* InMaskReader);
+
+	UE_API bool ApplyRead(TNotNull<UObject*> InMaterialContainer, TNotNull<const UGeometryMaskCanvas*> InCurrentCanvas);
+
+	void ApplyMaskState(const UE::Ava::FMaterialBridgeWriteSlotContext& InContext, UE::Ava::FMaterialBridgeWriteSlot& InSlot, const FAvaMask2DMaskState& InMaskState);
 
-	UE_DEPRECATED(5.7, "Use ApplyRead that takes in a fail reason")
-	bool ApplyRead(AActor* InActor, FAvaMask2DActorData& InActorData);
+	UE_DEPRECATED(5.8, "Use ApplyRead that takes in a canvas instead")
+	UE_API bool ApplyRead(AActor* InActor, FAvaMask2DActorData& InActorData);
 
-	bool ApplyRead(AActor* InActor, FAvaMask2DActorData& InActorData, FText& OutFailReason);
+	UE_DEPRECATED(5.8, "Use ApplyRead that takes in a canvas instead")
+	UE_API bool ApplyRead(AActor* InActor, FAvaMask2DActorData& InActorData, FText& OutFailReason);
 
 	/** Base opacity/alpha to use in Read mode */
 	UPROPERTY(EditInstanceOnly, BlueprintReadWrite, Getter, Setter, Category = "Mask2D", meta = (ClampMin = 0.0, ClampMax = 1.0, AllowPrivateAccess = "true"))
 	float BaseOpacity = 0.0f;
 
+	/** The state of the actors to their material containers to save / restore */
+	UPROPERTY()
+	TMap<TObjectPtr<UObject>, TInstancedStruct<FAvaMaterialContainerState>> MaterialContainerStates;
+
+	/** 
+	 * State of the MID parameters prior to applying the mask parameters.
+	 * It is not transient because an MID could be non-transient (e.g. Material Designer).
+	 * Transient MIDs will be null on load, and it's ok since this will signify the saved prestate did not matter.
+	 * This could be optimized by not doing a cleanup prior to serializing removing all transient MIDs.
+	 */
+	UPROPERTY()
+	TArray<FAvaMask2DMaterialState> MaterialPreStates;
+
 private:
 #if WITH_EDITOR
 	/** Used for PECP */
@@ -86,3 +120,5 @@ private:
 	void OnPIEEnded(bool bInIsSimulating);
 #endif
 };
+
+#undef UE_API
diff --git a/Engine/Plugins/VirtualProduction/Avalanche/Source/AvalancheMask/Public/Mask2D/AvaMask2DWriteModifier.h b/Engine/Plugins/VirtualProduction/Avalanche/Source/AvalancheMask/Public/Mask2D/AvaMask2DWriteModifier.h
--- a/Engine/Plugins/VirtualProduction/Avalanche/Source/AvalancheMask/Public/Mask2D/AvaMask2DWriteModifier.h
+++ b/Engine/Plugins/VirtualProduction/Avalanche/Source/AvalancheMask/Public/Mask2D/AvaMask2DWriteModifier.h
@@ -11,6 +11,8 @@
 
 #include "AvaMask2DWriteModifier.generated.h"
 
+#define UE_API AVALANCHEMASK_API
+
 class AActor;
 class IAvaActorHandle;
 class IAvaComponentHandle;
@@ -32,8 +34,7 @@ struct FMaterialParameterInfo;
 
 /** Uses scene actors to create a mask texture and applies it to attached actors */
 UCLASS(BlueprintType)
-class AVALANCHEMASK_API UAvaMask2DWriteModifier
-	: public UAvaMask2DBaseModifier
+class  UAvaMask2DWriteModifier : public UAvaMask2DBaseModifier
 {
 	GENERATED_BODY()
 
@@ -44,27 +45,32 @@ public:
 	{
 		return WriteOperation;
 	}
-	void SetWriteOperation(const EGeometryMaskCompositeOperation InWriteOperation);
+	UE_API void SetWriteOperation(const EGeometryMaskCompositeOperation InWriteOperation);
 
 protected:
 	//~ Begin UObject
 #if WITH_EDITOR
-	virtual void PostEditChangeProperty(FPropertyChangedEvent& InPropertyChangedEvent) override;
+	UE_API virtual void PostEditChangeProperty(FPropertyChangedEvent& InPropertyChangedEvent) override;
 #endif
 	//~ End UObject
 
 	//~ Begin UActorModifierCoreBase
-	virtual void OnModifierCDOSetup(FActorModifierCoreMetadata& InMetadata) override;
-	virtual void Apply() override;
+	UE_API virtual void OnModifierCDOSetup(FActorModifierCoreMetadata& InMetadata) override;
+	UE_API virtual void Apply() override;
 	//~ End UActorModifierCoreBase
 
 	//~ Begin UAvaMask2DBaseModifier
-	virtual void SetupMaskComponent(UActorComponent* InComponent) override;
+	UE_API virtual void OnRenderStateUpdated(AActor* InActor, UActorComponent* InComponent) override;
+	UE_API virtual void SetupMaskComponent(UActorComponent* InComponent) override;
 	//~ End UAvaMask2DBaseModifier
 
-	void OnWriteOperationChanged();
-	void SetupMaskWriteComponent(IGeometryMaskWriteInterface* InMaskWriter);
-	bool ApplyWrite(AActor* InActor, FAvaMask2DActorData& InActorData);
+	UE_API void OnWriteOperationChanged();
+	UE_API void SetupMaskWriteComponent(IGeometryMaskWriteInterface* InMaskWriter);
+
+	UE_API bool ApplyWrite(TNotNull<AActor*> InActor, TNotNull<UGeometryMaskCanvas*> InCanvas);
+
+	UE_DEPRECATED(5.8, "Actor data has been deprecated. Use the overload that takes in a canvas instead")
+	UE_API bool ApplyWrite(AActor* InActor, FAvaMask2DActorData& InActorData);
 
 	/** How to write to the chosen mask channel */
 	UPROPERTY(EditInstanceOnly, BlueprintReadWrite, Getter, Setter, Category = "Mask2D", meta = (ValidEnumValues = "Add, Subtract", AllowPrivateAccess = "true"))
@@ -72,6 +78,8 @@ protected:
 
 #if WITH_EDITOR
 	/** Used for PECP */
-	static const TAvaPropertyChangeDispatcher<UAvaMask2DWriteModifier> PropertyChangeDispatcher;
+	UE_API static const TAvaPropertyChangeDispatcher<UAvaMask2DWriteModifier> PropertyChangeDispatcher;
 #endif
 };
+
+#undef UE_API
