Motion Design: added slot identifier struct + added a material bridge for dynamic mesh components that are managed by the motion design shape builder component

#Synced-CL 48854462

diff --git a/Engine/Plugins/VirtualProduction/Avalanche/Source/AvalancheMaterial/Private/MaterialBridge/Common/AvaDecalComponentMaterialBridge.cpp b/Engine/Plugins/VirtualProduction/Avalanche/Source/AvalancheMaterial/Private/MaterialBridge/Common/AvaDecalComponentMaterialBridge.cpp
--- a/Engine/Plugins/VirtualProduction/Avalanche/Source/AvalancheMaterial/Private/MaterialBridge/Common/AvaDecalComponentMaterialBridge.cpp
+++ b/Engine/Plugins/VirtualProduction/Avalanche/Source/AvalancheMaterial/Private/MaterialBridge/Common/AvaDecalComponentMaterialBridge.cpp
@@ -22,7 +22,7 @@ EControlFlow FDecalComponentMaterialBridge::OnAccessSlots(const FReadSlotContext
 {
 	const UDecalComponent& Component = InContext.MaterialContainer.Get<const UDecalComponent>();
 
-	const FReadSlot Slot(Component.GetDecalMaterial());
+	const FReadSlot Slot(Component.GetDecalMaterial(), FSlotId(0));
 	return InFunc(InContext, Slot);
 }
 
@@ -30,7 +30,7 @@ EControlFlow FDecalComponentMaterialBridge::OnAccessSlots(const FWriteSlotContex
 {
 	UDecalComponent& Component = InContext.MaterialContainer.GetMutable<UDecalComponent>();
 
-	FWriteSlot Slot(Component.GetDecalMaterial());
+	FWriteSlot Slot(Component.GetDecalMaterial(), FSlotId(0));
 
 	const EControlFlow ControlFlow = InFunc(InContext, Slot);
 
diff --git a/Engine/Plugins/VirtualProduction/Avalanche/Source/AvalancheMaterial/Private/MaterialBridge/Common/AvaPrimitiveComponentMaterialBridge.cpp b/Engine/Plugins/VirtualProduction/Avalanche/Source/AvalancheMaterial/Private/MaterialBridge/Common/AvaPrimitiveComponentMaterialBridge.cpp
--- a/Engine/Plugins/VirtualProduction/Avalanche/Source/AvalancheMaterial/Private/MaterialBridge/Common/AvaPrimitiveComponentMaterialBridge.cpp
+++ b/Engine/Plugins/VirtualProduction/Avalanche/Source/AvalancheMaterial/Private/MaterialBridge/Common/AvaPrimitiveComponentMaterialBridge.cpp
@@ -24,7 +24,7 @@ EControlFlow FPrimitiveComponentMaterialBridge::OnAccessSlots(const FReadSlotCon
 
 	for (int32 MaterialIndex = 0; MaterialIndex < Component.GetNumMaterials(); ++MaterialIndex)
 	{
-		const FReadSlot Slot(Component.GetMaterial(MaterialIndex));
+		const FReadSlot Slot(Component.GetMaterial(MaterialIndex), FSlotId(MaterialIndex));
 		if (InFunc(InContext, Slot) == EControlFlow::Break)
 		{
 			return EControlFlow::Break;
@@ -40,7 +40,7 @@ EControlFlow FPrimitiveComponentMaterialBridge::OnAccessSlots(const FWriteSlotCo
 
 	for (int32 MaterialIndex = 0; MaterialIndex < Component.GetNumMaterials(); ++MaterialIndex)
 	{
-		FWriteSlot Slot(Component.GetMaterial(MaterialIndex));
+		FWriteSlot Slot(Component.GetMaterial(MaterialIndex), FSlotId(MaterialIndex));
 		const EControlFlow ControlFlow = InFunc(InContext, Slot);
 
 		if (Component.GetMaterial(MaterialIndex) != Slot.GetMaterial())
diff --git a/Engine/Plugins/VirtualProduction/Avalanche/Source/AvalancheMaterial/Private/MaterialBridge/Common/AvaStaticMeshMaterialBridge.cpp b/Engine/Plugins/VirtualProduction/Avalanche/Source/AvalancheMaterial/Private/MaterialBridge/Common/AvaStaticMeshMaterialBridge.cpp
--- a/Engine/Plugins/VirtualProduction/Avalanche/Source/AvalancheMaterial/Private/MaterialBridge/Common/AvaStaticMeshMaterialBridge.cpp
+++ b/Engine/Plugins/VirtualProduction/Avalanche/Source/AvalancheMaterial/Private/MaterialBridge/Common/AvaStaticMeshMaterialBridge.cpp
@@ -7,10 +7,19 @@
 #include "MaterialBridge/Context/AvaMaterialBridgeWriteSlotContext.h"
 #include "MaterialBridge/Slot/AvaMaterialBridgeReadSlot.h"
 #include "MaterialBridge/Slot/AvaMaterialBridgeWriteSlot.h"
+#include "Misc/EnumerateRange.h"
 
 namespace UE::Ava
 {
 
+namespace Private
+{
+
+static const FLazyName BaseSlot(TEXT("Base"));
+static const FLazyName OverlaySlot(TEXT("Overlay"));
+
+} // UE::Ava::Private
+	
 const UStruct* FStaticMeshMaterialBridge::OnGetBridgedType() const
 {
 	return UStaticMesh::StaticClass();
@@ -35,11 +44,11 @@ EControlFlow FStaticMeshMaterialBridge::OnAccessSlots(const FReadSlotContext& In
 		return EControlFlow::Continue;
 	}
 
-	for (const FStaticMaterial& StaticMaterial : StaticMesh.GetStaticMaterials())
+	for (TConstEnumerateRef<FStaticMaterial> StaticMaterial : EnumerateRange(StaticMesh.GetStaticMaterials()))
 	{
 		// Base Material
 		{
-			const FReadSlot Slot(StaticMaterial.MaterialInterface);
+			const FReadSlot Slot(StaticMaterial->MaterialInterface, FSlotId(Private::BaseSlot, StaticMaterial.GetIndex()));
 			if (InFunc(InContext, Slot) == EControlFlow::Break)
 			{
 				return EControlFlow::Break;
@@ -48,7 +57,7 @@ EControlFlow FStaticMeshMaterialBridge::OnAccessSlots(const FReadSlotContext& In
 
 		// Overlay Material
 		{
-			const FReadSlot Slot(StaticMaterial.OverlayMaterialInterface);
+			const FReadSlot Slot(StaticMaterial->OverlayMaterialInterface, FSlotId(Private::OverlaySlot, StaticMaterial.GetIndex()));
 			if (InFunc(InContext, Slot) == EControlFlow::Break)
 			{
 				return EControlFlow::Break;
@@ -77,14 +86,14 @@ EControlFlow FStaticMeshMaterialBridge::OnAccessSlots(const FWriteSlotContext& I
 		return EControlFlow::Continue;
 	}
 
-	for (FStaticMaterial& StaticMaterial : StaticMesh.GetStaticMaterials())
+	for (TEnumerateRef<FStaticMaterial> StaticMaterial : EnumerateRange(StaticMesh.GetStaticMaterials()))
 	{
 		// Base Material
 		{
-			FWriteSlot Slot(StaticMaterial.MaterialInterface);
+			FWriteSlot Slot(StaticMaterial->MaterialInterface, FSlotId(Private::BaseSlot, StaticMaterial.GetIndex()));
 
 			const EControlFlow ControlFlow = InFunc(InContext, Slot);
-			StaticMaterial.MaterialInterface = Slot.GetMaterial();
+			StaticMaterial->MaterialInterface = Slot.GetMaterial();
 
 			if (ControlFlow == EControlFlow::Break)
 			{
@@ -94,10 +103,10 @@ EControlFlow FStaticMeshMaterialBridge::OnAccessSlots(const FWriteSlotContext& I
 
 		// Overlay Material
 		{
-			FWriteSlot Slot(StaticMaterial.OverlayMaterialInterface);
+			FWriteSlot Slot(StaticMaterial->OverlayMaterialInterface, FSlotId(Private::OverlaySlot, StaticMaterial.GetIndex()));
 
 			const EControlFlow ControlFlow = InFunc(InContext, Slot);
-			StaticMaterial.OverlayMaterialInterface = Slot.GetMaterial();
+			StaticMaterial->OverlayMaterialInterface = Slot.GetMaterial();
 
 			if (ControlFlow == EControlFlow::Break)
 			{
diff --git a/Engine/Plugins/VirtualProduction/Avalanche/Source/AvalancheMaterial/Private/MaterialBridge/Context/AvaMaterialBridgeContext.cpp b/Engine/Plugins/VirtualProduction/Avalanche/Source/AvalancheMaterial/Private/MaterialBridge/Context/AvaMaterialBridgeContext.cpp
--- a/Engine/Plugins/VirtualProduction/Avalanche/Source/AvalancheMaterial/Private/MaterialBridge/Context/AvaMaterialBridgeContext.cpp
+++ b/Engine/Plugins/VirtualProduction/Avalanche/Source/AvalancheMaterial/Private/MaterialBridge/Context/AvaMaterialBridgeContext.cpp
@@ -21,6 +21,13 @@ TMaterialBridgeContext<InContextType>::TMaterialBridgeContext(FDataViewType InMa
 {
 }
 
+template<typename InContextType>
+TMaterialBridgeContext<InContextType>::TMaterialBridgeContext(FDataViewType InMaterialContainer, TNotNull<const InContextType*> InParentContext)
+	: MaterialContainer(InMaterialContainer)
+	, ParentContext(InParentContext)
+{
+}
+
 template<typename InContextType>
 TMaterialBridgeContext<InContextType>::TMaterialBridgeContext(UObjectType* InMaterialContainer)
 	: MaterialContainer(InMaterialContainer)
diff --git a/Engine/Plugins/VirtualProduction/Avalanche/Source/AvalancheMaterial/Private/MaterialBridge/Slot/AvaMaterialBridgeReadSlot.cpp b/Engine/Plugins/VirtualProduction/Avalanche/Source/AvalancheMaterial/Private/MaterialBridge/Slot/AvaMaterialBridgeReadSlot.cpp
deleted file mode 100644
diff --git a/Engine/Plugins/VirtualProduction/Avalanche/Source/AvalancheMaterial/Private/MaterialBridge/Slot/AvaMaterialBridgeSlot.cpp b/Engine/Plugins/VirtualProduction/Avalanche/Source/AvalancheMaterial/Private/MaterialBridge/Slot/AvaMaterialBridgeSlot.cpp
new file mode 100644
--- /dev/null
+++ b/Engine/Plugins/VirtualProduction/Avalanche/Source/AvalancheMaterial/Private/MaterialBridge/Slot/AvaMaterialBridgeSlot.cpp
@@ -0,0 +1,14 @@
+// Copyright Epic Games, Inc. All Rights Reserved.
+
+#include "MaterialBridge/Slot/AvaMaterialBridgeSlot.h"
+
+namespace UE::Ava
+{
+
+FMaterialBridgeSlot::FMaterialBridgeSlot(UMaterialInterface* InMaterial, const FSlotId& InSlotId)
+	: Material(InMaterial)
+	, SlotId(InSlotId)
+{
+}
+
+} // UE::Ava
diff --git a/Engine/Plugins/VirtualProduction/Avalanche/Source/AvalancheMaterial/Private/MaterialBridge/Slot/AvaMaterialBridgeSlotId.cpp b/Engine/Plugins/VirtualProduction/Avalanche/Source/AvalancheMaterial/Private/MaterialBridge/Slot/AvaMaterialBridgeSlotId.cpp
new file mode 100644
--- /dev/null
+++ b/Engine/Plugins/VirtualProduction/Avalanche/Source/AvalancheMaterial/Private/MaterialBridge/Slot/AvaMaterialBridgeSlotId.cpp
@@ -0,0 +1,24 @@
+// Copyright Epic Games, Inc. All Rights Reserved.
+
+#include "MaterialBridge/Slot/AvaMaterialBridgeSlotId.h"
+
+FAvaMaterialBridgeSlotId::FAvaMaterialBridgeSlotId(int32 InSlotIndex)
+	: Index(InSlotIndex)
+{
+}
+
+FAvaMaterialBridgeSlotId::FAvaMaterialBridgeSlotId(FName InSlotName)
+	: Name(InSlotName)
+{
+}
+
+FAvaMaterialBridgeSlotId::FAvaMaterialBridgeSlotId(FName InSlotName, int32 InSlotIndex)
+	: Name(InSlotName)
+	, Index(InSlotIndex)
+{
+}
+
+bool FAvaMaterialBridgeSlotId::IsValid() const
+{
+	return Name != NAME_None || Index != INDEX_NONE;
+}
diff --git a/Engine/Plugins/VirtualProduction/Avalanche/Source/AvalancheMaterial/Private/MaterialBridge/Slot/AvaMaterialBridgeWriteSlot.cpp b/Engine/Plugins/VirtualProduction/Avalanche/Source/AvalancheMaterial/Private/MaterialBridge/Slot/AvaMaterialBridgeWriteSlot.cpp
--- a/Engine/Plugins/VirtualProduction/Avalanche/Source/AvalancheMaterial/Private/MaterialBridge/Slot/AvaMaterialBridgeWriteSlot.cpp
+++ b/Engine/Plugins/VirtualProduction/Avalanche/Source/AvalancheMaterial/Private/MaterialBridge/Slot/AvaMaterialBridgeWriteSlot.cpp
@@ -5,11 +5,6 @@
 namespace UE::Ava
 {
 
-FMaterialBridgeWriteSlot::FMaterialBridgeWriteSlot(UMaterialInterface* InMaterial)
-	: Material(InMaterial)
-{
-}
-
 void FMaterialBridgeWriteSlot::SetMaterial(UMaterialInterface* InMaterial)
 {
 	if (InMaterial != Material)
diff --git a/Engine/Plugins/VirtualProduction/Avalanche/Source/AvalancheMaterial/Public/MaterialBridge/AvaMaterialBridge.h b/Engine/Plugins/VirtualProduction/Avalanche/Source/AvalancheMaterial/Public/MaterialBridge/AvaMaterialBridge.h
--- a/Engine/Plugins/VirtualProduction/Avalanche/Source/AvalancheMaterial/Public/MaterialBridge/AvaMaterialBridge.h
+++ b/Engine/Plugins/VirtualProduction/Avalanche/Source/AvalancheMaterial/Public/MaterialBridge/AvaMaterialBridge.h
@@ -11,6 +11,7 @@
 
 #define UE_API AVALANCHEMATERIAL_API
 
+struct FAvaMaterialBridgeSlotId;
 struct FAvaMaterialContainerState;
 
 namespace UE::Ava
@@ -57,6 +58,7 @@ namespace UE::Ava
 class FMaterialBridge
 {
 public:
+	using FSlotId            = FAvaMaterialBridgeSlotId;
 	using FReadSlot          = FMaterialBridgeReadSlot;
 	using FReadSlotOptions   = FMaterialBridgeReadSlotOptions;
 	using FReadSlotContext   = FMaterialBridgeReadSlotContext;
diff --git a/Engine/Plugins/VirtualProduction/Avalanche/Source/AvalancheMaterial/Public/MaterialBridge/Context/AvaMaterialBridgeContext.h b/Engine/Plugins/VirtualProduction/Avalanche/Source/AvalancheMaterial/Public/MaterialBridge/Context/AvaMaterialBridgeContext.h
--- a/Engine/Plugins/VirtualProduction/Avalanche/Source/AvalancheMaterial/Public/MaterialBridge/Context/AvaMaterialBridgeContext.h
+++ b/Engine/Plugins/VirtualProduction/Avalanche/Source/AvalancheMaterial/Public/MaterialBridge/Context/AvaMaterialBridgeContext.h
@@ -52,6 +52,7 @@ public:
 	TMaterialBridgeContext() = default;
 
 	UE_API explicit TMaterialBridgeContext(FDataViewType InMaterialContainer);
+	UE_API explicit TMaterialBridgeContext(FDataViewType InMaterialContainer, TNotNull<const InContextType*> InParentContext);
 	UE_API explicit TMaterialBridgeContext(UObjectType* InMaterialContainer);
 	UE_API explicit TMaterialBridgeContext(UObjectType* InMaterialContainer, TNotNull<const InContextType*> InParentContext);
 
diff --git a/Engine/Plugins/VirtualProduction/Avalanche/Source/AvalancheMaterial/Public/MaterialBridge/Slot/AvaMaterialBridgeReadSlot.h b/Engine/Plugins/VirtualProduction/Avalanche/Source/AvalancheMaterial/Public/MaterialBridge/Slot/AvaMaterialBridgeReadSlot.h
--- a/Engine/Plugins/VirtualProduction/Avalanche/Source/AvalancheMaterial/Public/MaterialBridge/Slot/AvaMaterialBridgeReadSlot.h
+++ b/Engine/Plugins/VirtualProduction/Avalanche/Source/AvalancheMaterial/Public/MaterialBridge/Slot/AvaMaterialBridgeReadSlot.h
@@ -2,30 +2,19 @@
 
 #pragma once
 
-#include "UObject/ObjectPtr.h"
+#include "AvaMaterialBridgeSlot.h"
 
 #define UE_API AVALANCHEMATERIAL_API
 
-class UMaterialInterface;
 
 namespace UE::Ava
 {
 
 /** Short-lived representation of an accessed material slot for read */
-class FMaterialBridgeReadSlot
+class FMaterialBridgeReadSlot : public FMaterialBridgeSlot
 {
 public:
-	UE_API explicit FMaterialBridgeReadSlot(UMaterialInterface* InMaterial);
-
-	/** Gets the material for this slot */
-	UMaterialInterface* GetMaterial() const
-	{
-		return Material;
-	}
-
-protected:
-	/** Material in this slot */
-	TObjectPtr<UMaterialInterface> Material;
+	using FMaterialBridgeSlot::FMaterialBridgeSlot;
 };
 
 } // UE::Ava
diff --git a/Engine/Plugins/VirtualProduction/Avalanche/Source/AvalancheMaterial/Public/MaterialBridge/Slot/AvaMaterialBridgeSlot.h b/Engine/Plugins/VirtualProduction/Avalanche/Source/AvalancheMaterial/Public/MaterialBridge/Slot/AvaMaterialBridgeSlot.h
new file mode 100644
--- /dev/null
+++ b/Engine/Plugins/VirtualProduction/Avalanche/Source/AvalancheMaterial/Public/MaterialBridge/Slot/AvaMaterialBridgeSlot.h
@@ -0,0 +1,48 @@
+// Copyright Epic Games, Inc. All Rights Reserved.
+
+#pragma once
+
+#include "AvaMaterialBridgeSlotId.h"
+#include "UObject/ObjectPtr.h"
+
+#define UE_API AVALANCHEMATERIAL_API
+
+class UMaterialInterface;
+
+namespace UE::Ava
+{
+
+/** Base slot class for both the read and write slot */
+class FMaterialBridgeSlot
+{
+public:
+	using FSlotId = FAvaMaterialBridgeSlotId;
+
+	/** Constructor for when a material slot is indexed */
+	UE_API explicit FMaterialBridgeSlot(UMaterialInterface* InMaterial, const FSlotId& InSlotId);
+
+	virtual ~FMaterialBridgeSlot() = default;
+
+	/** Gets the material for this slot */
+	UMaterialInterface* GetMaterial() const
+	{
+		return Material;
+	}
+
+	/** Gets the identifier for this slot */
+	const FSlotId& GetSlotId() const
+	{
+		return SlotId;
+	}
+
+protected:
+	/** Material in this slot */
+	TObjectPtr<UMaterialInterface> Material;
+
+	/** Identifier for the slot within its container */
+	FSlotId SlotId;
+};
+
+} // UE::Ava
+
+#undef UE_API
diff --git a/Engine/Plugins/VirtualProduction/Avalanche/Source/AvalancheMaterial/Public/MaterialBridge/Slot/AvaMaterialBridgeSlotId.h b/Engine/Plugins/VirtualProduction/Avalanche/Source/AvalancheMaterial/Public/MaterialBridge/Slot/AvaMaterialBridgeSlotId.h
new file mode 100644
--- /dev/null
+++ b/Engine/Plugins/VirtualProduction/Avalanche/Source/AvalancheMaterial/Public/MaterialBridge/Slot/AvaMaterialBridgeSlotId.h
@@ -0,0 +1,48 @@
+// Copyright Epic Games, Inc. All Rights Reserved.
+
+#pragma once
+
+#include "AvaMaterialBridgeSlotId.generated.h"
+
+#define UE_API AVALANCHEMATERIAL_API
+
+/** Holds the data to identify a slot within a material container */
+USTRUCT()
+struct FAvaMaterialBridgeSlotId
+{
+	GENERATED_BODY()
+
+	FAvaMaterialBridgeSlotId() = default;
+
+	/** Constructor for when a material slot can be identified by its index */
+	UE_API explicit FAvaMaterialBridgeSlotId(int32 InSlotIndex);
+
+	/** Constructor for when a material slot can be identified by an FName */
+	UE_API explicit FAvaMaterialBridgeSlotId(FName InSlotName);
+
+	/** Constructor for when a material slot can be identified by a name and index */
+	UE_API explicit FAvaMaterialBridgeSlotId(FName InSlotName, int32 InSlotIndex);
+
+	/** Whether the slot id is valid (i.e. not the defaults) */
+	UE_API bool IsValid() const;
+
+	bool operator==(const FAvaMaterialBridgeSlotId& InOther) const
+	{
+		return Name == InOther.Name 
+			&& Index == InOther.Index;
+	}
+
+	friend uint32 GetTypeHash(const FAvaMaterialBridgeSlotId& InSlotId)
+	{
+		return HashCombineFast(GetTypeHash(InSlotId.Name), GetTypeHash(InSlotId.Index));
+	}
+
+private:
+	/** Slot name, if any */
+	FName Name;
+
+	/** Slot index, if any */
+	int32 Index = INDEX_NONE;
+};
+
+#undef UE_API
diff --git a/Engine/Plugins/VirtualProduction/Avalanche/Source/AvalancheMaterial/Public/MaterialBridge/Slot/AvaMaterialBridgeWriteSlot.h b/Engine/Plugins/VirtualProduction/Avalanche/Source/AvalancheMaterial/Public/MaterialBridge/Slot/AvaMaterialBridgeWriteSlot.h
--- a/Engine/Plugins/VirtualProduction/Avalanche/Source/AvalancheMaterial/Public/MaterialBridge/Slot/AvaMaterialBridgeWriteSlot.h
+++ b/Engine/Plugins/VirtualProduction/Avalanche/Source/AvalancheMaterial/Public/MaterialBridge/Slot/AvaMaterialBridgeWriteSlot.h
@@ -2,30 +2,21 @@
 
 #pragma once
 
+#include "AvaMaterialBridgeSlot.h"
 #include "StructUtils/StructView.h"
-#include "UObject/ObjectPtr.h"
 
 #define UE_API AVALANCHEMATERIAL_API
 
-class UMaterialInterface;
 struct FAvaMaterialBridgeFeature;
 
 namespace UE::Ava
 {
 
 /** Base struct for short-lived representation of an accessed material slot for write */
-class FMaterialBridgeWriteSlot
+class FMaterialBridgeWriteSlot : public FMaterialBridgeSlot
 {
 public:
-	UE_API explicit FMaterialBridgeWriteSlot(UMaterialInterface* InMaterial);
-
-	virtual ~FMaterialBridgeWriteSlot() = default;
-
-	/** Gets the material for this slot */
-	UMaterialInterface* GetMaterial() const
-	{
-		return Material;
-	}
+	using FMaterialBridgeSlot::FMaterialBridgeSlot;
 
 	/** Sets the new material for the slot */
 	UE_API void SetMaterial(UMaterialInterface* InMaterial);
@@ -55,10 +46,6 @@ protected:
 	{
 		return false;
 	}
-
-private:
-	/** Material in this slot */
-	TObjectPtr<UMaterialInterface> Material;
 };
 
 } // UE::Ava
diff --git a/Engine/Plugins/VirtualProduction/Avalanche/Source/AvalancheShapes/Private/AvaShapeMaterialBridge.cpp b/Engine/Plugins/VirtualProduction/Avalanche/Source/AvalancheShapes/Private/AvaShapeMaterialBridge.cpp
--- a/Engine/Plugins/VirtualProduction/Avalanche/Source/AvalancheShapes/Private/AvaShapeMaterialBridge.cpp
+++ b/Engine/Plugins/VirtualProduction/Avalanche/Source/AvalancheShapes/Private/AvaShapeMaterialBridge.cpp
@@ -31,7 +31,7 @@ EAvaShapeParametricMaterialTranslucency GetTargetTranslucency(EBlendMode InBlend
 } // UE::Ava::Private
 
 FShapeMaterialWriteSlot::FShapeMaterialWriteSlot(TNotNull<UAvaShapeDynamicMeshBase*> InShape, int32 InMeshIndex)
-	: FMaterialBridgeWriteSlot(InShape->GetMaterial(InMeshIndex))
+	: FMaterialBridgeWriteSlot(InShape->GetMaterial(InMeshIndex), FSlotId(InMeshIndex))
 	, Shape(InShape)
 	, MeshIndex(InMeshIndex)
 {
@@ -72,7 +72,7 @@ EControlFlow FShapeMaterialBridge::OnAccessSlots(const FReadSlotContext& InConte
 
 	for (const TPair<int32, FAvaShapeMeshData>& Pair : Shape.GetMeshDatas())
 	{
-		const FReadSlot Slot(Pair.Value.GetMaterial());
+		const FReadSlot Slot(Pair.Value.GetMaterial(), FSlotId(Pair.Key));
 		if (InFunc(InContext, Slot) == EControlFlow::Break)
 		{
 			return EControlFlow::Break;
diff --git a/Engine/Plugins/VirtualProduction/Avalanche/Source/AvalancheShapes/Private/AvaShapeMeshComponentMaterialBridge.cpp b/Engine/Plugins/VirtualProduction/Avalanche/Source/AvalancheShapes/Private/AvaShapeMeshComponentMaterialBridge.cpp
new file mode 100644
--- /dev/null
+++ b/Engine/Plugins/VirtualProduction/Avalanche/Source/AvalancheShapes/Private/AvaShapeMeshComponentMaterialBridge.cpp
@@ -0,0 +1,23 @@
+// Copyright Epic Games, Inc. All Rights Reserved.
+
+#include "AvaShapeMeshComponentMaterialBridge.h"
+#include "AvaShapeActor.h"
+#include "Components/DynamicMeshComponent.h"
+
+namespace UE::Ava
+{
+
+const UStruct* FShapeMeshComponentMaterialBridge::OnGetBridgedType() const
+{
+	return UDynamicMeshComponent::StaticClass();
+}
+
+bool FShapeMeshComponentMaterialBridge::OnIsMaterialContainerSupported(FConstDataView InMaterialContainer) const
+{
+	const UDynamicMeshComponent& MeshComponent = InMaterialContainer.Get<const UDynamicMeshComponent>();
+
+	const AAvaShapeActor* ShapeActor = Cast<AAvaShapeActor>(MeshComponent.GetOuter());
+
+	return ShapeActor && ShapeActor->GetShapeMeshComponent() == &MeshComponent;
+}
+} // UE::Ava
diff --git a/Engine/Plugins/VirtualProduction/Avalanche/Source/AvalancheShapes/Private/AvaShapesModule.cpp b/Engine/Plugins/VirtualProduction/Avalanche/Source/AvalancheShapes/Private/AvaShapesModule.cpp
--- a/Engine/Plugins/VirtualProduction/Avalanche/Source/AvalancheShapes/Private/AvaShapesModule.cpp
+++ b/Engine/Plugins/VirtualProduction/Avalanche/Source/AvalancheShapes/Private/AvaShapesModule.cpp
@@ -2,6 +2,7 @@
 
 #include "AvaShapesModule.h"
 #include "AvaShapeMaterialBridge.h"
+#include "AvaShapeMeshComponentMaterialBridge.h"
 #include "MaterialBridge/AvaMaterialBridgeRegistry.h"
 #include "Modules/ModuleManager.h"
 
@@ -9,8 +10,12 @@ IMPLEMENT_MODULE(FAvaShapesModule, AvalancheShapes)
 
 void FAvaShapesModule::StartupModule()
 {
-	constexpr uint32 DefaultPriority = 0;
-
 	UE::Ava::FMaterialBridgeRegistry& MaterialBridgeRegistry = UE::Ava::FMaterialBridgeRegistry::GetMutable();
+
+	constexpr uint32 DefaultPriority = 0;
 	MaterialBridgeRegistry.Register<UE::Ava::FShapeMaterialBridge>(DefaultPriority);
+
+	// Higher than default to be prioritized over registered default material bridges of the same bridged type.
+	constexpr uint32 HigherPriority = 100;
+	MaterialBridgeRegistry.Register<UE::Ava::FShapeMeshComponentMaterialBridge>(HigherPriority);
 }
diff --git a/Engine/Plugins/VirtualProduction/Avalanche/Source/AvalancheShapes/Public/AvaShapeMeshComponentMaterialBridge.h b/Engine/Plugins/VirtualProduction/Avalanche/Source/AvalancheShapes/Public/AvaShapeMeshComponentMaterialBridge.h
new file mode 100644
--- /dev/null
+++ b/Engine/Plugins/VirtualProduction/Avalanche/Source/AvalancheShapes/Public/AvaShapeMeshComponentMaterialBridge.h
@@ -0,0 +1,27 @@
+// Copyright Epic Games, Inc. All Rights Reserved.
+
+#pragma once
+
+#include "MaterialBridge/AvaMaterialBridge.h"
+
+#define UE_API AVALANCHESHAPES_API
+
+namespace UE::Ava
+{
+/**
+ * Material Bridge for shape mesh components that are managed by the shape builders 
+ * Slot accessing and material container state are not implemented by design as the material slot access for these is managed by the shape builder component. 
+ * FShapeMaterialBridge should be used instead.
+ */
+class FShapeMeshComponentMaterialBridge : public FMaterialBridge
+{
+protected:
+	//~ Begin FMaterialBridge
+	UE_API virtual const UStruct* OnGetBridgedType() const override;
+	UE_API virtual bool OnIsMaterialContainerSupported(FConstDataView InMaterialContainer) const override;
+	//~ End FMaterialBridge
+};
+
+} // UE::Ava
+
+#undef UE_API
diff --git a/Engine/Plugins/VirtualProduction/Avalanche/Source/AvalancheText/Private/MaterialBridge/AvaText3DExtensionMaterialBridge.cpp b/Engine/Plugins/VirtualProduction/Avalanche/Source/AvalancheText/Private/MaterialBridge/AvaText3DExtensionMaterialBridge.cpp
--- a/Engine/Plugins/VirtualProduction/Avalanche/Source/AvalancheText/Private/MaterialBridge/AvaText3DExtensionMaterialBridge.cpp
+++ b/Engine/Plugins/VirtualProduction/Avalanche/Source/AvalancheText/Private/MaterialBridge/AvaText3DExtensionMaterialBridge.cpp
@@ -12,8 +12,20 @@
 namespace UE::Ava
 {
 
+namespace Private
+{
+
+// Creates a Slot Id out of the Material Parameters
+FAvaMaterialBridgeSlotId CreateSlotId(const Text3D::Material::FMaterialParameters& InParameters)
+{
+	const int32 GroupIndex = static_cast<int32>(InParameters.Group);
+	return FAvaMaterialBridgeSlotId(InParameters.Tag, GroupIndex);
+}
+
+} // UE::Ava::Private
+
 FText3DMaterialWriteSlot::FText3DMaterialWriteSlot(TNotNull<UText3DMaterialExtensionBase*> InMaterialExtension, const UE::Text3D::Material::FMaterialParameters& InMaterialParameters)
-	: FMaterialBridgeWriteSlot(InMaterialExtension->GetMaterial(InMaterialParameters))
+	: FMaterialBridgeWriteSlot(InMaterialExtension->GetMaterial(InMaterialParameters), Private::CreateSlotId(InMaterialParameters))
 	, BaseMaterialExtension(InMaterialExtension)
 	, MaterialParameters(InMaterialParameters)
 {
@@ -36,9 +48,9 @@ EControlFlow FText3DExtensionMaterialBridge::OnAccessSlots(const FReadSlotContex
 	EControlFlow ControlFlow = EControlFlow::Continue;
 
 	Extension.ForEachMaterial(
-		[&InFunc, &InContext, &ControlFlow](const UE::Text3D::Material::FMaterialParameters&, UMaterialInterface* InMaterial)->bool
+		[&InFunc, &InContext, &ControlFlow](const UE::Text3D::Material::FMaterialParameters& InMaterialParameters, UMaterialInterface* InMaterial)->bool
 		{
-			const FReadSlot Slot(InMaterial);
+			const FReadSlot Slot(InMaterial, Private::CreateSlotId(InMaterialParameters));
 			ControlFlow = InFunc(InContext, Slot);
 			return ControlFlow == EControlFlow::Continue;
 		});
