Motion Design: fixed issue where the original material state data would not persist

#Synced-CL 48854462

diff --git a/Engine/Plugins/VirtualProduction/Avalanche/Source/AvalancheMask/Private/Mask2D/AvaMask2DMaskState.cpp b/Engine/Plugins/VirtualProduction/Avalanche/Source/AvalancheMask/Private/Mask2D/AvaMask2DMaskState.cpp
--- a/Engine/Plugins/VirtualProduction/Avalanche/Source/AvalancheMask/Private/Mask2D/AvaMask2DMaskState.cpp
+++ b/Engine/Plugins/VirtualProduction/Avalanche/Source/AvalancheMask/Private/Mask2D/AvaMask2DMaskState.cpp
@@ -3,6 +3,7 @@
 #include "Mask2D/AvaMask2DMaskState.h"
 #include "AvaMaskSettings.h"
 #include "AvaMaskUtilities.h"
+#include "Materials/Material.h"
 #include "Materials/MaterialInstanceDynamic.h"
 
 #if WITH_EDITOR
@@ -48,18 +49,31 @@ void FAvaMask2DMaskState::Store(UMaterialInstanceDynamic* InMaterial)
 #endif
 }
 
-FAvaMask2DMaterialState::FAvaMask2DMaterialState(UMaterialInstanceDynamic* InMaterial)
-	: MaterialWeak(InMaterial)
+FAvaMask2DMaterialSlotId::FAvaMask2DMaterialSlotId(FString&& InMaterialContainerPath, const FAvaMaterialBridgeSlotId& InSlotId, const UMaterialInterface* InMaterial)
+	: MaterialContainerPath(MoveTemp(InMaterialContainerPath)) 
+	, SlotId(InSlotId)
+	, BaseMaterial(InMaterial ? InMaterial->GetMaterial() : nullptr)
 {
-	MaskState.Store(InMaterial);
 }
 
-bool FAvaMask2DMaterialState::operator==(const UMaterialInstanceDynamic* InMaterial) const
+bool FAvaMask2DMaterialSlotId::operator==(const FAvaMask2DMaterialSlotId& InOther) const
 {
-	return MaterialWeak == InMaterial;
+	return SlotId == InOther.SlotId
+		&& BaseMaterial == InOther.BaseMaterial
+		&& MaterialContainerPath == InOther.MaterialContainerPath;
 }
 
-void FAvaMask2DMaterialState::Apply() const
+FAvaMask2DMaterialMaskState::FAvaMask2DMaterialMaskState(const FAvaMask2DMaterialSlotId& InSlotId)
+	: SlotMaterialId(InSlotId)
+{
+}
+
+bool FAvaMask2DMaterialMaskState::operator==(const FAvaMask2DMaterialSlotId& InSlotId) const
+{
+	return SlotMaterialId == InSlotId;
+}
+
+void FAvaMask2DMaterialMaskState::Apply() const
 {
 	if (UMaterialInstanceDynamic* Material = MaterialWeak.Get())
 	{
@@ -67,7 +81,7 @@ void FAvaMask2DMaterialState::Apply() const
 	}
 }
 
-void FAvaMask2DMaterialState::Store()
+void FAvaMask2DMaterialMaskState::Store()
 {
 	if (UMaterialInstanceDynamic* Material = MaterialWeak.Get())
 	{
@@ -75,7 +89,7 @@ void FAvaMask2DMaterialState::Store()
 	}
 }
 
-bool FAvaMask2DMaterialState::IsDependentOf(UMaterialInterface* InMaterial) const
+bool FAvaMask2DMaterialMaskState::IsDependentOf(UMaterialInterface* InMaterial) const
 {
 	UMaterialInstanceDynamic* const Material = MaterialWeak.Get();
 	return Material && Material->IsDependent(InMaterial);
diff --git a/Engine/Plugins/VirtualProduction/Avalanche/Source/AvalancheMask/Private/Mask2D/AvaMask2DReadModifier.cpp b/Engine/Plugins/VirtualProduction/Avalanche/Source/AvalancheMask/Private/Mask2D/AvaMask2DReadModifier.cpp
--- a/Engine/Plugins/VirtualProduction/Avalanche/Source/AvalancheMask/Private/Mask2D/AvaMask2DReadModifier.cpp
+++ b/Engine/Plugins/VirtualProduction/Avalanche/Source/AvalancheMask/Private/Mask2D/AvaMask2DReadModifier.cpp
@@ -62,13 +62,13 @@ void UAvaMask2DReadModifier::OnMaterialCompiled(UMaterialInterface* InMaterial)
 		return;
 	}
 
-	const int32 RemovedMaterialPreStateCount = MaterialPreStates.RemoveAll(
-		[InMaterial](const FAvaMask2DMaterialState& InMaterialState)
+	const bool bHasMaterialDependency = MaterialStates.ContainsByPredicate(
+		[InMaterial](const FAvaMask2DMaterialMaskState& InMaterialState)
 		{
 			return InMaterialState.IsDependentOf(InMaterial);
 		});
 
-	if (RemovedMaterialPreStateCount > 0)
+	if (bHasMaterialDependency)
 	{
 		MarkModifierDirty();
 	}
@@ -134,7 +134,6 @@ void UAvaMask2DReadModifier::OnRenderStateUpdated(AActor* InActor, UActorCompone
 
 	if (bHasMaterials)
 	{
-		RemoveFromActor(InActor);
 		MarkModifierDirty();
 	}
 }
@@ -145,15 +144,10 @@ void UAvaMask2DReadModifier::RemoveFromActor(AActor* InActor)
 
 	MaterialContainerStates.Remove(InActor);
 
-	for (FAvaMask2DMaterialState& PreState : MaterialPreStates)
-	{
-		PreState.TopmostMaterialContainers.Remove(InActor);
-	}
-
-	MaterialPreStates.RemoveAll(
-		[](const FAvaMask2DMaterialState& InState)
+	MaterialStates.RemoveAll(
+		[InActor](const FAvaMask2DMaterialMaskState& InState)
 		{
-			return InState.TopmostMaterialContainers.IsEmpty();
+			return InState.TopmostMaterialContainer == InActor;
 		});
 }
 
@@ -217,33 +211,34 @@ void UAvaMask2DReadModifier::ApplyMaskState(const UE::Ava::FMaterialBridgeWriteS
 		return;
 	}
 
-	// Re-use slot's existing MID
-	if (UMaterialInstanceDynamic* MID = Cast<UMaterialInstanceDynamic>(SlotMaterial))
+	UObject* const MaterialContainer = InContext.GetMaterialContainerObject();
+	if (!ensure(MaterialContainer))
 	{
-		// Only store the pre-state if it doesn't exist already.
-		FAvaMask2DMaterialState* MaterialPreState = MaterialPreStates.FindByKey(MID);
-		if (!MaterialPreState)
-		{
-			MaterialPreState = &MaterialPreStates.Emplace_GetRef(MID);
-		}
-
-		if (UObject* TopmostMaterialContainer = InContext.GetTopmostContext().GetMaterialContainerObject())
-		{
-			MaterialPreState->TopmostMaterialContainers.Add(TopmostMaterialContainer);
-		}
-		InMaskState.Apply(MID);
 		return;
 	}
 
-	// Slot's material is an MIC or a UMaterial. Create an MID out of these
-	if (UObject* MaterialContainer = InContext.GetMaterialContainerObject())
+	UMaterialInstanceDynamic* MID = Cast<UMaterialInstanceDynamic>(SlotMaterial);
+	if (!MID)
+	{
+		// Slot's material is an MIC or a UMaterial. Create an MID out of these
+		MID = UMaterialInstanceDynamic::Create(SlotMaterial, MaterialContainer);
+	}
+
+	// Update Material State's MID
 	{
-		UMaterialInstanceDynamic* NewMID = UMaterialInstanceDynamic::Create(SlotMaterial, MaterialContainer);
-		InSlot.SetMaterial(NewMID);
+		const UObject* const TopmostContainer = InContext.GetTopmostContext().GetMaterialContainerObject();
+		const UObject* const TopmostContainerOuter = TopmostContainer ? TopmostContainer->GetOuter() : nullptr;
+
+		const FAvaMask2DMaterialSlotId SlotId(MaterialContainer->GetPathName(TopmostContainerOuter), InSlot.GetSlotId(), SlotMaterial);
 
-		// no need to store material prestate as it's a new material. The Material Bridge Slot state will restore to the older material.
-		InMaskState.Apply(NewMID);
+		if (FAvaMask2DMaterialMaskState* MaterialState = MaterialStates.FindByKey(SlotId))
+		{
+			MaterialState->MaterialWeak = MID;
+		}
 	}
+
+	InSlot.SetMaterial(MID);
+	InMaskState.Apply(MID);
 }
 
 bool UAvaMask2DReadModifier::ApplyRead(AActor* InActor, FAvaMask2DActorData& InActorData)
@@ -305,36 +300,15 @@ void UAvaMask2DReadModifier::SavePreState()
 {
 	Super::SavePreState();
 
-	auto StoreActorState = [&MaterialContainerStateData = MaterialContainerStates](AActor* InActor)
-		{
-			TInstancedStruct<FAvaMaterialContainerState>* StateData = MaterialContainerStateData.Find(InActor);
-			if (StateData)
-			{
-				// Skip elements that already have pre-state
-				return;
-			}
-
-			const UE::Ava::FMaterialBridge* const MaterialBridge = UE::Ava::FMaterialBridgeRegistry::Get().GetMaterialBridge(UE::Ava::FConstDataView(InActor));
-
-			// Actors have a built-in bridge so there should always be a valid material bridge found
-			if (!ensureMsgf(MaterialBridge, TEXT("Could not find a valid material bridge for actor of type '%s'!"), *GetNameSafe(InActor->GetClass())))
-			{
-				return;
-			}
-
-			StateData = &MaterialContainerStateData.Add(InActor);
-			MaterialBridge->StoreState(UE::Ava::FMaterialBridgeStoreStateContext(InActor), StateData, /*Options*/{});
-		};
-
 	ForEachActor<AActor>(
-		[this, &StoreActorState](AActor* InActor)
+		[This=this](AActor* InActor)
 		{
 			const UActorModifierCoreBase* MaskModifier = UE::ActorModifierCore::Utilities::FindFirstActorModifierByClass(InActor, UAvaMask2DBaseModifier::StaticClass());
 
 			// If first modifier in the hierarchy is not this modifier, then skip actor
-			if (MaskModifier == this)
+			if (MaskModifier == This)
 			{
-				StoreActorState(InActor);
+				This->StoreActorState(InActor);
 			}
 			return true;
 		}
@@ -346,7 +320,7 @@ void UAvaMask2DReadModifier::RestorePreState()
 	Super::RestorePreState();
 
 	// Restore material container state
-	TMap<TObjectPtr<UObject>, TInstancedStruct<FAvaMaterialContainerState>> OldContainerStateData = MoveTemp(MaterialContainerStates);
+	TMap<TObjectPtr<UObject>, TInstancedStruct<FAvaMaterialContainerState>> OldContainerStateData = (MaterialContainerStates);
 	MaterialContainerStates.Reset();
 
 	for (TPair<TObjectPtr<UObject>, TInstancedStruct<FAvaMaterialContainerState>>& Pair : OldContainerStateData)
@@ -361,11 +335,54 @@ void UAvaMask2DReadModifier::RestorePreState()
 	}
 
 	// Restore material state
-	for (const FAvaMask2DMaterialState& PreState : MaterialPreStates)
+	for (const FAvaMask2DMaterialMaskState& MaterialState : MaterialStates)
+	{
+		MaterialState.Apply();
+	}
+	MaterialStates.Reset();
+}
+
+void UAvaMask2DReadModifier::StoreActorState(AActor* InActor)
+{
+	const UE::Ava::FMaterialBridge* const MaterialBridge = UE::Ava::FMaterialBridgeRegistry::Get().GetMaterialBridge(UE::Ava::FConstDataView(InActor));
+
+	// Actors have a built-in bridge so there should always be a valid material bridge found
+	if (!ensureMsgf(MaterialBridge, TEXT("Could not find a valid material bridge for actor of type '%s'!"), *GetNameSafe(InActor->GetClass())))
+	{
+		return;
+	}
+
+	TInstancedStruct<FAvaMaterialContainerState>* MaterialContainerState = MaterialContainerStates.Find(InActor);
+	if (!MaterialContainerState)
 	{
-		PreState.Apply();
+		// only call store on new states
+		MaterialContainerState = &MaterialContainerStates.Add(InActor);
+		MaterialBridge->StoreState(UE::Ava::FMaterialBridgeStoreStateContext(InActor), MaterialContainerState, /*Options*/{});
 	}
-	MaterialPreStates.Reset();
+
+	MaterialBridge->AccessSlots(UE::Ava::FMaterialBridgeReadSlotContext(InActor), 
+		[&MaterialStates = MaterialStates, InActor](const UE::Ava::FMaterialBridgeReadSlotContext& InContext, const UE::Ava::FMaterialBridgeReadSlot& InSlot)->UE::Ava::EControlFlow
+		{
+			const UObject* const MaterialContainer = InContext.GetMaterialContainerObject();
+			if (!MaterialContainer)
+			{
+				return UE::Ava::EControlFlow::Continue;
+			}
+
+			const FAvaMask2DMaterialSlotId SlotId(MaterialContainer->GetPathName(InActor->GetOuter()), InSlot.GetSlotId(), InSlot.GetMaterial());
+
+			FAvaMask2DMaterialMaskState* MaterialState = MaterialStates.FindByKey(SlotId);
+			if (!MaterialState)
+			{
+				// only call store on new states
+				MaterialState = &MaterialStates.Emplace_GetRef(SlotId);
+				MaterialState->MaterialWeak = Cast<UMaterialInstanceDynamic>(InSlot.GetMaterial());
+				MaterialState->TopmostMaterialContainer = InActor;
+				MaterialState->Store();
+			}
+			return UE::Ava::EControlFlow::Continue;
+		}
+		, /*Options*/{});
 }
 
 void UAvaMask2DReadModifier::OnBaseOpacityChanged()
diff --git a/Engine/Plugins/VirtualProduction/Avalanche/Source/AvalancheMask/Public/Mask2D/AvaMask2DMaskState.h b/Engine/Plugins/VirtualProduction/Avalanche/Source/AvalancheMask/Public/Mask2D/AvaMask2DMaskState.h
--- a/Engine/Plugins/VirtualProduction/Avalanche/Source/AvalancheMask/Public/Mask2D/AvaMask2DMaskState.h
+++ b/Engine/Plugins/VirtualProduction/Avalanche/Source/AvalancheMask/Public/Mask2D/AvaMask2DMaskState.h
@@ -3,9 +3,11 @@
 #pragma once
 
 #include "AvaMaskTypes.h"
+#include "MaterialBridge/Slot/AvaMaterialBridgeSlotId.h"
 #include "UObject/WeakObjectPtr.h"
 #include "AvaMask2DMaskState.generated.h"
 
+class UMaterial;
 class UMaterialFunctionInterface;
 class UMaterialInstanceDynamic;
 
@@ -34,17 +36,45 @@ struct FAvaMask2DMaskState
 #endif
 };
 
+/** Struct identifying a slot by its slot id and material */
+USTRUCT()
+struct FAvaMask2DMaterialSlotId
+{
+	GENERATED_BODY()
+
+	FAvaMask2DMaterialSlotId() = default;
+
+	explicit FAvaMask2DMaterialSlotId(FString&& InMaterialContainerPath, const FAvaMaterialBridgeSlotId& InSlotId, const UMaterialInterface* InMaterial);
+
+	bool operator==(const FAvaMask2DMaterialSlotId& InOther) const;
+
+	/** Object containing the material slot */
+	UPROPERTY()
+	FString MaterialContainerPath;
+
+	/** The slot id to save / restore the material parameters for */
+	UPROPERTY()
+	FAvaMaterialBridgeSlotId SlotId;
+
+	/**
+	 * Base material of the material whose parameter state to save / restore 
+	 * This base is expected to be non-transient.
+	 */
+	UPROPERTY()
+	TObjectPtr<const UMaterial> BaseMaterial;
+};
+
 /** Represents a material and its mask state */
 USTRUCT()
-struct FAvaMask2DMaterialState
+struct FAvaMask2DMaterialMaskState
 {
 	GENERATED_BODY()
 
-	FAvaMask2DMaterialState() = default;
+	FAvaMask2DMaterialMaskState() = default;
 
-	explicit FAvaMask2DMaterialState(UMaterialInstanceDynamic* InMaterial);
+	explicit FAvaMask2DMaterialMaskState(const FAvaMask2DMaterialSlotId& InSlotId);
 
-	bool operator==(const UMaterialInstanceDynamic* InMaterial) const;
+	bool operator==(const FAvaMask2DMaterialSlotId& InSlotId) const;
 
 	/** Applies the mask state to the material */
 	void Apply() const;
@@ -55,15 +85,17 @@ struct FAvaMask2DMaterialState
 	/** Determines whether this material is dependent of the given material interface */
 	bool IsDependentOf(UMaterialInterface* InMaterial) const;
 
-	/** The material whose parameter state to save / restore */
+	/** The slot id to save / restore the material parameters for */
 	UPROPERTY()
-	TWeakObjectPtr<UMaterialInstanceDynamic> MaterialWeak;
+	FAvaMask2DMaterialSlotId SlotMaterialId;
 
 	/** The mask state of the material */
 	UPROPERTY()
 	FAvaMask2DMaskState MaskState;
 
-	/** All the topmost material containers that this material is used */
-	UPROPERTY()
-	TArray<TWeakObjectPtr<UObject>> TopmostMaterialContainers;
+	/** Material instance dynamic assigned for this state */
+	TWeakObjectPtr<UMaterialInstanceDynamic> MaterialWeak;
+
+	/** The topmost material containers that this state belongs to */
+	TWeakObjectPtr<UObject> TopmostMaterialContainer;
 };
diff --git a/Engine/Plugins/VirtualProduction/Avalanche/Source/AvalancheMask/Public/Mask2D/AvaMask2DReadModifier.h b/Engine/Plugins/VirtualProduction/Avalanche/Source/AvalancheMask/Public/Mask2D/AvaMask2DReadModifier.h
--- a/Engine/Plugins/VirtualProduction/Avalanche/Source/AvalancheMask/Public/Mask2D/AvaMask2DReadModifier.h
+++ b/Engine/Plugins/VirtualProduction/Avalanche/Source/AvalancheMask/Public/Mask2D/AvaMask2DReadModifier.h
@@ -94,6 +94,9 @@ protected:
 	UE_DEPRECATED(5.8, "Use ApplyRead that takes in a canvas instead")
 	UE_API bool ApplyRead(AActor* InActor, FAvaMask2DActorData& InActorData, FText& OutFailReason);
 
+	/** Stores the container and material states for a given actor */
+	void StoreActorState(AActor* InActor);
+
 	/** Base opacity/alpha to use in Read mode */
 	UPROPERTY(EditInstanceOnly, BlueprintReadWrite, Getter, Setter, Category = "Mask2D", meta = (ClampMin = 0.0, ClampMax = 1.0, AllowPrivateAccess = "true"))
 	float BaseOpacity = 0.0f;
@@ -102,14 +105,9 @@ protected:
 	UPROPERTY()
 	TMap<TObjectPtr<UObject>, TInstancedStruct<FAvaMaterialContainerState>> MaterialContainerStates;
 
-	/** 
-	 * State of the MID parameters prior to applying the mask parameters.
-	 * It is not transient because an MID could be non-transient (e.g. Material Designer).
-	 * Transient MIDs will be null on load, and it's ok since this will signify the saved prestate did not matter.
-	 * This could be optimized by not doing a cleanup prior to serializing removing all transient MIDs.
-	 */
+	/** State of the materials to save/restore */
 	UPROPERTY()
-	TArray<FAvaMask2DMaterialState> MaterialPreStates;
+	TArray<FAvaMask2DMaterialMaskState> MaterialStates;
 
 private:
 #if WITH_EDITOR
