Motion Design: small static mesh material bridge change that goes through proper propagation of changes in static materials

#Synced-CL 48854462

diff --git a/Engine/Plugins/VirtualProduction/Avalanche/Source/AvalancheMaterial/Private/MaterialBridge/Common/AvaStaticMeshMaterialBridge.cpp b/Engine/Plugins/VirtualProduction/Avalanche/Source/AvalancheMaterial/Private/MaterialBridge/Common/AvaStaticMeshMaterialBridge.cpp
--- a/Engine/Plugins/VirtualProduction/Avalanche/Source/AvalancheMaterial/Private/MaterialBridge/Common/AvaStaticMeshMaterialBridge.cpp
+++ b/Engine/Plugins/VirtualProduction/Avalanche/Source/AvalancheMaterial/Private/MaterialBridge/Common/AvaStaticMeshMaterialBridge.cpp
@@ -86,7 +86,18 @@ EControlFlow FStaticMeshMaterialBridge::OnAccessSlots(const FWriteSlotContext& I
 		return EControlFlow::Continue;
 	}
 
-	for (TEnumerateRef<FStaticMaterial> StaticMaterial : EnumerateRange(StaticMesh.GetStaticMaterials()))
+	TArray<FStaticMaterial> StaticMaterials = StaticMesh.GetStaticMaterials();
+	if (StaticMaterials.IsEmpty())
+	{
+		return EControlFlow::Continue;
+	}
+
+	ON_SCOPE_EXIT
+	{
+		StaticMesh.SetStaticMaterials(StaticMaterials);
+	};
+
+	for (TEnumerateRef<FStaticMaterial> StaticMaterial : EnumerateRange(StaticMaterials))
 	{
 		// Base Material
 		{
