Motion Design: clean up 5.5 deprecated elements

#Synced-CL 49036542

diff --git a/Engine/Plugins/VirtualProduction/Avalanche/Source/Avalanche/Private/Transition/Conditions/AvaSceneContainsTagAttributeCondition.cpp b/Engine/Plugins/VirtualProduction/Avalanche/Source/Avalanche/Private/Transition/Conditions/AvaSceneContainsTagAttributeCondition.cpp
--- a/Engine/Plugins/VirtualProduction/Avalanche/Source/Avalanche/Private/Transition/Conditions/AvaSceneContainsTagAttributeCondition.cpp
+++ b/Engine/Plugins/VirtualProduction/Avalanche/Source/Avalanche/Private/Transition/Conditions/AvaSceneContainsTagAttributeCondition.cpp
@@ -77,24 +77,6 @@ FText FAvaSceneContainsTagAttributeConditionBase::GetDescription(const FGuid& In
 }
 #endif
 
-void FAvaSceneContainsTagAttributeConditionBase::PostLoad(FStateTreeDataView InInstanceDataView)
-{
-	Super::PostLoad(InInstanceDataView);
-
-	PRAGMA_DISABLE_DEPRECATION_WARNINGS
-	if (LayerType_DEPRECATED != EAvaTransitionLayerCompareType::None)
-	{
-		if (FInstanceDataType* InstanceData = UE::AvaTransition::TryGetInstanceData(*this, InInstanceDataView))
-		{
-			InstanceData->SceneType      = SceneType_DEPRECATED;
-			InstanceData->LayerType      = LayerType_DEPRECATED;
-			InstanceData->SpecificLayers = SpecificLayers_DEPRECATED;
-			InstanceData->TagAttribute   = TagAttribute_DEPRECATED;
-		}
-	}
-	PRAGMA_ENABLE_DEPRECATION_WARNINGS
-}
-
 bool FAvaSceneContainsTagAttributeConditionBase::Link(FStateTreeLinker& InLinker)
 {
 	FAvaTransitionCondition::Link(InLinker);
diff --git a/Engine/Plugins/VirtualProduction/Avalanche/Source/Avalanche/Private/Transition/Conditions/AvaTransitionRCControllerMatchCondition.cpp b/Engine/Plugins/VirtualProduction/Avalanche/Source/Avalanche/Private/Transition/Conditions/AvaTransitionRCControllerMatchCondition.cpp
--- a/Engine/Plugins/VirtualProduction/Avalanche/Source/Avalanche/Private/Transition/Conditions/AvaTransitionRCControllerMatchCondition.cpp
+++ b/Engine/Plugins/VirtualProduction/Avalanche/Source/Avalanche/Private/Transition/Conditions/AvaTransitionRCControllerMatchCondition.cpp
@@ -25,22 +25,6 @@ FText FAvaTransitionRCControllerMatchCondition::GetDescription(const FGuid& InId
 }
 #endif
 
-void FAvaTransitionRCControllerMatchCondition::PostLoad(FStateTreeDataView InInstanceDataView)
-{
-	Super::PostLoad(InInstanceDataView);
-
-	PRAGMA_DISABLE_DEPRECATION_WARNINGS
-	if (ControllerId_DEPRECATED.Name != NAME_None)
-	{
-		if (FInstanceDataType* InstanceData = UE::AvaTransition::TryGetInstanceData(*this, InInstanceDataView))
-		{
-			InstanceData->ControllerId        = ControllerId_DEPRECATED;
-			InstanceData->ValueComparisonType = ValueComparisonType_DEPRECATED;
-		}
-	}
-	PRAGMA_ENABLE_DEPRECATION_WARNINGS
-}
-
 bool FAvaTransitionRCControllerMatchCondition::Link(FStateTreeLinker& InLinker)
 {
 	FAvaTransitionCondition::Link(InLinker);
diff --git a/Engine/Plugins/VirtualProduction/Avalanche/Source/Avalanche/Private/Transition/Tasks/AvaSceneTask.cpp b/Engine/Plugins/VirtualProduction/Avalanche/Source/Avalanche/Private/Transition/Tasks/AvaSceneTask.cpp
--- a/Engine/Plugins/VirtualProduction/Avalanche/Source/Avalanche/Private/Transition/Tasks/AvaSceneTask.cpp
+++ b/Engine/Plugins/VirtualProduction/Avalanche/Source/Avalanche/Private/Transition/Tasks/AvaSceneTask.cpp
@@ -7,21 +7,6 @@
 #include "StateTreeExecutionContext.h"
 #include "StateTreeLinker.h"
 
-void FAvaSceneTask::PostLoad(FStateTreeDataView InInstanceDataView)
-{
-	Super::PostLoad(InInstanceDataView);
-
-	PRAGMA_DISABLE_DEPRECATION_WARNINGS
-	if (TagAttribute_DEPRECATED.IsValid())
-	{
-		if (FInstanceDataType* InstanceData = UE::AvaTransition::TryGetInstanceData(*this, InInstanceDataView))
-		{
-			InstanceData->TagAttribute = TagAttribute_DEPRECATED;
-		}
-	}
-	PRAGMA_ENABLE_DEPRECATION_WARNINGS
-}
-
 bool FAvaSceneTask::Link(FStateTreeLinker& InLinker)
 {
 	Super::Link(InLinker);
diff --git a/Engine/Plugins/VirtualProduction/Avalanche/Source/Avalanche/Public/Transition/Conditions/AvaSceneContainsTagAttributeCondition.h b/Engine/Plugins/VirtualProduction/Avalanche/Source/Avalanche/Public/Transition/Conditions/AvaSceneContainsTagAttributeCondition.h
--- a/Engine/Plugins/VirtualProduction/Avalanche/Source/Avalanche/Public/Transition/Conditions/AvaSceneContainsTagAttributeCondition.h
+++ b/Engine/Plugins/VirtualProduction/Avalanche/Source/Avalanche/Public/Transition/Conditions/AvaSceneContainsTagAttributeCondition.h
@@ -40,14 +40,7 @@ struct AVALANCHE_API FAvaSceneContainsTagAttributeConditionBase : public FAvaTra
 
 	using FInstanceDataType = FAvaSceneContainsTagAttributeConditionInstanceData;
 
-	PRAGMA_DISABLE_DEPRECATION_WARNINGS
 	FAvaSceneContainsTagAttributeConditionBase() = default;
-	virtual ~FAvaSceneContainsTagAttributeConditionBase() override = default;
-	FAvaSceneContainsTagAttributeConditionBase(const FAvaSceneContainsTagAttributeConditionBase&) = default;
-	FAvaSceneContainsTagAttributeConditionBase(FAvaSceneContainsTagAttributeConditionBase&&) = default;
-	FAvaSceneContainsTagAttributeConditionBase& operator=(const FAvaSceneContainsTagAttributeConditionBase&) = default;
-	FAvaSceneContainsTagAttributeConditionBase& operator=(FAvaSceneContainsTagAttributeConditionBase&&) = default;
-	PRAGMA_ENABLE_DEPRECATION_WARNINGS
 
 	FAvaSceneContainsTagAttributeConditionBase(bool bInInvertCondition)
 		: bInvertCondition(bInInvertCondition)
@@ -59,7 +52,6 @@ struct AVALANCHE_API FAvaSceneContainsTagAttributeConditionBase : public FAvaTra
 	virtual FText GetDescription(const FGuid& InId, FStateTreeDataView InInstanceDataView, const IStateTreeBindingLookup& InBindingLookup, EStateTreeNodeFormatting InFormatting) const override;
 #endif
 	virtual const UStruct* GetInstanceDataType() const override { return FInstanceDataType::StaticStruct(); }
-	virtual void PostLoad(FStateTreeDataView InInstanceDataView) override;
 	virtual bool Link(FStateTreeLinker& InLinker) override;
 	//~ End FStateTreeNodeBase
 
@@ -71,22 +63,6 @@ struct AVALANCHE_API FAvaSceneContainsTagAttributeConditionBase : public FAvaTra
 
 	TArray<const FAvaTransitionScene*> GetTransitionScenes(FStateTreeExecutionContext& InContext) const;
 
-	UE_DEPRECATED(5.5, "SceneType has been moved to Instance Data")
-	UPROPERTY(meta = (DeprecatedProperty, DeprecationMessage = "Use the Instance Data SceneType instead"))
-	EAvaTransitionSceneType SceneType_DEPRECATED = EAvaTransitionSceneType::This;
-
-	UE_DEPRECATED(5.5, "LayerType has been moved to Instance Data")
-	UPROPERTY(meta = (DeprecatedProperty, DeprecationMessage = "Use the Instance Data LayerType instead"))
-	EAvaTransitionLayerCompareType LayerType_DEPRECATED = EAvaTransitionLayerCompareType::None;
-
-	UE_DEPRECATED(5.5, "SpecificLayers has been moved to Instance Data")
-	UPROPERTY(meta = (DeprecatedProperty, DeprecationMessage = "Use the Instance Data SpecificLayers instead"))
-	FAvaTagHandleContainer SpecificLayers_DEPRECATED;
-
-	UE_DEPRECATED(5.5, "TagAttribute has been moved to Instance Data")
-	UPROPERTY(meta = (DeprecatedProperty, DeprecationMessage = "Use the Instance Data TagAttribute instead"))
-	FAvaTagHandle TagAttribute_DEPRECATED;
-
 protected:
 	bool bInvertCondition = false;
 
diff --git a/Engine/Plugins/VirtualProduction/Avalanche/Source/Avalanche/Public/Transition/Conditions/AvaTransitionRCControllerMatchCondition.h b/Engine/Plugins/VirtualProduction/Avalanche/Source/Avalanche/Public/Transition/Conditions/AvaTransitionRCControllerMatchCondition.h
--- a/Engine/Plugins/VirtualProduction/Avalanche/Source/Avalanche/Public/Transition/Conditions/AvaTransitionRCControllerMatchCondition.h
+++ b/Engine/Plugins/VirtualProduction/Avalanche/Source/Avalanche/Public/Transition/Conditions/AvaTransitionRCControllerMatchCondition.h
@@ -30,21 +30,11 @@ struct AVALANCHE_API FAvaTransitionRCControllerMatchCondition : public FAvaTrans
 
 	using FInstanceDataType = FAvaTransitionRCControllerMatchConditionInstanceData;
 
-	PRAGMA_DISABLE_DEPRECATION_WARNINGS
-	FAvaTransitionRCControllerMatchCondition() = default;
-	virtual ~FAvaTransitionRCControllerMatchCondition() override = default;
-	FAvaTransitionRCControllerMatchCondition(const FAvaTransitionRCControllerMatchCondition&) = default;
-	FAvaTransitionRCControllerMatchCondition(FAvaTransitionRCControllerMatchCondition&&) = default;
-	FAvaTransitionRCControllerMatchCondition& operator=(const FAvaTransitionRCControllerMatchCondition&) = default;
-	FAvaTransitionRCControllerMatchCondition& operator=(FAvaTransitionRCControllerMatchCondition&&) = default;
-	PRAGMA_ENABLE_DEPRECATION_WARNINGS
-
 	//~ Begin FStateTreeNodeBase
 #if WITH_EDITOR
 	virtual FText GetDescription(const FGuid& InId, FStateTreeDataView InInstanceDataView, const IStateTreeBindingLookup& InBindingLookup, EStateTreeNodeFormatting InFormatting) const override;
 #endif
 	virtual const UStruct* GetInstanceDataType() const override { return FInstanceDataType::StaticStruct(); }
-	virtual void PostLoad(FStateTreeDataView InInstanceDataView) override;
 	virtual bool Link(FStateTreeLinker& InLinker) override;
 	//~ End FStateTreeNodeBase
 
@@ -52,13 +42,5 @@ struct AVALANCHE_API FAvaTransitionRCControllerMatchCondition : public FAvaTrans
 	virtual bool TestCondition(FStateTreeExecutionContext& InContext) const override;
 	//~ End FStateTreeConditionBase
 
-	UE_DEPRECATED(5.5, "ControllerId has been moved to Instance Data")
-	UPROPERTY(meta = (DeprecatedProperty, DeprecationMessage = "Use the Instance Data ControllerId instead"))
-	FAvaRCControllerId ControllerId_DEPRECATED;
-
-	UE_DEPRECATED(5.5, "ValueComparisonType has been moved to Instance Data")
-	UPROPERTY(meta = (DeprecatedProperty, DeprecationMessage = "Use the Instance Data ValueComparisonType instead"))
-	EAvaTransitionComparisonResult ValueComparisonType_DEPRECATED = EAvaTransitionComparisonResult::None;
-
 	TStateTreeExternalDataHandle<UAvaSceneSubsystem> SceneSubsystemHandle;
 };
diff --git a/Engine/Plugins/VirtualProduction/Avalanche/Source/Avalanche/Public/Transition/Extensions/IAvaTransitionRCExtension.h b/Engine/Plugins/VirtualProduction/Avalanche/Source/Avalanche/Public/Transition/Extensions/IAvaTransitionRCExtension.h
--- a/Engine/Plugins/VirtualProduction/Avalanche/Source/Avalanche/Public/Transition/Extensions/IAvaTransitionRCExtension.h
+++ b/Engine/Plugins/VirtualProduction/Avalanche/Source/Avalanche/Public/Transition/Extensions/IAvaTransitionRCExtension.h
@@ -15,14 +15,6 @@ class IAvaRCTransitionExtension : public IAvaTransitionExtension
 public:
 	static constexpr const TCHAR* ExtensionIdentifier = TEXT("IAvaRCTransitionExtension");
 
-	UE_DEPRECATED(5.5, "Please use CompareControllers that takes in FAvaTransitionContext instead")
-	virtual EAvaTransitionComparisonResult CompareControllers(const FGuid& InControllerId
-		, const FAvaTransitionScene& InMyScene
-		, const FAvaTransitionScene& InOtherScene) const
-	{
-		return EAvaTransitionComparisonResult::None;
-	}
-
 	/**
 	 * Compares the values of an RC controller with the given ControllerId between two transition contexts
 	 * @param InControllerId the Id of the controller to compare
diff --git a/Engine/Plugins/VirtualProduction/Avalanche/Source/Avalanche/Public/Transition/Tasks/AvaSceneTask.h b/Engine/Plugins/VirtualProduction/Avalanche/Source/Avalanche/Public/Transition/Tasks/AvaSceneTask.h
--- a/Engine/Plugins/VirtualProduction/Avalanche/Source/Avalanche/Public/Transition/Tasks/AvaSceneTask.h
+++ b/Engine/Plugins/VirtualProduction/Avalanche/Source/Avalanche/Public/Transition/Tasks/AvaSceneTask.h
@@ -25,26 +25,12 @@ struct AVALANCHE_API FAvaSceneTask : public FAvaTransitionTask
 
 	using FInstanceDataType = FAvaSceneTaskInstanceData;
 
-	PRAGMA_DISABLE_DEPRECATION_WARNINGS
-	FAvaSceneTask() = default;
-	virtual ~FAvaSceneTask() override = default;
-	FAvaSceneTask(const FAvaSceneTask&) = default;
-	FAvaSceneTask(FAvaSceneTask&&) = default;
-	FAvaSceneTask& operator=(const FAvaSceneTask&) = default;
-	FAvaSceneTask& operator=(FAvaSceneTask&&) = default;
-	PRAGMA_ENABLE_DEPRECATION_WARNINGS
-
 	//~ Begin FStateTreeNodeBase
 	virtual const UStruct* GetInstanceDataType() const override { return FInstanceDataType::StaticStruct(); }
-	virtual void PostLoad(FStateTreeDataView InInstanceDataView) override;
 	virtual bool Link(FStateTreeLinker& InLinker) override;
 	//~ End FStateTreeNodeBase
 
 	IAvaSceneInterface* GetScene(FStateTreeExecutionContext& InContext) const;
 
-	UE_DEPRECATED(5.5, "TagAttribute has been moved to Instance Data")
-	UPROPERTY(meta = (DeprecatedProperty, DeprecationMessage = "Use the Instance Data TagAttribute instead"))
-	FAvaTagHandle TagAttribute_DEPRECATED;
-
 	TStateTreeExternalDataHandle<UAvaSceneSubsystem> SceneSubsystemHandle;
 };
diff --git a/Engine/Plugins/VirtualProduction/Avalanche/Source/AvalancheSequence/Private/AvaSequencePlayer.cpp b/Engine/Plugins/VirtualProduction/Avalanche/Source/AvalancheSequence/Private/AvaSequencePlayer.cpp
--- a/Engine/Plugins/VirtualProduction/Avalanche/Source/AvalancheSequence/Private/AvaSequencePlayer.cpp
+++ b/Engine/Plugins/VirtualProduction/Avalanche/Source/AvalancheSequence/Private/AvaSequencePlayer.cpp
@@ -117,24 +117,6 @@ void UAvaSequencePlayer::ContinueSequence()
 	PlayInternal();
 }
 
-void UAvaSequencePlayer::PreviewFrame()
-{
-	UAvaSequence* PlayedSequence = Cast<UAvaSequence>(Sequence);
-	if (!PlayedSequence)
-	{
-		return;
-	}
-
-	if (const FAvaMark* const Mark = PlayedSequence->GetPreviewMark())
-	{
-		if (!Mark->Frames.IsEmpty())
-		{
-			FMovieSceneSequencePlaybackParams PlaybackParams(FString(Mark->GetLabel()), EUpdatePositionMethod::Play);
-			SetPlaybackPosition(PlaybackParams);
-		}
-	}
-}
-
 void UAvaSequencePlayer::JumpTo(FFrameTime InJumpToFrame, bool bInEvaluate)
 {
 	if (NeedsQueueLatentAction())
diff --git a/Engine/Plugins/VirtualProduction/Avalanche/Source/AvalancheSequence/Private/Transition/AvaTransitionInitializeSequence.cpp b/Engine/Plugins/VirtualProduction/Avalanche/Source/AvalancheSequence/Private/Transition/AvaTransitionInitializeSequence.cpp
--- a/Engine/Plugins/VirtualProduction/Avalanche/Source/AvalancheSequence/Private/Transition/AvaTransitionInitializeSequence.cpp
+++ b/Engine/Plugins/VirtualProduction/Avalanche/Source/AvalancheSequence/Private/Transition/AvaTransitionInitializeSequence.cpp
@@ -8,22 +8,6 @@
 #include "StateTreeExecutionContext.h"
 #include "Transition/AvaTransitionSequenceUtils.h"
 
-void FAvaTransitionInitializeSequence::PostLoad(FStateTreeDataView InInstanceDataView)
-{
-	Super::PostLoad(InInstanceDataView);
-
-	PRAGMA_DISABLE_DEPRECATION_WARNINGS
-	if (QueryType_DEPRECATED != EAvaTransitionSequenceQueryType::None)
-	{
-		if (FInstanceDataType* InstanceData = UE::AvaTransition::TryGetInstanceData(*this, InInstanceDataView))
-		{
-			InstanceData->InitializeTime = InitializeTime_DEPRECATED;
-			InstanceData->PlayMode       = PlayMode_DEPRECATED;
-		}
-	}
-	PRAGMA_ENABLE_DEPRECATION_WARNINGS
-}
-
 EAvaTransitionSequenceWaitType FAvaTransitionInitializeSequence::GetWaitType(FStateTreeExecutionContext& InContext) const
 {
 	return EAvaTransitionSequenceWaitType::NoWait;
diff --git a/Engine/Plugins/VirtualProduction/Avalanche/Source/AvalancheSequence/Private/Transition/AvaTransitionPlaySequenceTask.cpp b/Engine/Plugins/VirtualProduction/Avalanche/Source/AvalancheSequence/Private/Transition/AvaTransitionPlaySequenceTask.cpp
--- a/Engine/Plugins/VirtualProduction/Avalanche/Source/AvalancheSequence/Private/Transition/AvaTransitionPlaySequenceTask.cpp
+++ b/Engine/Plugins/VirtualProduction/Avalanche/Source/AvalancheSequence/Private/Transition/AvaTransitionPlaySequenceTask.cpp
@@ -12,21 +12,6 @@
 
 #define LOCTEXT_NAMESPACE "AvaTransitionPlaySequenceTask"
 
-void FAvaTransitionPlaySequenceTask::PostLoad(FStateTreeDataView InInstanceDataView)
-{
-	Super::PostLoad(InInstanceDataView);
-
-	PRAGMA_DISABLE_DEPRECATION_WARNINGS
-	if (QueryType_DEPRECATED != EAvaTransitionSequenceQueryType::None)
-	{
-		if (FInstanceDataType* InstanceData = UE::AvaTransition::TryGetInstanceData(*this, InInstanceDataView))
-		{
-			InstanceData->PlaySettings = PlaySettings_DEPRECATED;
-		}
-	}
-	PRAGMA_ENABLE_DEPRECATION_WARNINGS
-}
-
 #if WITH_EDITOR
 FText FAvaTransitionPlaySequenceTask::GetDescription(const FGuid& InId, FStateTreeDataView InInstanceDataView, const IStateTreeBindingLookup& InBindingLookup, EStateTreeNodeFormatting InFormatting) const
 {
diff --git a/Engine/Plugins/VirtualProduction/Avalanche/Source/AvalancheSequence/Private/Transition/AvaTransitionSequenceTask.cpp b/Engine/Plugins/VirtualProduction/Avalanche/Source/AvalancheSequence/Private/Transition/AvaTransitionSequenceTask.cpp
--- a/Engine/Plugins/VirtualProduction/Avalanche/Source/AvalancheSequence/Private/Transition/AvaTransitionSequenceTask.cpp
+++ b/Engine/Plugins/VirtualProduction/Avalanche/Source/AvalancheSequence/Private/Transition/AvaTransitionSequenceTask.cpp
@@ -5,21 +5,6 @@
 #include "StateTreeExecutionContext.h"
 #include "Transition/AvaTransitionSequenceUtils.h"
 
-void FAvaTransitionSequenceTask::PostLoad(FStateTreeDataView InInstanceDataView)
-{
-	Super::PostLoad(InInstanceDataView);
-
-	PRAGMA_DISABLE_DEPRECATION_WARNINGS
-	if (QueryType_DEPRECATED != EAvaTransitionSequenceQueryType::None)
-	{
-		if (FInstanceDataType* InstanceData = UE::AvaTransition::TryGetInstanceData(*this, InInstanceDataView))
-		{
-			InstanceData->WaitType = WaitType_DEPRECATED;
-		}
-	}
-	PRAGMA_ENABLE_DEPRECATION_WARNINGS
-}
-
 EAvaTransitionSequenceWaitType FAvaTransitionSequenceTask::GetWaitType(FStateTreeExecutionContext& InContext) const
 {
 	return InContext.GetInstanceData(*this).WaitType;
diff --git a/Engine/Plugins/VirtualProduction/Avalanche/Source/AvalancheSequence/Private/Transition/AvaTransitionSequenceTaskBase.cpp b/Engine/Plugins/VirtualProduction/Avalanche/Source/AvalancheSequence/Private/Transition/AvaTransitionSequenceTaskBase.cpp
--- a/Engine/Plugins/VirtualProduction/Avalanche/Source/AvalancheSequence/Private/Transition/AvaTransitionSequenceTaskBase.cpp
+++ b/Engine/Plugins/VirtualProduction/Avalanche/Source/AvalancheSequence/Private/Transition/AvaTransitionSequenceTaskBase.cpp
@@ -152,23 +152,6 @@ void FAvaTransitionSequenceTaskBase::ExitState(FStateTreeExecutionContext& InCon
 	}
 }
 
-void FAvaTransitionSequenceTaskBase::PostLoad(FStateTreeDataView InInstanceDataView)
-{
-	Super::PostLoad(InInstanceDataView);
-
-	PRAGMA_DISABLE_DEPRECATION_WARNINGS
-	if (QueryType_DEPRECATED != EAvaTransitionSequenceQueryType::None)
-	{
-		if (FInstanceDataType* InstanceData = UE::AvaTransition::TryGetInstanceData(*this, InInstanceDataView))
-		{
-			InstanceData->QueryType    = QueryType_DEPRECATED;
-			InstanceData->SequenceName = SequenceName_DEPRECATED;
-			InstanceData->SequenceTag  = SequenceTag_DEPRECATED;
-		}
-	}
-	PRAGMA_ENABLE_DEPRECATION_WARNINGS
-}
-
 bool FAvaTransitionSequenceTaskBase::Link(FStateTreeLinker& InLinker)
 {
 	Super::Link(InLinker);
diff --git a/Engine/Plugins/VirtualProduction/Avalanche/Source/AvalancheSequence/Public/AvaSequencePlayer.h b/Engine/Plugins/VirtualProduction/Avalanche/Source/AvalancheSequence/Public/AvaSequencePlayer.h
--- a/Engine/Plugins/VirtualProduction/Avalanche/Source/AvalancheSequence/Public/AvaSequencePlayer.h
+++ b/Engine/Plugins/VirtualProduction/Avalanche/Source/AvalancheSequence/Public/AvaSequencePlayer.h
@@ -44,9 +44,6 @@ public:
 
 	void ContinueSequence();
 
-	UE_DEPRECATED(5.5, "Use IAvaSequencePlaybackObject::PreviewFrame instead")
-	void PreviewFrame();
-
 	/** Jump to the given Frame, in Tick Resolution space */
 	void JumpTo(FFrameTime InJumpToFrame, bool bInEvaluate);
 
diff --git a/Engine/Plugins/VirtualProduction/Avalanche/Source/AvalancheSequence/Public/Transition/AvaTransitionInitializeSequence.h b/Engine/Plugins/VirtualProduction/Avalanche/Source/AvalancheSequence/Public/Transition/AvaTransitionInitializeSequence.h
--- a/Engine/Plugins/VirtualProduction/Avalanche/Source/AvalancheSequence/Public/Transition/AvaTransitionInitializeSequence.h
+++ b/Engine/Plugins/VirtualProduction/Avalanche/Source/AvalancheSequence/Public/Transition/AvaTransitionInitializeSequence.h
@@ -25,30 +25,12 @@ struct AVALANCHESEQUENCE_API FAvaTransitionInitializeSequence : public FAvaTrans
 
 	using FInstanceDataType = FAvaTransitionInitSequenceTaskInstanceData;
 
-	PRAGMA_DISABLE_DEPRECATION_WARNINGS
-	FAvaTransitionInitializeSequence() = default;
-	virtual ~FAvaTransitionInitializeSequence() override = default;
-	FAvaTransitionInitializeSequence(const FAvaTransitionInitializeSequence&) = default;
-	FAvaTransitionInitializeSequence(FAvaTransitionInitializeSequence&&) = default;
-	FAvaTransitionInitializeSequence& operator=(const FAvaTransitionInitializeSequence&) = default;
-	FAvaTransitionInitializeSequence& operator=(FAvaTransitionInitializeSequence&&) = default;
-	PRAGMA_ENABLE_DEPRECATION_WARNINGS
-
 	//~ Begin FStateTreeNodeBase
 	virtual const UStruct* GetInstanceDataType() const override { return FInstanceDataType::StaticStruct(); }
-	virtual void PostLoad(FStateTreeDataView InInstanceDataView) override;
 	//~ End FStateTreeNodeBase
 
 	//~ Begin FAvaTransitionSequenceTaskBase
 	virtual EAvaTransitionSequenceWaitType GetWaitType(FStateTreeExecutionContext& InContext) const override;
 	virtual TArray<UAvaSequencePlayer*> ExecuteSequenceTask(FStateTreeExecutionContext& InContext) const override;
 	//~ End FAvaTransitionSequenceTaskBase
-
-	UE_DEPRECATED(5.5, "InitializeTime has been moved to Instance Data")
-	UPROPERTY(meta = (DeprecatedProperty, DeprecationMessage = "Use the Instance Data InitializeTime instead"))
-	FAvaSequenceTime InitializeTime_DEPRECATED = FAvaSequenceTime(0.0);
-
-	UE_DEPRECATED(5.5, "PlayMode has been moved to Instance Data")
-	UPROPERTY(meta = (DeprecatedProperty, DeprecationMessage = "Use the Instance Data PlayMode instead"))
-	EAvaSequencePlayMode PlayMode_DEPRECATED = EAvaSequencePlayMode::Forward;
 };
diff --git a/Engine/Plugins/VirtualProduction/Avalanche/Source/AvalancheSequence/Public/Transition/AvaTransitionPlaySequenceTask.h b/Engine/Plugins/VirtualProduction/Avalanche/Source/AvalancheSequence/Public/Transition/AvaTransitionPlaySequenceTask.h
--- a/Engine/Plugins/VirtualProduction/Avalanche/Source/AvalancheSequence/Public/Transition/AvaTransitionPlaySequenceTask.h
+++ b/Engine/Plugins/VirtualProduction/Avalanche/Source/AvalancheSequence/Public/Transition/AvaTransitionPlaySequenceTask.h
@@ -23,18 +23,8 @@ struct AVALANCHESEQUENCE_API FAvaTransitionPlaySequenceTask : public FAvaTransit
 
 	using FInstanceDataType = FAvaTransitionPlaySequenceTaskInstanceData;
 
-	PRAGMA_DISABLE_DEPRECATION_WARNINGS
-	FAvaTransitionPlaySequenceTask() = default;
-	virtual ~FAvaTransitionPlaySequenceTask() override = default;
-	FAvaTransitionPlaySequenceTask(const FAvaTransitionPlaySequenceTask&) = default;
-	FAvaTransitionPlaySequenceTask(FAvaTransitionPlaySequenceTask&&) = default;
-	FAvaTransitionPlaySequenceTask& operator=(const FAvaTransitionPlaySequenceTask&) = default;
-	FAvaTransitionPlaySequenceTask& operator=(FAvaTransitionPlaySequenceTask&&) = default;
-	PRAGMA_ENABLE_DEPRECATION_WARNINGS
-
 	//~ Begin FStateTreeNodeBase
 	virtual const UStruct* GetInstanceDataType() const override { return FInstanceDataType::StaticStruct(); }
-	virtual void PostLoad(FStateTreeDataView InInstanceDataView) override;
 #if WITH_EDITOR
 	virtual FText GetDescription(const FGuid& InId, FStateTreeDataView InInstanceDataView, const IStateTreeBindingLookup& InBindingLookup, EStateTreeNodeFormatting InFormatting) const override;
 #endif
@@ -43,8 +33,4 @@ struct AVALANCHESEQUENCE_API FAvaTransitionPlaySequenceTask : public FAvaTransit
 	//~ Begin FAvaTransitionSequenceTaskBase
 	virtual TArray<UAvaSequencePlayer*> ExecuteSequenceTask(FStateTreeExecutionContext& InContext) const override;
 	//~ End FAvaTransitionSequenceTaskBase
-
-	UE_DEPRECATED(5.5, "PlaySettings has been moved to Instance Data")
-	UPROPERTY(meta = (DeprecatedProperty, DeprecationMessage = "Use the Instance Data PlaySettings instead"))
-	FAvaSequencePlayParams PlaySettings_DEPRECATED;
 };
diff --git a/Engine/Plugins/VirtualProduction/Avalanche/Source/AvalancheSequence/Public/Transition/AvaTransitionSequenceTask.h b/Engine/Plugins/VirtualProduction/Avalanche/Source/AvalancheSequence/Public/Transition/AvaTransitionSequenceTask.h
--- a/Engine/Plugins/VirtualProduction/Avalanche/Source/AvalancheSequence/Public/Transition/AvaTransitionSequenceTask.h
+++ b/Engine/Plugins/VirtualProduction/Avalanche/Source/AvalancheSequence/Public/Transition/AvaTransitionSequenceTask.h
@@ -23,25 +23,11 @@ struct AVALANCHESEQUENCE_API FAvaTransitionSequenceTask : public FAvaTransitionS
 
 	using FInstanceDataType = FAvaTransitionSequenceTaskInstanceData;
 
-	PRAGMA_DISABLE_DEPRECATION_WARNINGS
-	FAvaTransitionSequenceTask() = default;
-	virtual ~FAvaTransitionSequenceTask() override = default;
-	FAvaTransitionSequenceTask(const FAvaTransitionSequenceTask&) = default;
-	FAvaTransitionSequenceTask(FAvaTransitionSequenceTask&&) = default;
-	FAvaTransitionSequenceTask& operator=(const FAvaTransitionSequenceTask&) = default;
-	FAvaTransitionSequenceTask& operator=(FAvaTransitionSequenceTask&&) = default;
-	PRAGMA_ENABLE_DEPRECATION_WARNINGS
-
 	//~ Begin FStateTreeNodeBase
 	virtual const UStruct* GetInstanceDataType() const override { return FInstanceDataType::StaticStruct(); }
-	virtual void PostLoad(FStateTreeDataView InInstanceDataView) override;
 	//~ End FStateTreeNodeBase
 
 	//~ Begin FAvaTransitionSequenceTaskBase
 	virtual EAvaTransitionSequenceWaitType GetWaitType(FStateTreeExecutionContext& InContext) const override;
 	//~ End FAvaTransitionSequenceTaskBase
-
-	UE_DEPRECATED(5.5, "WaitType has been moved to Instance Data")
-	UPROPERTY(meta = (DeprecatedProperty, DeprecationMessage = "Use the Instance Data WaitType instead"))
-	EAvaTransitionSequenceWaitType WaitType_DEPRECATED = EAvaTransitionSequenceWaitType::WaitUntilStop;
 };
diff --git a/Engine/Plugins/VirtualProduction/Avalanche/Source/AvalancheSequence/Public/Transition/AvaTransitionSequenceTaskBase.h b/Engine/Plugins/VirtualProduction/Avalanche/Source/AvalancheSequence/Public/Transition/AvaTransitionSequenceTaskBase.h
--- a/Engine/Plugins/VirtualProduction/Avalanche/Source/AvalancheSequence/Public/Transition/AvaTransitionSequenceTaskBase.h
+++ b/Engine/Plugins/VirtualProduction/Avalanche/Source/AvalancheSequence/Public/Transition/AvaTransitionSequenceTaskBase.h
@@ -47,15 +47,6 @@ struct AVALANCHESEQUENCE_API FAvaTransitionSequenceTaskBase : public FAvaTransit
 {
 	GENERATED_BODY()
 
-	PRAGMA_DISABLE_DEPRECATION_WARNINGS
-	FAvaTransitionSequenceTaskBase() = default;
-	virtual ~FAvaTransitionSequenceTaskBase() override = default;
-	FAvaTransitionSequenceTaskBase(const FAvaTransitionSequenceTaskBase&) = default;
-	FAvaTransitionSequenceTaskBase(FAvaTransitionSequenceTaskBase&&) = default;
-	FAvaTransitionSequenceTaskBase& operator=(const FAvaTransitionSequenceTaskBase&) = default;
-	FAvaTransitionSequenceTaskBase& operator=(FAvaTransitionSequenceTaskBase&&) = default;
-	PRAGMA_ENABLE_DEPRECATION_WARNINGS
-
 	using FInstanceDataType = FAvaTransitionSequenceTaskBaseInstanceData;
 
 	/**
@@ -103,21 +94,8 @@ struct AVALANCHESEQUENCE_API FAvaTransitionSequenceTaskBase : public FAvaTransit
 
 	//~ Begin FStateTreeNodeBase
 	virtual const UStruct* GetInstanceDataType() const override { return FInstanceDataType::StaticStruct(); }
-	virtual void PostLoad(FStateTreeDataView InInstanceDataView) override;
 	virtual bool Link(FStateTreeLinker& InLinker) override;
 	//~ End FStateTreeNodeBase
 
-	UE_DEPRECATED(5.5, "QueryType has been moved to Instance Data")
-    UPROPERTY(meta = (DeprecatedProperty, DeprecationMessage = "Use the Instance Data QueryType instead"))
-	EAvaTransitionSequenceQueryType QueryType_DEPRECATED = EAvaTransitionSequenceQueryType::None;
-
-	UE_DEPRECATED(5.5, "SequenceName has been moved to Instance Data")
-	UPROPERTY(meta = (DeprecatedProperty, DeprecationMessage = "Use the Instance Data SequenceName instead"))
-	FName SequenceName_DEPRECATED;
-
-	UE_DEPRECATED(5.5, "SequenceTag has been moved to Instance Data")
-	UPROPERTY(meta = (DeprecatedProperty, DeprecationMessage = "Use the Instance Data SequenceTag instead"))
-	FAvaTagHandle SequenceTag_DEPRECATED;
-
 	TStateTreeExternalDataHandle<UAvaSequenceSubsystem> SequenceSubsystemHandle;
 };
diff --git a/Engine/Plugins/VirtualProduction/Avalanche/Source/AvalancheTransition/Private/Conditions/AvaTransitionLayerCondition.cpp b/Engine/Plugins/VirtualProduction/Avalanche/Source/AvalancheTransition/Private/Conditions/AvaTransitionLayerCondition.cpp
--- a/Engine/Plugins/VirtualProduction/Avalanche/Source/AvalancheTransition/Private/Conditions/AvaTransitionLayerCondition.cpp
+++ b/Engine/Plugins/VirtualProduction/Avalanche/Source/AvalancheTransition/Private/Conditions/AvaTransitionLayerCondition.cpp
@@ -24,22 +24,6 @@ FText FAvaTransitionLayerCondition::GetDescription(const FGuid& InId, FStateTree
 }
 #endif
 
-void FAvaTransitionLayerCondition::PostLoad(FStateTreeDataView InInstanceDataView)
-{
-	Super::PostLoad(InInstanceDataView);
-
-	PRAGMA_DISABLE_DEPRECATION_WARNINGS
-	if (LayerType_DEPRECATED != EAvaTransitionLayerCompareType::None)
-	{
-		if (FInstanceDataType* InstanceData = UE::AvaTransition::TryGetInstanceData(*this, InInstanceDataView))
-		{
-			InstanceData->LayerType     = LayerType_DEPRECATED;
-			InstanceData->SpecificLayer = SpecificLayer_DEPRECATED;
-		}
-	}
-	PRAGMA_ENABLE_DEPRECATION_WARNINGS
-}
-
 TArray<const FAvaTransitionBehaviorInstance*> FAvaTransitionLayerCondition::QueryBehaviorInstances(const FStateTreeExecutionContext& InContext) const
 {
 	const FInstanceDataType& InstanceData = InContext.GetInstanceData(*this);
diff --git a/Engine/Plugins/VirtualProduction/Avalanche/Source/AvalancheTransition/Private/Conditions/AvaTransitionSceneMatchCondition.cpp b/Engine/Plugins/VirtualProduction/Avalanche/Source/AvalancheTransition/Private/Conditions/AvaTransitionSceneMatchCondition.cpp
--- a/Engine/Plugins/VirtualProduction/Avalanche/Source/AvalancheTransition/Private/Conditions/AvaTransitionSceneMatchCondition.cpp
+++ b/Engine/Plugins/VirtualProduction/Avalanche/Source/AvalancheTransition/Private/Conditions/AvaTransitionSceneMatchCondition.cpp
@@ -24,21 +24,6 @@ FText FAvaTransitionSceneMatchCondition::GetDescription(const FGuid& InId, FStat
 }
 #endif
 
-void FAvaTransitionSceneMatchCondition::PostLoad(FStateTreeDataView InInstanceDataView)
-{
-	Super::PostLoad(InInstanceDataView);
-
-	PRAGMA_DISABLE_DEPRECATION_WARNINGS
-	if (LayerType_DEPRECATED != EAvaTransitionLayerCompareType::None)
-	{
-		if (FInstanceDataType* InstanceData = UE::AvaTransition::TryGetInstanceData(*this, InInstanceDataView))
-		{
-			InstanceData->SceneComparisonType = SceneComparisonType_DEPRECATED;
-		}
-	}
-	PRAGMA_ENABLE_DEPRECATION_WARNINGS
-}
-
 bool FAvaTransitionSceneMatchCondition::TestCondition(FStateTreeExecutionContext& InContext) const
 {
 	const FAvaTransitionContext& TransitionContext = InContext.GetExternalData(TransitionContextHandle);
diff --git a/Engine/Plugins/VirtualProduction/Avalanche/Source/AvalancheTransition/Private/Conditions/AvaTransitionStateMatchCondition.cpp b/Engine/Plugins/VirtualProduction/Avalanche/Source/AvalancheTransition/Private/Conditions/AvaTransitionStateMatchCondition.cpp
--- a/Engine/Plugins/VirtualProduction/Avalanche/Source/AvalancheTransition/Private/Conditions/AvaTransitionStateMatchCondition.cpp
+++ b/Engine/Plugins/VirtualProduction/Avalanche/Source/AvalancheTransition/Private/Conditions/AvaTransitionStateMatchCondition.cpp
@@ -22,21 +22,6 @@ FText FAvaTransitionStateMatchCondition::GetDescription(const FGuid& InId, FStat
 }
 #endif
 
-void FAvaTransitionStateMatchCondition::PostLoad(FStateTreeDataView InInstanceDataView)
-{
-	Super::PostLoad(InInstanceDataView);
-
-	PRAGMA_DISABLE_DEPRECATION_WARNINGS
-	if (TransitionState_DEPRECATED != EAvaTransitionRunState::Unknown)
-	{
-		if (FInstanceDataType* InstanceData = UE::AvaTransition::TryGetInstanceData(*this, InInstanceDataView))
-		{
-			InstanceData->TransitionState = TransitionState_DEPRECATED;
-		}
-	}
-	PRAGMA_ENABLE_DEPRECATION_WARNINGS
-}
-
 bool FAvaTransitionStateMatchCondition::TestCondition(FStateTreeExecutionContext& InContext) const
 {
 	TArray<const FAvaTransitionBehaviorInstance*> BehaviorInstances = QueryBehaviorInstances(InContext);
diff --git a/Engine/Plugins/VirtualProduction/Avalanche/Source/AvalancheTransition/Private/Conditions/AvaTransitionTypeMatchCondition.cpp b/Engine/Plugins/VirtualProduction/Avalanche/Source/AvalancheTransition/Private/Conditions/AvaTransitionTypeMatchCondition.cpp
--- a/Engine/Plugins/VirtualProduction/Avalanche/Source/AvalancheTransition/Private/Conditions/AvaTransitionTypeMatchCondition.cpp
+++ b/Engine/Plugins/VirtualProduction/Avalanche/Source/AvalancheTransition/Private/Conditions/AvaTransitionTypeMatchCondition.cpp
@@ -20,21 +20,6 @@ FText FAvaTransitionTypeMatchCondition::GetDescription(const FGuid& InId, FState
 }
 #endif
 
-void FAvaTransitionTypeMatchCondition::PostLoad(FStateTreeDataView InInstanceDataView)
-{
-	Super::PostLoad(InInstanceDataView);
-
-	PRAGMA_DISABLE_DEPRECATION_WARNINGS
-	if (TransitionType_DEPRECATED != EAvaTransitionType::None)
-	{
-		if (FInstanceDataType* InstanceData = UE::AvaTransition::TryGetInstanceData(*this, InInstanceDataView))
-		{
-			InstanceData->TransitionType = TransitionType_DEPRECATED;
-		}
-	}
-	PRAGMA_ENABLE_DEPRECATION_WARNINGS
-}
-
 bool FAvaTransitionTypeMatchCondition::TestCondition(FStateTreeExecutionContext& InContext) const
 {
 	const FAvaTransitionContext& TransitionContext = InContext.GetExternalData(TransitionContextHandle);
diff --git a/Engine/Plugins/VirtualProduction/Avalanche/Source/AvalancheTransition/Private/Tasks/AvaTransitionDelayTask.cpp b/Engine/Plugins/VirtualProduction/Avalanche/Source/AvalancheTransition/Private/Tasks/AvaTransitionDelayTask.cpp
--- a/Engine/Plugins/VirtualProduction/Avalanche/Source/AvalancheTransition/Private/Tasks/AvaTransitionDelayTask.cpp
+++ b/Engine/Plugins/VirtualProduction/Avalanche/Source/AvalancheTransition/Private/Tasks/AvaTransitionDelayTask.cpp
@@ -32,21 +32,6 @@ bool FAvaTransitionDelayTask::Link(FStateTreeLinker& InLinker)
 	return true;
 }
 
-void FAvaTransitionDelayTask::PostLoad(FStateTreeDataView InInstanceDataView)
-{
-	Super::PostLoad(InInstanceDataView);
-
-	PRAGMA_DISABLE_DEPRECATION_WARNINGS
-	if (Duration_DEPRECATED >= 0.f)
-	{
-		if (FInstanceDataType* InstanceData = UE::AvaTransition::TryGetInstanceData(*this, InInstanceDataView))
-		{
-			InstanceData->Duration = Duration_DEPRECATED;
-		}
-	}
-	PRAGMA_ENABLE_DEPRECATION_WARNINGS
-}
-
 EStateTreeRunStatus FAvaTransitionDelayTask::EnterState(FStateTreeExecutionContext& InContext, const FStateTreeTransitionResult& InTransition) const
 {
 	FInstanceDataType& InstanceData = InContext.GetInstanceData(*this);
diff --git a/Engine/Plugins/VirtualProduction/Avalanche/Source/AvalancheTransition/Private/Tasks/AvaTransitionLayerTask.cpp b/Engine/Plugins/VirtualProduction/Avalanche/Source/AvalancheTransition/Private/Tasks/AvaTransitionLayerTask.cpp
--- a/Engine/Plugins/VirtualProduction/Avalanche/Source/AvalancheTransition/Private/Tasks/AvaTransitionLayerTask.cpp
+++ b/Engine/Plugins/VirtualProduction/Avalanche/Source/AvalancheTransition/Private/Tasks/AvaTransitionLayerTask.cpp
@@ -24,22 +24,6 @@ FText FAvaTransitionLayerTask::GetDescription(const FGuid& InId, FStateTreeDataV
 }
 #endif
 
-void FAvaTransitionLayerTask::PostLoad(FStateTreeDataView InInstanceDataView)
-{
-	Super::PostLoad(InInstanceDataView);
-
-	PRAGMA_DISABLE_DEPRECATION_WARNINGS
-	if (LayerType_DEPRECATED != EAvaTransitionLayerCompareType::None)
-	{
-		if (FInstanceDataType* InstanceData = UE::AvaTransition::TryGetInstanceData(*this, InInstanceDataView))
-		{
-			InstanceData->LayerType     = LayerType_DEPRECATED;
-			InstanceData->SpecificLayer = SpecificLayer_DEPRECATED;
-		}
-	}
-	PRAGMA_ENABLE_DEPRECATION_WARNINGS
-}
-
 TArray<const FAvaTransitionBehaviorInstance*> FAvaTransitionLayerTask::QueryBehaviorInstances(const FStateTreeExecutionContext& InContext) const
 {
 	const FInstanceDataType& InstanceData = InContext.GetInstanceData(*this);
diff --git a/Engine/Plugins/VirtualProduction/Avalanche/Source/AvalancheTransition/Public/Conditions/AvaTransitionLayerCondition.h b/Engine/Plugins/VirtualProduction/Avalanche/Source/AvalancheTransition/Public/Conditions/AvaTransitionLayerCondition.h
--- a/Engine/Plugins/VirtualProduction/Avalanche/Source/AvalancheTransition/Public/Conditions/AvaTransitionLayerCondition.h
+++ b/Engine/Plugins/VirtualProduction/Avalanche/Source/AvalancheTransition/Public/Conditions/AvaTransitionLayerCondition.h
@@ -28,31 +28,13 @@ struct AVALANCHETRANSITION_API FAvaTransitionLayerCondition : public FAvaTransit
 
 	using FInstanceDataType = FAvaTransitionLayerConditionInstanceData;
 
-	PRAGMA_DISABLE_DEPRECATION_WARNINGS
-	FAvaTransitionLayerCondition() = default;
-	virtual ~FAvaTransitionLayerCondition() override = default;
-	FAvaTransitionLayerCondition(const FAvaTransitionLayerCondition&) = default;
-	FAvaTransitionLayerCondition(FAvaTransitionLayerCondition&&) = default;
-	FAvaTransitionLayerCondition& operator=(const FAvaTransitionLayerCondition&) = default;
-	FAvaTransitionLayerCondition& operator=(FAvaTransitionLayerCondition&&) = default;
-	PRAGMA_ENABLE_DEPRECATION_WARNINGS
-
 	//~ Begin FStateTreeNodeBase
 #if WITH_EDITOR
 	virtual FText GetDescription(const FGuid& InId, FStateTreeDataView InInstanceDataView, const IStateTreeBindingLookup& InBindingLookup, EStateTreeNodeFormatting InFormatting) const override;
 #endif
 	virtual const UStruct* GetInstanceDataType() const override { return FInstanceDataType::StaticStruct(); }
-	virtual void PostLoad(FStateTreeDataView InInstanceDataView) override;
 	//~ End FStateTreeNodeBase
 
 	/** Gets all the Behavior Instances that match the Layer Query. Always excludes the Instance belonging to this Transition */
 	TArray<const FAvaTransitionBehaviorInstance*> QueryBehaviorInstances(const FStateTreeExecutionContext& InContext) const;
-
-	UE_DEPRECATED(5.5, "LayerType has been moved to Instance Data")
-	UPROPERTY(meta = (DeprecatedProperty, DeprecationMessage = "Use the Instance Data LayerType instead"))
-	EAvaTransitionLayerCompareType LayerType_DEPRECATED = EAvaTransitionLayerCompareType::None;
-
-	UE_DEPRECATED(5.5, "SpecificLayer has been moved to Instance Data")
-	UPROPERTY(meta = (DeprecatedProperty, DeprecationMessage = "Use the Instance Data SpecificLayer instead"))
-	FAvaTagHandle SpecificLayer_DEPRECATED;
 };
diff --git a/Engine/Plugins/VirtualProduction/Avalanche/Source/AvalancheTransition/Public/Conditions/AvaTransitionSceneMatchCondition.h b/Engine/Plugins/VirtualProduction/Avalanche/Source/AvalancheTransition/Public/Conditions/AvaTransitionSceneMatchCondition.h
--- a/Engine/Plugins/VirtualProduction/Avalanche/Source/AvalancheTransition/Public/Conditions/AvaTransitionSceneMatchCondition.h
+++ b/Engine/Plugins/VirtualProduction/Avalanche/Source/AvalancheTransition/Public/Conditions/AvaTransitionSceneMatchCondition.h
@@ -22,28 +22,14 @@ struct AVALANCHETRANSITION_API FAvaTransitionSceneMatchCondition : public FAvaTr
 
 	using FInstanceDataType = FAvaTransitionSceneMatchConditionInstanceData;
 
-	PRAGMA_DISABLE_DEPRECATION_WARNINGS
-	FAvaTransitionSceneMatchCondition() = default;
-	virtual ~FAvaTransitionSceneMatchCondition() override = default;
-	FAvaTransitionSceneMatchCondition(const FAvaTransitionSceneMatchCondition&) = default;
-	FAvaTransitionSceneMatchCondition(FAvaTransitionSceneMatchCondition&&) = default;
-	FAvaTransitionSceneMatchCondition& operator=(const FAvaTransitionSceneMatchCondition&) = default;
-	FAvaTransitionSceneMatchCondition& operator=(FAvaTransitionSceneMatchCondition&&) = default;
-	PRAGMA_ENABLE_DEPRECATION_WARNINGS
-
 	//~ Begin FStateTreeNodeBase
 #if WITH_EDITOR
 	virtual FText GetDescription(const FGuid& InId, FStateTreeDataView InInstanceDataView, const IStateTreeBindingLookup& InBindingLookup, EStateTreeNodeFormatting InFormatting) const override;
 #endif
 	virtual const UStruct* GetInstanceDataType() const override { return FInstanceDataType::StaticStruct(); }
-	virtual void PostLoad(FStateTreeDataView InInstanceDataView) override;
 	//~ End FStateTreeNodeBase
 
 	//~ Begin FStateTreeConditionBase
 	virtual bool TestCondition(FStateTreeExecutionContext& InContext) const override;
 	//~ End FStateTreeConditionBase
-
-	UE_DEPRECATED(5.5, "SceneComparisonType has been moved to Instance Data")
-	UPROPERTY(meta = (DeprecatedProperty, DeprecationMessage = "Use the Instance Data SceneComparisonType instead"))
-	EAvaTransitionComparisonResult SceneComparisonType_DEPRECATED = EAvaTransitionComparisonResult::None;
 };
diff --git a/Engine/Plugins/VirtualProduction/Avalanche/Source/AvalancheTransition/Public/Conditions/AvaTransitionStateMatchCondition.h b/Engine/Plugins/VirtualProduction/Avalanche/Source/AvalancheTransition/Public/Conditions/AvaTransitionStateMatchCondition.h
--- a/Engine/Plugins/VirtualProduction/Avalanche/Source/AvalancheTransition/Public/Conditions/AvaTransitionStateMatchCondition.h
+++ b/Engine/Plugins/VirtualProduction/Avalanche/Source/AvalancheTransition/Public/Conditions/AvaTransitionStateMatchCondition.h
@@ -22,28 +22,14 @@ struct AVALANCHETRANSITION_API FAvaTransitionStateMatchCondition : public FAvaTr
 
 	using FInstanceDataType = FAvaTransitionStateMatchConditionInstanceData;
 
-	PRAGMA_DISABLE_DEPRECATION_WARNINGS
-	FAvaTransitionStateMatchCondition() = default;
-	virtual ~FAvaTransitionStateMatchCondition() override = default;
-	FAvaTransitionStateMatchCondition(const FAvaTransitionStateMatchCondition&) = default;
-	FAvaTransitionStateMatchCondition(FAvaTransitionStateMatchCondition&&) = default;
-	FAvaTransitionStateMatchCondition& operator=(const FAvaTransitionStateMatchCondition&) = default;
-	FAvaTransitionStateMatchCondition& operator=(FAvaTransitionStateMatchCondition&&) = default;
-	PRAGMA_ENABLE_DEPRECATION_WARNINGS
-
 	//~ Begin FStateTreeNodeBase
 #if WITH_EDITOR
 	virtual FText GetDescription(const FGuid& InId, FStateTreeDataView InInstanceDataView, const IStateTreeBindingLookup& InBindingLookup, EStateTreeNodeFormatting InFormatting) const override;
 #endif
 	virtual const UStruct* GetInstanceDataType() const override { return FAvaTransitionStateMatchConditionInstanceData::StaticStruct(); }
-	virtual void PostLoad(FStateTreeDataView InInstanceDataView) override;
 	//~ End FStateTreeNodeBase
 
 	//~ Begin FStateTreeConditionBase
 	virtual bool TestCondition(FStateTreeExecutionContext& InContext) const override;
 	//~ End FStateTreeConditionBase
-
-	UE_DEPRECATED(5.5, "TransitionState has been moved to Instance Data")
-	UPROPERTY(meta = (DeprecatedProperty, DeprecationMessage = "Use the Instance Data TransitionState instead"))
-	EAvaTransitionRunState TransitionState_DEPRECATED = EAvaTransitionRunState::Unknown;
 };
diff --git a/Engine/Plugins/VirtualProduction/Avalanche/Source/AvalancheTransition/Public/Conditions/AvaTransitionTypeMatchCondition.h b/Engine/Plugins/VirtualProduction/Avalanche/Source/AvalancheTransition/Public/Conditions/AvaTransitionTypeMatchCondition.h
--- a/Engine/Plugins/VirtualProduction/Avalanche/Source/AvalancheTransition/Public/Conditions/AvaTransitionTypeMatchCondition.h
+++ b/Engine/Plugins/VirtualProduction/Avalanche/Source/AvalancheTransition/Public/Conditions/AvaTransitionTypeMatchCondition.h
@@ -22,28 +22,14 @@ struct AVALANCHETRANSITION_API FAvaTransitionTypeMatchCondition : public FAvaTra
 
 	using FInstanceDataType = FAvaTransitionTypeMatchConditionInstanceData;
 
-	PRAGMA_DISABLE_DEPRECATION_WARNINGS
-	FAvaTransitionTypeMatchCondition() = default;
-	virtual ~FAvaTransitionTypeMatchCondition() override = default;
-	FAvaTransitionTypeMatchCondition(const FAvaTransitionTypeMatchCondition&) = default;
-	FAvaTransitionTypeMatchCondition(FAvaTransitionTypeMatchCondition&&) = default;
-	FAvaTransitionTypeMatchCondition& operator=(const FAvaTransitionTypeMatchCondition&) = default;
-	FAvaTransitionTypeMatchCondition& operator=(FAvaTransitionTypeMatchCondition&&) = default;
-	PRAGMA_ENABLE_DEPRECATION_WARNINGS
-
 	//~ Begin FStateTreeNodeBase
 #if WITH_EDITOR
 	virtual FText GetDescription(const FGuid& InId, FStateTreeDataView InInstanceDataView, const IStateTreeBindingLookup& InBindingLookup, EStateTreeNodeFormatting InFormatting) const override;
 #endif
 	virtual const UStruct* GetInstanceDataType() const override { return FInstanceDataType::StaticStruct(); }
-	virtual void PostLoad(FStateTreeDataView InInstanceDataView) override;
 	//~ End FStateTreeNodeBase
 
 	//~ Begin FStateTreeConditionBase
 	virtual bool TestCondition(FStateTreeExecutionContext& InContext) const override;
 	//~ End FStateTreeConditionBase
-
-	UE_DEPRECATED(5.5, "TransitionType has been moved to Instance Data")
-	UPROPERTY(meta = (DeprecatedProperty, DeprecationMessage = "Use the Instance Data TransitionType instead"))
-	EAvaTransitionType TransitionType_DEPRECATED = EAvaTransitionType::None;
 };
diff --git a/Engine/Plugins/VirtualProduction/Avalanche/Source/AvalancheTransition/Public/Tasks/AvaTransitionDelayTask.h b/Engine/Plugins/VirtualProduction/Avalanche/Source/AvalancheTransition/Public/Tasks/AvaTransitionDelayTask.h
--- a/Engine/Plugins/VirtualProduction/Avalanche/Source/AvalancheTransition/Public/Tasks/AvaTransitionDelayTask.h
+++ b/Engine/Plugins/VirtualProduction/Avalanche/Source/AvalancheTransition/Public/Tasks/AvaTransitionDelayTask.h
@@ -36,22 +36,12 @@ struct AVALANCHETRANSITION_API FAvaTransitionDelayTask : public FAvaTransitionTa
 
 	using FInstanceDataType = FAvaTransitionDelayTaskInstanceData;
 
-	PRAGMA_DISABLE_DEPRECATION_WARNINGS
-	FAvaTransitionDelayTask() = default;
-	virtual ~FAvaTransitionDelayTask() override = default;
-	FAvaTransitionDelayTask(const FAvaTransitionDelayTask&) = default;
-	FAvaTransitionDelayTask(FAvaTransitionDelayTask&&) = default;
-	FAvaTransitionDelayTask& operator=(const FAvaTransitionDelayTask&) = default;
-	FAvaTransitionDelayTask& operator=(FAvaTransitionDelayTask&&) = default;
-	PRAGMA_ENABLE_DEPRECATION_WARNINGS
-
 	//~ Begin FStateTreeNodeBase
 #if WITH_EDITOR
 	virtual FText GetDescription(const FGuid& InId, FStateTreeDataView InInstanceDataView, const IStateTreeBindingLookup& InBindingLookup, EStateTreeNodeFormatting InFormatting) const override;
 #endif
 	virtual const UStruct* GetInstanceDataType() const override { return FInstanceDataType::StaticStruct(); }
 	virtual bool Link(FStateTreeLinker& InLinker) override;
-	virtual void PostLoad(FStateTreeDataView InInstanceDataView) override;
 	//~ End FStateTreeNodeBase
 
 	//~ Begin FStateTreeTaskBase
@@ -63,9 +53,5 @@ struct AVALANCHETRANSITION_API FAvaTransitionDelayTask : public FAvaTransitionTa
 
 	bool ShouldHideLevel(const FStateTreeExecutionContext& InContext, const FInstanceDataType& InInstanceData) const;
 
-	UE_DEPRECATED(5.5, "Duration has been moved to Instance Data")
-	UPROPERTY(meta = (DeprecatedProperty, DeprecationMessage = "Use the Instance Data Duration instead"))
-	float Duration_DEPRECATED = -1.f;
-
 	TStateTreeExternalDataHandle<UAvaTransitionRenderingSubsystem> RenderingSubsystemHandle;
 };
diff --git a/Engine/Plugins/VirtualProduction/Avalanche/Source/AvalancheTransition/Public/Tasks/AvaTransitionLayerTask.h b/Engine/Plugins/VirtualProduction/Avalanche/Source/AvalancheTransition/Public/Tasks/AvaTransitionLayerTask.h
--- a/Engine/Plugins/VirtualProduction/Avalanche/Source/AvalancheTransition/Public/Tasks/AvaTransitionLayerTask.h
+++ b/Engine/Plugins/VirtualProduction/Avalanche/Source/AvalancheTransition/Public/Tasks/AvaTransitionLayerTask.h
@@ -26,31 +26,13 @@ struct AVALANCHETRANSITION_API FAvaTransitionLayerTask : public FAvaTransitionTa
 
 	using FInstanceDataType = FAvaTransitionLayerTaskInstanceData;
 
-	PRAGMA_DISABLE_DEPRECATION_WARNINGS
-	FAvaTransitionLayerTask() = default;
-	virtual ~FAvaTransitionLayerTask() override = default;
-	FAvaTransitionLayerTask(const FAvaTransitionLayerTask&) = default;
-	FAvaTransitionLayerTask(FAvaTransitionLayerTask&&) = default;
-	FAvaTransitionLayerTask& operator=(const FAvaTransitionLayerTask&) = default;
-	FAvaTransitionLayerTask& operator=(FAvaTransitionLayerTask&&) = default;
-	PRAGMA_ENABLE_DEPRECATION_WARNINGS
-
 	//~ Begin FStateTreeNodeBase
 #if WITH_EDITOR
 	virtual FText GetDescription(const FGuid& InId, FStateTreeDataView InInstanceDataView, const IStateTreeBindingLookup& InBindingLookup, EStateTreeNodeFormatting InFormatting) const override;
 #endif
 	virtual const UStruct* GetInstanceDataType() const override { return FInstanceDataType::StaticStruct(); }
-	virtual void PostLoad(FStateTreeDataView InInstanceDataView) override;
 	//~ End FStateTreeNodeBase
 
 	/** Gets all the Behavior Instances that match the Layer Query. Always excludes the Instance belonging to this Transition */
 	TArray<const FAvaTransitionBehaviorInstance*> QueryBehaviorInstances(const FStateTreeExecutionContext& InContext) const;
-
-	UE_DEPRECATED(5.5, "LayerType has been moved to Instance Data")
-	UPROPERTY(meta = (DeprecatedProperty, DeprecationMessage = "Use the Instance Data LayerType instead"))
-	EAvaTransitionLayerCompareType LayerType_DEPRECATED = EAvaTransitionLayerCompareType::None;
-
-	UE_DEPRECATED(5.5, "SpecificLayer has been moved to Instance Data")
-	UPROPERTY(meta = (DeprecatedProperty, DeprecationMessage = "Use the Instance Data SpecificLayer instead"))
-	FAvaTagHandle SpecificLayer_DEPRECATED;
 };
