Motion Design: Workaround for screenshot test fails

#Synced-CL 49103131

diff --git a/Engine/Plugins/VirtualProduction/Avalanche/Source/AvalancheFunctionalTest/Private/ScreenshotFunctionalTestAvalanche.cpp b/Engine/Plugins/VirtualProduction/Avalanche/Source/AvalancheFunctionalTest/Private/ScreenshotFunctionalTestAvalanche.cpp
--- a/Engine/Plugins/VirtualProduction/Avalanche/Source/AvalancheFunctionalTest/Private/ScreenshotFunctionalTestAvalanche.cpp
+++ b/Engine/Plugins/VirtualProduction/Avalanche/Source/AvalancheFunctionalTest/Private/ScreenshotFunctionalTestAvalanche.cpp
@@ -45,10 +45,6 @@ void AScreenshotFunctionalTestAvalanche::RequestScreenshot()
 	// PIE uses FEditorViewportClient::ProcessScreenShots() which fires FScreenshotRequest::OnScreenshotCaptured() (static)
 	// instead of UGameViewportClient::OnScreenshotCaptured() (instance) that base class registers to
 	FScreenshotRequest::OnScreenshotCaptured().AddUObject(this, &AScreenshotFunctionalTestAvalanche::OnScreenShotCaptured);
-
-	// Request the actual screenshot (base class only registers callbacks)
-	FIntRect ViewRect = FIntRect(); // Empty ViewRect = entire render target
-	FScreenshotRequest::RequestScreenshot(TEXT(""), /*bInShowUI*/ false, /*bAddUniqueSuffix*/ true, /*bHdrScreenshot*/ false, ViewRect);
 }
 
 void AScreenshotFunctionalTestAvalanche::OnScreenShotCaptured(int32 InSizeX, int32 InSizeY, const TArray<FColor>& InImageData)
diff --git a/Engine/Plugins/VirtualProduction/Avalanche/Source/AvalancheLevelViewport/AvalancheLevelViewport.Build.cs b/Engine/Plugins/VirtualProduction/Avalanche/Source/AvalancheLevelViewport/AvalancheLevelViewport.Build.cs
--- a/Engine/Plugins/VirtualProduction/Avalanche/Source/AvalancheLevelViewport/AvalancheLevelViewport.Build.cs
+++ b/Engine/Plugins/VirtualProduction/Avalanche/Source/AvalancheLevelViewport/AvalancheLevelViewport.Build.cs
@@ -36,11 +36,13 @@ public class AvalancheLevelViewport : ModuleRules
 				"Avalanche",
 				"AvalancheComponentVisualizers",
 				"AvalancheEditorCore",
+				"AvalancheFunctionalTest",
 				"AvalancheInteractiveTools",
 				"AvalancheViewport",
 				"CinematicCamera",
 				"CustomDetailsView",
 				"EditorFramework",
+				"FunctionalTesting",
 				"InputCore",
 				"LevelEditor",
 				"Projects",
diff --git a/Engine/Plugins/VirtualProduction/Avalanche/Source/AvalancheLevelViewport/Private/ViewportClient/AvaLevelViewportClient.cpp b/Engine/Plugins/VirtualProduction/Avalanche/Source/AvalancheLevelViewport/Private/ViewportClient/AvaLevelViewportClient.cpp
--- a/Engine/Plugins/VirtualProduction/Avalanche/Source/AvalancheLevelViewport/Private/ViewportClient/AvaLevelViewportClient.cpp
+++ b/Engine/Plugins/VirtualProduction/Avalanche/Source/AvalancheLevelViewport/Private/ViewportClient/AvaLevelViewportClient.cpp
@@ -12,6 +12,7 @@
 #include "Editor/UnrealEdEngine.h"
 #include "EngineUtils.h"
 #include "Framework/Application/SlateApplication.h"
+#include "FunctionalTestBase.h"
 #include "GameFramework/Actor.h"
 #include "Interaction/AvaCameraZoomController.h"
 #include "Interaction/AvaChildTransformLockOperation.h"
@@ -21,12 +22,15 @@
 #include "LevelEditor.h"
 #include "Modules/ModuleManager.h"
 #include "MouseDeltaTracker.h"
+#include "FunctionalTestingModule.h"
 #include "SAvaLevelViewport.h"
 #include "SceneView.h"
+#include "ScreenshotFunctionalTestAvalanche.h"
 #include "Selection.h"
 #include "Selection/AvaSelectionProviderSubsystem.h"
 #include "SLevelViewport.h"
 #include "UnrealEdGlobals.h"
+#include "UObject/UObjectIterator.h"
 #include "Visualizers/IAvaViewportBoundingBoxVisualizer.h"
 
 #define LOCTEXT_NAMESPACE "AvaLevelViewportClient"
@@ -441,6 +445,28 @@ FMatrix FAvaLevelViewportClient::CalcViewRotationMatrixForControllingActorView(c
 	return CalcViewRotationMatrix(InViewRotation);
 }
 
+bool FAvaLevelViewportClient::ProcessScreenShots(FViewport* InViewport)
+{
+	bool bIsRunningAvaScreenshotTest = false;
+
+	for (AScreenshotFunctionalTestAvalanche* Test : TObjectRange<AScreenshotFunctionalTestAvalanche>())
+	{
+		if (!Test->IsTemplate() && Test->IsRunning())
+		{
+			return false;
+		}
+	}
+
+	/*
+	if (IFunctionalTestingModule::Get().IsRunning())
+	{
+		return false;
+	}
+	*/
+
+	return FEditorViewportClient::ProcessScreenShots(InViewport);
+}
+
 FSceneView* FAvaLevelViewportClient::CalcSceneView(FSceneViewFamily* ViewFamily, const int32 StereoViewIndex)
 {
 	UCameraComponent* CachedCameraComponent = nullptr;
diff --git a/Engine/Plugins/VirtualProduction/Avalanche/Source/AvalancheLevelViewport/Private/ViewportClient/AvaLevelViewportClient.h b/Engine/Plugins/VirtualProduction/Avalanche/Source/AvalancheLevelViewport/Private/ViewportClient/AvaLevelViewportClient.h
--- a/Engine/Plugins/VirtualProduction/Avalanche/Source/AvalancheLevelViewport/Private/ViewportClient/AvaLevelViewportClient.h
+++ b/Engine/Plugins/VirtualProduction/Avalanche/Source/AvalancheLevelViewport/Private/ViewportClient/AvaLevelViewportClient.h
@@ -73,6 +73,7 @@ public:
 	virtual void UpdateMouseDelta() override;
 	virtual EMouseCursor::Type GetCursor(FViewport* InViewport, int32 InX, int32 InY) override;
 	virtual FMatrix CalcViewRotationMatrixForControllingActorView(const FRotator& InViewRotation) const override;
+	virtual bool ProcessScreenShots(FViewport* InViewport) override;
 	//~ End FEditorViewportClient
 
 	//~ Begin IAvaViewportWorldCoordinateConverter
