Motion Design: updated transition tree debugger to use state tree delegates instead of tracing channels.
This has the benefit of not needing to enable trace to obtain debugging information. Additionally, the editor compile hash was unreliable as a way to compare trees, especially when new changes to state tree cause the compiled hash to be different from the loaded compiled hash (even if no changes were made). For this reason, a 'compare' function was added to the transition tree editor data that compares the root parameters guid. Root parameter guid is unique per state tree (and a new one gets generated for a duplicated tree).

#Synced-CL 49298335

diff --git a/Engine/Plugins/VirtualProduction/Avalanche/Source/AvalancheTransitionEditor/AvalancheTransitionEditor.Build.cs b/Engine/Plugins/VirtualProduction/Avalanche/Source/AvalancheTransitionEditor/AvalancheTransitionEditor.Build.cs
--- a/Engine/Plugins/VirtualProduction/Avalanche/Source/AvalancheTransitionEditor/AvalancheTransitionEditor.Build.cs
+++ b/Engine/Plugins/VirtualProduction/Avalanche/Source/AvalancheTransitionEditor/AvalancheTransitionEditor.Build.cs
@@ -42,23 +42,5 @@ public class AvalancheTransitionEditor : ModuleRules
                 "UnrealEd",
             }
         );
-
-        // Until we split use of TraceServices and TraceAnalysis in the debugger we enable only on platforms supporting both at the moment
-        if (Target.Platform.IsInGroup(UnrealPlatformGroup.Desktop) && (Target.Configuration != UnrealTargetConfiguration.Shipping || Target.bBuildEditor))
-        {
-            PrivateDefinitions.Add("WITH_STATETREE_DEBUGGER=1");
-            PrivateDependencyModuleNames.AddRange(
-                new string[]
-                {
-                    "TraceAnalysis",
-                    "TraceLog",
-                    "TraceServices",
-                }
-            );
-        }
-        else
-        {
-            PrivateDefinitions.Add("WITH_STATETREE_DEBUGGER=0");
-        }
     }
 }
diff --git a/Engine/Plugins/VirtualProduction/Avalanche/Source/AvalancheTransitionEditor/Private/Actions/AvaTransitionTreeActions.cpp b/Engine/Plugins/VirtualProduction/Avalanche/Source/AvalancheTransitionEditor/Private/Actions/AvaTransitionTreeActions.cpp
--- a/Engine/Plugins/VirtualProduction/Avalanche/Source/AvalancheTransitionEditor/Private/Actions/AvaTransitionTreeActions.cpp
+++ b/Engine/Plugins/VirtualProduction/Avalanche/Source/AvalancheTransitionEditor/Private/Actions/AvaTransitionTreeActions.cpp
@@ -12,7 +12,7 @@
 #include "ViewModels/AvaTransitionEditorViewModel.h"
 #include "ViewModels/AvaTransitionViewModelSharedData.h"
 
-#if WITH_STATETREE_DEBUGGER
+#if UE_AVA_WITH_TRANSITION_DEBUG
 #include "Debugger/AvaTransitionDebugger.h"
 #endif
 
@@ -46,7 +46,7 @@ void FAvaTransitionTreeActions::BindCommands(const TSharedRef<FUICommandList>& I
 		, FCanExecuteAction()
 		, FIsActionChecked::CreateStatic(&FAvaTransitionCompiler::HasSaveOnCompile, EStateTreeSaveOnCompile::SuccessOnly));
 
-#if WITH_STATETREE_DEBUGGER
+#if UE_AVA_WITH_TRANSITION_DEBUG
 	InCommandList->MapAction(Commands.ToggleDebug
 		, FExecuteAction::CreateSP(this, &FAvaTransitionTreeActions::ToggleDebugger)
 		, FCanExecuteAction()
@@ -140,7 +140,7 @@ void FAvaTransitionTreeActions::ImportTransitionTree(UAvaTransitionTree* InTempl
 	}
 }
 
-#if WITH_STATETREE_DEBUGGER
+#if UE_AVA_WITH_TRANSITION_DEBUG
 void FAvaTransitionTreeActions::ToggleDebugger()
 {
 	TSharedRef<FAvaTransitionDebugger> Debugger = Owner.GetSharedData()->GetDebugger();
@@ -159,6 +159,6 @@ bool FAvaTransitionTreeActions::IsDebuggerActive() const
 	TSharedRef<FAvaTransitionDebugger> Debugger = Owner.GetSharedData()->GetDebugger();
 	return Debugger->IsActive();
 }
-#endif // WITH_STATETREE_DEBUGGER
+#endif // UE_AVA_WITH_TRANSITION_DEBUG
 
 #undef LOCTEXT_NAMESPACE
diff --git a/Engine/Plugins/VirtualProduction/Avalanche/Source/AvalancheTransitionEditor/Private/Actions/AvaTransitionTreeActions.h b/Engine/Plugins/VirtualProduction/Avalanche/Source/AvalancheTransitionEditor/Private/Actions/AvaTransitionTreeActions.h
--- a/Engine/Plugins/VirtualProduction/Avalanche/Source/AvalancheTransitionEditor/Private/Actions/AvaTransitionTreeActions.h
+++ b/Engine/Plugins/VirtualProduction/Avalanche/Source/AvalancheTransitionEditor/Private/Actions/AvaTransitionTreeActions.h
@@ -2,6 +2,7 @@
 
 #pragma once
 
+#include "Debugger/AvaTransitionDebugDefinitions.h"
 #include "AvaTransitionActions.h"
 
 class UAvaTransitionTree;
@@ -27,7 +28,7 @@ protected:
 
 	void ImportTransitionTree(UAvaTransitionTree* InTemplateTree);
 
-#if WITH_STATETREE_DEBUGGER
+#if UE_AVA_WITH_TRANSITION_DEBUG
 	void ToggleDebugger();
 
 	bool IsDebuggerActive() const;
diff --git a/Engine/Plugins/VirtualProduction/Avalanche/Source/AvalancheTransitionEditor/Private/AvaTransitionCommands.cpp b/Engine/Plugins/VirtualProduction/Avalanche/Source/AvalancheTransitionEditor/Private/AvaTransitionCommands.cpp
--- a/Engine/Plugins/VirtualProduction/Avalanche/Source/AvalancheTransitionEditor/Private/AvaTransitionCommands.cpp
+++ b/Engine/Plugins/VirtualProduction/Avalanche/Source/AvalancheTransitionEditor/Private/AvaTransitionCommands.cpp
@@ -82,7 +82,7 @@ void FAvaTransitionEditorCommands::RegisterCommands()
 		, EUserInterfaceActionType::RadioButton
 		, FInputChord());
 
-#if WITH_STATETREE_DEBUGGER
+#if UE_AVA_WITH_TRANSITION_DEBUG
 	UI_COMMAND(ToggleDebug
 		, "Debug"
 		, "Enables the debugger to analyze how each instance of the State Tree is run"
diff --git a/Engine/Plugins/VirtualProduction/Avalanche/Source/AvalancheTransitionEditor/Private/AvaTransitionCommands.h b/Engine/Plugins/VirtualProduction/Avalanche/Source/AvalancheTransitionEditor/Private/AvaTransitionCommands.h
--- a/Engine/Plugins/VirtualProduction/Avalanche/Source/AvalancheTransitionEditor/Private/AvaTransitionCommands.h
+++ b/Engine/Plugins/VirtualProduction/Avalanche/Source/AvalancheTransitionEditor/Private/AvaTransitionCommands.h
@@ -2,6 +2,7 @@
 
 #pragma once
 
+#include "Debugger/AvaTransitionDebugDefinitions.h"
 #include "Framework/Commands/Commands.h"
 
 class FAvaTransitionEditorCommands : public TCommands<FAvaTransitionEditorCommands>
@@ -28,7 +29,7 @@ public:
 	TSharedPtr<FUICommandInfo> SaveOnCompile_SuccessOnly;
 	TSharedPtr<FUICommandInfo> SaveOnCompile_Always;
 
-#if WITH_STATETREE_DEBUGGER
+#if UE_AVA_WITH_TRANSITION_DEBUG
 	TSharedPtr<FUICommandInfo> ToggleDebug;
 #endif
 };
diff --git a/Engine/Plugins/VirtualProduction/Avalanche/Source/AvalancheTransitionEditor/Private/AvaTransitionEditorModule.cpp b/Engine/Plugins/VirtualProduction/Avalanche/Source/AvalancheTransitionEditor/Private/AvaTransitionEditorModule.cpp
--- a/Engine/Plugins/VirtualProduction/Avalanche/Source/AvalancheTransitionEditor/Private/AvaTransitionEditorModule.cpp
+++ b/Engine/Plugins/VirtualProduction/Avalanche/Source/AvalancheTransitionEditor/Private/AvaTransitionEditorModule.cpp
@@ -18,10 +18,6 @@
 #include "ViewModels/AvaTransitionEditorViewModel.h"
 #include "Widgets/Layout/SBox.h"
 
-#if WITH_STATETREE_DEBUGGER
-#include "Debugger/AvaTransitionTraceModule.h"
-#endif
-
 #define LOCTEXT_NAMESPACE "AvaTransitionEditorModule"
 
 DEFINE_LOG_CATEGORY(LogAvaEditorTransition)
@@ -30,10 +26,6 @@ IMPLEMENT_MODULE(FAvaTransitionEditorModule, AvalancheTransitionEditor)
 
 void FAvaTransitionEditorModule::StartupModule()
 {
-#if WITH_STATETREE_DEBUGGER
-	FAvaTransitionTraceModule::Startup();
-#endif
-
 	FAvaTransitionEditorCommands::Register();
 
 	IAvaTransitionModule::Get().GetOnValidateTransitionTree().BindRaw(this, &FAvaTransitionEditorModule::ValidateStateTree);
@@ -45,10 +37,6 @@ void FAvaTransitionEditorModule::ShutdownModule()
 {
 	FAvaTransitionEditorCommands::Unregister();
 
-#if WITH_STATETREE_DEBUGGER
-	FAvaTransitionTraceModule::Shutdown();
-#endif
-
 	IAvaTransitionModule::Get().GetOnValidateTransitionTree().Unbind();
 
 	UE::StateTree::Delegates::OnPostCompile.Remove(OnPostCompileHandle);
diff --git a/Engine/Plugins/VirtualProduction/Avalanche/Source/AvalancheTransitionEditor/Private/AvaTransitionTreeEditorData.cpp b/Engine/Plugins/VirtualProduction/Avalanche/Source/AvalancheTransitionEditor/Private/AvaTransitionTreeEditorData.cpp
--- a/Engine/Plugins/VirtualProduction/Avalanche/Source/AvalancheTransitionEditor/Private/AvaTransitionTreeEditorData.cpp
+++ b/Engine/Plugins/VirtualProduction/Avalanche/Source/AvalancheTransitionEditor/Private/AvaTransitionTreeEditorData.cpp
@@ -58,3 +58,10 @@ UStateTreeState& UAvaTransitionTreeEditorData::CreateState(const UStateTreeState
 	Children.Insert(State, ChildIndex);
 	return *State;
 }
+
+bool UAvaTransitionTreeEditorData::Compare(const UAvaTransitionTreeEditorData* InEditorData) const
+{
+	// Root Parameters guids are unique per tree.
+	// When duplicating trees, even though the parameter itself isn't duplicate transient, it will assign a new guid in UStateTreeEditorData::DuplicateIDs
+	return InEditorData && (this == InEditorData || GetRootParametersGuid() == InEditorData->GetRootParametersGuid());
+}
diff --git a/Engine/Plugins/VirtualProduction/Avalanche/Source/AvalancheTransitionEditor/Private/Debugger/AvaTransitionDebugDefinitions.h b/Engine/Plugins/VirtualProduction/Avalanche/Source/AvalancheTransitionEditor/Private/Debugger/AvaTransitionDebugDefinitions.h
new file mode 100644
--- /dev/null
+++ b/Engine/Plugins/VirtualProduction/Avalanche/Source/AvalancheTransitionEditor/Private/Debugger/AvaTransitionDebugDefinitions.h
@@ -0,0 +1,9 @@
+// Copyright Epic Games, Inc. All Rights Reserved.
+
+#pragma once
+
+#include "StateTreeTypes.h"
+
+#ifndef UE_AVA_WITH_TRANSITION_DEBUG
+#define UE_AVA_WITH_TRANSITION_DEBUG (WITH_STATETREE_DEBUG && WITH_STATETREE_TRACE_DEBUGGER)
+#endif
diff --git a/Engine/Plugins/VirtualProduction/Avalanche/Source/AvalancheTransitionEditor/Private/Debugger/AvaTransitionDebugInfo.h b/Engine/Plugins/VirtualProduction/Avalanche/Source/AvalancheTransitionEditor/Private/Debugger/AvaTransitionDebugInfo.h
--- a/Engine/Plugins/VirtualProduction/Avalanche/Source/AvalancheTransitionEditor/Private/Debugger/AvaTransitionDebugInfo.h
+++ b/Engine/Plugins/VirtualProduction/Avalanche/Source/AvalancheTransitionEditor/Private/Debugger/AvaTransitionDebugInfo.h
@@ -2,7 +2,9 @@
 
 #pragma once
 
-#if WITH_STATETREE_DEBUGGER
+#include "AvaTransitionDebugDefinitions.h"
+
+#if UE_AVA_WITH_TRANSITION_DEBUG
 
 #include "Containers/UnrealString.h"
 #include "Math/Color.h"
@@ -18,4 +20,4 @@ struct FAvaTransitionDebugInfo
 	FLinearColor Color;
 };
 
-#endif // WITH_STATETREE_DEBUGGER
+#endif // UE_AVA_WITH_TRANSITION_DEBUG
diff --git a/Engine/Plugins/VirtualProduction/Avalanche/Source/AvalancheTransitionEditor/Private/Debugger/AvaTransitionDebugger.cpp b/Engine/Plugins/VirtualProduction/Avalanche/Source/AvalancheTransitionEditor/Private/Debugger/AvaTransitionDebugger.cpp
--- a/Engine/Plugins/VirtualProduction/Avalanche/Source/AvalancheTransitionEditor/Private/Debugger/AvaTransitionDebugger.cpp
+++ b/Engine/Plugins/VirtualProduction/Avalanche/Source/AvalancheTransitionEditor/Private/Debugger/AvaTransitionDebugger.cpp
@@ -1,14 +1,17 @@
 // Copyright Epic Games, Inc. All Rights Reserved.
 
-#if WITH_STATETREE_DEBUGGER
-
 #include "AvaTransitionDebugger.h"
+
+#if UE_AVA_WITH_TRANSITION_DEBUG
+
 #include "AvaTransitionEditorLog.h"
-#include "AvaTransitionTraceAnalyzer.h"
 #include "AvaTransitionTree.h"
-#include "Debugger/StateTreeDebugger.h"
+#include "AvaTransitionTreeEditorData.h"
+#include "StateTreeExecutionContext.h"
+#include "Debugger/StateTreeDebug.h"
 #include "Extensions/IAvaTransitionDebuggableExtension.h"
 #include "StateTreeExecutionTypes.h"
+#include "Debugger/StateTreeDebugger.h"
 #include "ViewModels/AvaTransitionEditorViewModel.h"
 #include "ViewModels/AvaTransitionViewModelSharedData.h"
 #include "ViewModels/Registry/AvaTransitionViewModelRegistryCollection.h"
@@ -34,91 +37,150 @@ void FAvaTransitionDebugger::Start()
 		return;
 	}
 
-	TSharedPtr<FAvaTransitionEditorViewModel> EditorViewModel = EditorViewModelWeak.Pin();
-	if (!EditorViewModel.IsValid())
+	bActive = true;
+	RegisterDelegates();
+}
+
+void FAvaTransitionDebugger::Stop()
+{
+	UnregisterDelegates();
+	TreeDebugInstances.Reset();
+	bActive = false;
+}
+
+bool FAvaTransitionDebugger::IsActive() const
+{
+	return bActive;
+}
+
+void FAvaTransitionDebugger::OnTreeInstanceStarted(const FStateTreeInstanceDebugId& InInstanceDebugId, const FString& InInstanceName)
+{
+	// Ensure that a debug instance of the given id doesn't exist already
+	const int32 Index = TreeDebugInstances.IndexOfByKey(InInstanceDebugId);
+	if (Index == INDEX_NONE)
 	{
-		return;
+		TreeDebugInstances.Emplace(InInstanceDebugId, InInstanceName);
 	}
+}
 
-	UAvaTransitionTree* TransitionTree = EditorViewModel->GetTransitionTree();
-	if (!TransitionTree)
+void FAvaTransitionDebugger::OnTreeInstanceStopped(const FStateTreeInstanceDebugId& InInstanceDebugId)
+{
+	const int32 Index = TreeDebugInstances.IndexOfByKey(InInstanceDebugId);
+	if (Index != INDEX_NONE)
 	{
-		UE_LOG(LogAvaEditorTransition, Warning, TEXT("Transition Tree is invalid. Could not start debugger"));
-		return;
+		TreeDebugInstances.RemoveAt(Index);
 	}
+}
 
-	Debugger = MakeShared<FStateTreeDebugger>();
-	Debugger->SetAsset(TransitionTree);
-
-	FAvaTransitionTraceAnalyzer::RegisterDebugger(this);
+void FAvaTransitionDebugger::OnNodeEntered(const FGuid& InNodeId, const FStateTreeInstanceDebugId& InInstanceDebugId)
+{
+	if (FAvaTransitionTreeDebugInstance* TreeDebugInstance = TreeDebugInstances.FindByKey(InInstanceDebugId))
+	{
+		TSharedPtr<IAvaTransitionDebuggableExtension> Debuggable = FindDebuggable(InNodeId);
+		TreeDebugInstance->EnterDebuggable(Debuggable);
+	}
 }
 
-void FAvaTransitionDebugger::Stop()
+void FAvaTransitionDebugger::OnNodeExited(const FGuid& InNodeId, const FStateTreeInstanceDebugId& InInstanceDebugId)
 {
-	Debugger.Reset();
-	TreeDebugInstances.Reset();
+	if (FAvaTransitionTreeDebugInstance* TreeDebugInstance = TreeDebugInstances.FindByKey(InInstanceDebugId))
+	{
+		TSharedPtr<IAvaTransitionDebuggableExtension> Debuggable = FindDebuggable(InNodeId);
+		TreeDebugInstance->ExitDebuggable(Debuggable);
+	}
+}
 
-	FAvaTransitionTraceAnalyzer::UnregisterDebugger(this);
+void FAvaTransitionDebugger::RegisterDelegates()
+{
+	UnregisterDelegates();
+	OnBeginUpdatePhaseHandle = UE::StateTree::Debug::OnBeginUpdatePhase_AnyThread.AddSP(this, &FAvaTransitionDebugger::OnBeginUpdatePhase);
+	OnStateEventHandle = UE::StateTree::Debug::OnStateEvent_AnyThread.AddSP(this, &FAvaTransitionDebugger::OnStateEvent);
 }
 
-bool FAvaTransitionDebugger::IsActive() const
+void FAvaTransitionDebugger::UnregisterDelegates()
 {
-	return Debugger.IsValid();
+	UE::StateTree::Debug::OnBeginUpdatePhase_AnyThread.Remove(OnBeginUpdatePhaseHandle);
+	OnBeginUpdatePhaseHandle.Reset();
+
+	UE::StateTree::Debug::OnStateEvent_AnyThread.Remove(OnStateEventHandle);
+	OnStateEventHandle.Reset();
 }
 
-void FAvaTransitionDebugger::OnTreeInstanceStarted(UStateTree& InStateTree, const FStateTreeInstanceDebugId& InInstanceDebugId, const FString& InInstanceName)
+void FAvaTransitionDebugger::OnBeginUpdatePhase(const FStateTreeExecutionContext& InExecutionContext, EStateTreeUpdatePhase InPhase, FStateTreeStateHandle InStateHandle)
 {
+	constexpr EStateTreeUpdatePhase StartTreePhase = EStateTreeUpdatePhase::StartTree;
+	constexpr EStateTreeUpdatePhase StopTreePhase = EStateTreeUpdatePhase::StopTree;
+
+	// This function is only concerned with Start/Stop tree for now.
+	if (InPhase != StartTreePhase && InPhase != StopTreePhase)
+	{
+		return;
+	}
+
 	TSharedPtr<FAvaTransitionEditorViewModel> EditorViewModel = EditorViewModelWeak.Pin();
 	if (!EditorViewModel.IsValid())
 	{
 		return;
 	}
 
-	// Return early if the State Tree of this debugger doesn't match the one being pushed
-	UStateTree* TransitionTree = EditorViewModel->GetTransitionTree();
-	if (!TransitionTree || TransitionTree->LastCompiledEditorDataHash != InStateTree.LastCompiledEditorDataHash)
+	// Return early if the State Tree of this debugger doesn't match the execution context's
+	const UAvaTransitionTreeEditorData* const EditorData = EditorViewModel->GetEditorData();
+	if (!EditorData || !EditorData->Compare(Cast<UAvaTransitionTreeEditorData>(InExecutionContext.GetStateTree()->EditorData)))
 	{
 		return;
 	}
 
-	// Ensure that a debug instance of the given id doesn't exist already
-	int32 Index = TreeDebugInstances.IndexOfByKey(InInstanceDebugId);
-	if (Index == INDEX_NONE)
+	switch (InPhase)
 	{
-		TreeDebugInstances.Emplace(InInstanceDebugId, InInstanceName);
+	case StartTreePhase:
+		OnTreeInstanceStarted(InExecutionContext.GetInstanceDebugId(), InExecutionContext.GetInstanceDebugDescription());
+		break;
+
+	case StopTreePhase:
+		OnTreeInstanceStopped(InExecutionContext.GetInstanceDebugId());
+		break;
 	}
 }
 
-void FAvaTransitionDebugger::OnTreeInstanceStopped(const FStateTreeInstanceDebugId& InInstanceDebugId)
+void FAvaTransitionDebugger::OnStateEvent(const FStateTreeExecutionContext& InExecutionContext, FStateTreeStateHandle InStateHandle, EStateTreeTraceEventType InEventType)
 {
+	constexpr EStateTreeTraceEventType EnterEvent = EStateTreeTraceEventType::OnEntering;
+	constexpr EStateTreeTraceEventType ExitEvent = EStateTreeTraceEventType::OnExited;
+
+	// This function is only concerned with enter/exit state for now.
+	if (InEventType != EnterEvent && InEventType != ExitEvent)
+	{
+		return;
+	}
+
 	TSharedPtr<FAvaTransitionEditorViewModel> EditorViewModel = EditorViewModelWeak.Pin();
 	if (!EditorViewModel.IsValid())
 	{
 		return;
 	}
 
-	int32 Index = TreeDebugInstances.IndexOfByKey(InInstanceDebugId);
-	if (Index != INDEX_NONE)
+	const UAvaTransitionTree* const TransitionTree = EditorViewModel->GetTransitionTree();
+	if (!TransitionTree)
 	{
-		TreeDebugInstances.RemoveAt(Index);
+		return;
 	}
-}
 
-void FAvaTransitionDebugger::OnNodeEntered(const FGuid& InNodeId, const FStateTreeInstanceDebugId& InInstanceDebugId)
-{
-	if (FAvaTransitionTreeDebugInstance* TreeDebugInstance = TreeDebugInstances.FindByKey(InInstanceDebugId))
+	// Return early if the State Tree of this debugger doesn't match the execution context's
+	const UAvaTransitionTreeEditorData* const EditorData = EditorViewModel->GetEditorData();
+	if (!EditorData || !EditorData->Compare(Cast<UAvaTransitionTreeEditorData>(InExecutionContext.GetStateTree()->EditorData)))
 	{
-		TSharedPtr<IAvaTransitionDebuggableExtension> Debuggable = FindDebuggable(InNodeId);
-		TreeDebugInstance->EnterDebuggable(Debuggable);
+		return;
 	}
-}
 
-void FAvaTransitionDebugger::OnNodeExited(const FGuid& InNodeId, const FStateTreeInstanceDebugId& InInstanceDebugId)
-{
-	if (FAvaTransitionTreeDebugInstance* TreeDebugInstance = TreeDebugInstances.FindByKey(InInstanceDebugId))
+	switch (InEventType)
 	{
-		TSharedPtr<IAvaTransitionDebuggableExtension> Debuggable = FindDebuggable(InNodeId);
-		TreeDebugInstance->ExitDebuggable(Debuggable);
+	case EnterEvent:
+		OnNodeEntered(TransitionTree->GetStateIdFromHandle(InStateHandle), InExecutionContext.GetInstanceDebugId());
+		break;
+
+	case ExitEvent:
+		OnNodeExited(TransitionTree->GetStateIdFromHandle(InStateHandle), InExecutionContext.GetInstanceDebugId());
+		break;
 	}
 }
 
@@ -132,4 +194,4 @@ TSharedPtr<IAvaTransitionDebuggableExtension> FAvaTransitionDebugger::FindDebugg
 	return nullptr;
 }
 
-#endif // WITH_STATETREE_DEBUGGER
+#endif // UE_AVA_WITH_TRANSITION_DEBUG
diff --git a/Engine/Plugins/VirtualProduction/Avalanche/Source/AvalancheTransitionEditor/Private/Debugger/AvaTransitionDebugger.h b/Engine/Plugins/VirtualProduction/Avalanche/Source/AvalancheTransitionEditor/Private/Debugger/AvaTransitionDebugger.h
--- a/Engine/Plugins/VirtualProduction/Avalanche/Source/AvalancheTransitionEditor/Private/Debugger/AvaTransitionDebugger.h
+++ b/Engine/Plugins/VirtualProduction/Avalanche/Source/AvalancheTransitionEditor/Private/Debugger/AvaTransitionDebugger.h
@@ -2,19 +2,22 @@
 
 #pragma once
 
-#if WITH_STATETREE_DEBUGGER
+#include "AvaTransitionDebugDefinitions.h"
+
+#if UE_AVA_WITH_TRANSITION_DEBUG
 
 #include "AvaTransitionTreeDebugInstance.h"
-#include "Delegates/Delegate.h"
 #include "Templates/SharedPointer.h"
 #include "UObject/WeakObjectPtr.h"
 
 class FAvaTransitionEditorViewModel;
 class IAvaTransitionDebuggableExtension;
 class UStateTree;
-struct FStateTreeDebugger;
+enum class EStateTreeTraceEventType : uint8;
+enum class EStateTreeUpdatePhase : uint8;
 struct FStateTreeInstanceDebugId;
 
+/** Handles forwarding the debug events from state tree execution to the debuggables */
 class FAvaTransitionDebugger : public TSharedFromThis<FAvaTransitionDebugger>
 {
 public:
@@ -22,30 +25,57 @@ public:
 
 	~FAvaTransitionDebugger();
 
+	/** Called once to set the editor view model */
 	void Initialize(const TSharedRef<FAvaTransitionEditorViewModel>& InEditorViewModel);
 
+	/** Starts the debugger */
 	void Start();
 
+	/** Stops the debugger */
 	void Stop();
 
+	/** Returns whether the debugger is currently active */
 	bool IsActive() const;
 
-	void OnTreeInstanceStarted(UStateTree& InStateTree, const FStateTreeInstanceDebugId& InInstanceDebugId, const FString& InInstanceName);
+	/** Called when an instance of the transition tree has started */
+	void OnTreeInstanceStarted(const FStateTreeInstanceDebugId& InInstanceDebugId, const FString& InInstanceName);
 
+	/** Called when an instance of the transition tree has stopped */
 	void OnTreeInstanceStopped(const FStateTreeInstanceDebugId& InInstanceDebugId);
 
+	/** Called when a node with the given id was started / entered */
 	void OnNodeEntered(const FGuid& InNodeId, const FStateTreeInstanceDebugId& InInstanceDebugId);
 
+	/** Called when a node with the given id was stopped / exited */
 	void OnNodeExited(const FGuid& InNodeId, const FStateTreeInstanceDebugId& InInstanceDebugId);
 
 private:
+	void RegisterDelegates();
+	void UnregisterDelegates();
+
+	/** Called when an update phase has started */
+	void OnBeginUpdatePhase(const FStateTreeExecutionContext& InExecutionContext, EStateTreeUpdatePhase InPhase, FStateTreeStateHandle InStateHandle);
+
+	/** Called when a state event has triggered */
+	void OnStateEvent(const FStateTreeExecutionContext& InExecutionContext, FStateTreeStateHandle InStateHandle, EStateTreeTraceEventType InEventType);
+
+	/** Finds the view model with the matching id and returns it as debuggable extension, or null if the view model doesn't exist or doesn't have a debuggable extension */
 	TSharedPtr<IAvaTransitionDebuggableExtension> FindDebuggable(const FGuid& InNodeId) const;
 
+	/** the view model owning this debugger */
 	TWeakPtr<FAvaTransitionEditorViewModel> EditorViewModelWeak;
 
-	TSharedPtr<FStateTreeDebugger> Debugger;
-
+	/** All the tree instances that are in execution and being debugged */
 	TArray<FAvaTransitionTreeDebugInstance> TreeDebugInstances;
+
+	/** Handle to the delegate when an update phase has begun */
+	FDelegateHandle OnBeginUpdatePhaseHandle;
+
+	/** Handle to the delegate when a state event has triggered */
+	FDelegateHandle OnStateEventHandle;
+
+	/** Determines whether the debugger is currently active */
+	bool bActive = false;
 };
 
-#endif // WITH_STATETREE_DEBUGGER
+#endif // UE_AVA_WITH_TRANSITION_DEBUG
diff --git a/Engine/Plugins/VirtualProduction/Avalanche/Source/AvalancheTransitionEditor/Private/Debugger/AvaTransitionStateDebugInstance.cpp b/Engine/Plugins/VirtualProduction/Avalanche/Source/AvalancheTransitionEditor/Private/Debugger/AvaTransitionStateDebugInstance.cpp
--- a/Engine/Plugins/VirtualProduction/Avalanche/Source/AvalancheTransitionEditor/Private/Debugger/AvaTransitionStateDebugInstance.cpp
+++ b/Engine/Plugins/VirtualProduction/Avalanche/Source/AvalancheTransitionEditor/Private/Debugger/AvaTransitionStateDebugInstance.cpp
@@ -1,9 +1,9 @@
 // Copyright Epic Games, Inc. All Rights Reserved.
 
-#if WITH_STATETREE_DEBUGGER
-
 #include "AvaTransitionStateDebugInstance.h"
 
+#if UE_AVA_WITH_TRANSITION_DEBUG
+
 namespace UE::AvaTransitionEditor::Private
 {
 	constexpr static float DebugEnteringTime = 0.25f;
@@ -101,4 +101,4 @@ void FAvaTransitionStateDebugInstance::SetDebugState(EAvaTransitionStateDebugSta
 	}
 }
 
-#endif // WITH_STATETREE_DEBUGGER
+#endif // UE_AVA_WITH_TRANSITION_DEBUG
diff --git a/Engine/Plugins/VirtualProduction/Avalanche/Source/AvalancheTransitionEditor/Private/Debugger/AvaTransitionStateDebugInstance.h b/Engine/Plugins/VirtualProduction/Avalanche/Source/AvalancheTransitionEditor/Private/Debugger/AvaTransitionStateDebugInstance.h
--- a/Engine/Plugins/VirtualProduction/Avalanche/Source/AvalancheTransitionEditor/Private/Debugger/AvaTransitionStateDebugInstance.h
+++ b/Engine/Plugins/VirtualProduction/Avalanche/Source/AvalancheTransitionEditor/Private/Debugger/AvaTransitionStateDebugInstance.h
@@ -2,7 +2,9 @@
 
 #pragma once
 
-#if WITH_STATETREE_DEBUGGER
+#include "AvaTransitionDebugDefinitions.h"
+
+#if UE_AVA_WITH_TRANSITION_DEBUG
 
 #include "Debugger/AvaTransitionDebugInfo.h"
 #include "StateTreeExecutionTypes.h"
@@ -70,4 +72,4 @@ private:
 	float StatusElapsedTime = 0.f;
 };
 
-#endif // WITH_STATETREE_DEBUGGER
+#endif // UE_AVA_WITH_TRANSITION_DEBUG
diff --git a/Engine/Plugins/VirtualProduction/Avalanche/Source/AvalancheTransitionEditor/Private/Debugger/AvaTransitionTraceAnalyzer.cpp b/Engine/Plugins/VirtualProduction/Avalanche/Source/AvalancheTransitionEditor/Private/Debugger/AvaTransitionTraceAnalyzer.cpp
deleted file mode 100644
diff --git a/Engine/Plugins/VirtualProduction/Avalanche/Source/AvalancheTransitionEditor/Private/Debugger/AvaTransitionTraceAnalyzer.h b/Engine/Plugins/VirtualProduction/Avalanche/Source/AvalancheTransitionEditor/Private/Debugger/AvaTransitionTraceAnalyzer.h
deleted file mode 100644
diff --git a/Engine/Plugins/VirtualProduction/Avalanche/Source/AvalancheTransitionEditor/Private/Debugger/AvaTransitionTraceModule.cpp b/Engine/Plugins/VirtualProduction/Avalanche/Source/AvalancheTransitionEditor/Private/Debugger/AvaTransitionTraceModule.cpp
deleted file mode 100644
diff --git a/Engine/Plugins/VirtualProduction/Avalanche/Source/AvalancheTransitionEditor/Private/Debugger/AvaTransitionTraceModule.h b/Engine/Plugins/VirtualProduction/Avalanche/Source/AvalancheTransitionEditor/Private/Debugger/AvaTransitionTraceModule.h
deleted file mode 100644
diff --git a/Engine/Plugins/VirtualProduction/Avalanche/Source/AvalancheTransitionEditor/Private/Debugger/AvaTransitionTreeDebugInstance.cpp b/Engine/Plugins/VirtualProduction/Avalanche/Source/AvalancheTransitionEditor/Private/Debugger/AvaTransitionTreeDebugInstance.cpp
--- a/Engine/Plugins/VirtualProduction/Avalanche/Source/AvalancheTransitionEditor/Private/Debugger/AvaTransitionTreeDebugInstance.cpp
+++ b/Engine/Plugins/VirtualProduction/Avalanche/Source/AvalancheTransitionEditor/Private/Debugger/AvaTransitionTreeDebugInstance.cpp
@@ -1,8 +1,9 @@
 // Copyright Epic Games, Inc. All Rights Reserved.
 
-#if WITH_STATETREE_DEBUGGER
-
 #include "AvaTransitionTreeDebugInstance.h"
+
+#if UE_AVA_WITH_TRANSITION_DEBUG
+
 #include "Extensions/IAvaTransitionDebuggableExtension.h"
 #include "StateTreeExecutionTypes.h"
 
@@ -72,4 +73,4 @@ void FAvaTransitionTreeDebugInstance::Reset()
 	Debuggables.Reset();
 }
 
-#endif // WITH_STATETREE_DEBUGGER
+#endif // UE_AVA_WITH_TRANSITION_DEBUG
diff --git a/Engine/Plugins/VirtualProduction/Avalanche/Source/AvalancheTransitionEditor/Private/Debugger/AvaTransitionTreeDebugInstance.h b/Engine/Plugins/VirtualProduction/Avalanche/Source/AvalancheTransitionEditor/Private/Debugger/AvaTransitionTreeDebugInstance.h
--- a/Engine/Plugins/VirtualProduction/Avalanche/Source/AvalancheTransitionEditor/Private/Debugger/AvaTransitionTreeDebugInstance.h
+++ b/Engine/Plugins/VirtualProduction/Avalanche/Source/AvalancheTransitionEditor/Private/Debugger/AvaTransitionTreeDebugInstance.h
@@ -2,7 +2,9 @@
 
 #pragma once
 
-#if WITH_STATETREE_DEBUGGER
+#include "AvaTransitionDebugDefinitions.h"
+
+#if UE_AVA_WITH_TRANSITION_DEBUG
 
 #include "AvaTransitionDebugInfo.h"
 #include "Containers/Array.h"
@@ -33,4 +35,4 @@ private:
 	FAvaTransitionDebugInfo DebugInfo;
 };
 
-#endif // WITH_STATETREE_DEBUGGER
+#endif // UE_AVA_WITH_TRANSITION_DEBUG
diff --git a/Engine/Plugins/VirtualProduction/Avalanche/Source/AvalancheTransitionEditor/Private/Extensions/IAvaTransitionDebuggableExtension.h b/Engine/Plugins/VirtualProduction/Avalanche/Source/AvalancheTransitionEditor/Private/Extensions/IAvaTransitionDebuggableExtension.h
--- a/Engine/Plugins/VirtualProduction/Avalanche/Source/AvalancheTransitionEditor/Private/Extensions/IAvaTransitionDebuggableExtension.h
+++ b/Engine/Plugins/VirtualProduction/Avalanche/Source/AvalancheTransitionEditor/Private/Extensions/IAvaTransitionDebuggableExtension.h
@@ -2,9 +2,10 @@
 
 #pragma once
 
+#include "Debugger/AvaTransitionDebugDefinitions.h"
 #include "AvaType.h"
 
-#if WITH_STATETREE_DEBUGGER
+#if UE_AVA_WITH_TRANSITION_DEBUG
 struct FAvaTransitionDebugInfo;
 #endif
 
@@ -14,7 +15,7 @@ class IAvaTransitionDebuggableExtension
 public:
 	UE_AVA_TYPE(IAvaTransitionDebuggableExtension)
 
-#if WITH_STATETREE_DEBUGGER
+#if UE_AVA_WITH_TRANSITION_DEBUG
 	virtual void DebugEnter(const FAvaTransitionDebugInfo& InDebugInfo) = 0;
 
 	virtual void DebugExit(const FAvaTransitionDebugInfo& InDebugInfo) = 0;
diff --git a/Engine/Plugins/VirtualProduction/Avalanche/Source/AvalancheTransitionEditor/Private/Menu/AvaTransitionToolbar.cpp b/Engine/Plugins/VirtualProduction/Avalanche/Source/AvalancheTransitionEditor/Private/Menu/AvaTransitionToolbar.cpp
--- a/Engine/Plugins/VirtualProduction/Avalanche/Source/AvalancheTransitionEditor/Private/Menu/AvaTransitionToolbar.cpp
+++ b/Engine/Plugins/VirtualProduction/Avalanche/Source/AvalancheTransitionEditor/Private/Menu/AvaTransitionToolbar.cpp
@@ -119,7 +119,7 @@ void FAvaTransitionToolbar::ExtendEditorToolbar(UToolMenu* InToolbarMenu)
 			, FText::GetEmpty())));
 	}
 
-#if WITH_STATETREE_DEBUGGER
+#if UE_AVA_WITH_TRANSITION_DEBUG
 	AllowInReadOnly(MakeSeparator(Section));
 	AllowInReadOnly(Section.AddEntry(FToolMenuEntry::InitToolBarButton(Commands.ToggleDebug)));
 #endif
diff --git a/Engine/Plugins/VirtualProduction/Avalanche/Source/AvalancheTransitionEditor/Private/ViewModels/AvaTransitionViewModelSharedData.cpp b/Engine/Plugins/VirtualProduction/Avalanche/Source/AvalancheTransitionEditor/Private/ViewModels/AvaTransitionViewModelSharedData.cpp
--- a/Engine/Plugins/VirtualProduction/Avalanche/Source/AvalancheTransitionEditor/Private/ViewModels/AvaTransitionViewModelSharedData.cpp
+++ b/Engine/Plugins/VirtualProduction/Avalanche/Source/AvalancheTransitionEditor/Private/ViewModels/AvaTransitionViewModelSharedData.cpp
@@ -8,7 +8,7 @@
 FAvaTransitionViewModelSharedData::FAvaTransitionViewModelSharedData()
 	: RegistryCollection(MakeShared<FAvaTransitionViewModelRegistryCollection>())
 	, Selection(MakeShared<FAvaTransitionSelection>())
-#if WITH_STATETREE_DEBUGGER
+#if UE_AVA_WITH_TRANSITION_DEBUG
 	, Debugger(MakeShared<FAvaTransitionDebugger>())
 #endif
 {
@@ -18,7 +18,7 @@ void FAvaTransitionViewModelSharedData::Initialize(const TSharedRef<FAvaTransiti
 {
 	EditorViewModelWeak = InEditorViewModel;
 
-#if WITH_STATETREE_DEBUGGER
+#if UE_AVA_WITH_TRANSITION_DEBUG
 	Debugger->Initialize(InEditorViewModel);
 #endif
 }
diff --git a/Engine/Plugins/VirtualProduction/Avalanche/Source/AvalancheTransitionEditor/Private/ViewModels/AvaTransitionViewModelSharedData.h b/Engine/Plugins/VirtualProduction/Avalanche/Source/AvalancheTransitionEditor/Private/ViewModels/AvaTransitionViewModelSharedData.h
--- a/Engine/Plugins/VirtualProduction/Avalanche/Source/AvalancheTransitionEditor/Private/ViewModels/AvaTransitionViewModelSharedData.h
+++ b/Engine/Plugins/VirtualProduction/Avalanche/Source/AvalancheTransitionEditor/Private/ViewModels/AvaTransitionViewModelSharedData.h
@@ -3,6 +3,7 @@
 #pragma once
 
 #include "AvaTransitionEditorEnums.h"
+#include "Debugger/AvaTransitionDebugDefinitions.h"
 #include "Templates/SharedPointer.h"
 
 class FAvaTransitionEditorViewModel;
@@ -10,7 +11,7 @@ class FAvaTransitionSelection;
 class FAvaTransitionViewModel;
 class FAvaTransitionViewModelRegistryCollection;
 
-#if WITH_STATETREE_DEBUGGER
+#if UE_AVA_WITH_TRANSITION_DEBUG
 class FAvaTransitionDebugger;
 #endif
 
@@ -42,7 +43,7 @@ public:
 		return Selection;
 	}
 
-#if WITH_STATETREE_DEBUGGER
+#if UE_AVA_WITH_TRANSITION_DEBUG
 	TSharedRef<FAvaTransitionDebugger> GetDebugger() const
 	{
 		return Debugger;
@@ -74,7 +75,7 @@ private:
 
 	TSharedRef<FAvaTransitionSelection> Selection;
 
-#if WITH_STATETREE_DEBUGGER
+#if UE_AVA_WITH_TRANSITION_DEBUG
 	TSharedRef<FAvaTransitionDebugger> Debugger;
 #endif
 
diff --git a/Engine/Plugins/VirtualProduction/Avalanche/Source/AvalancheTransitionEditor/Private/ViewModels/State/AvaTransitionStateViewModel.cpp b/Engine/Plugins/VirtualProduction/Avalanche/Source/AvalancheTransitionEditor/Private/ViewModels/State/AvaTransitionStateViewModel.cpp
--- a/Engine/Plugins/VirtualProduction/Avalanche/Source/AvalancheTransitionEditor/Private/ViewModels/State/AvaTransitionStateViewModel.cpp
+++ b/Engine/Plugins/VirtualProduction/Avalanche/Source/AvalancheTransitionEditor/Private/ViewModels/State/AvaTransitionStateViewModel.cpp
@@ -17,7 +17,7 @@
 #include "Views/SAvaTransitionStateView.h"
 #include "Widgets/SBoxPanel.h"
 
-#if WITH_STATETREE_DEBUGGER
+#if UE_AVA_WITH_TRANSITION_DEBUG
 #include "Views/Debugger/SAvaTransitionStateDebugInstanceContainer.h"
 #endif
 
@@ -172,7 +172,7 @@ void FAvaTransitionStateViewModel::SetComment(const FString& InComment)
 
 bool FAvaTransitionStateViewModel::HasAnyBreakpoint() const
 {
-#if WITH_STATETREE_DEBUGGER
+#if UE_AVA_WITH_TRANSITION_DEBUG
 	const UStateTreeState* State = GetState();
 	const UStateTreeEditorData* EditorData = GetEditorData();
 	return State && EditorData && EditorData->HasAnyBreakpoint(State->ID);
@@ -183,7 +183,7 @@ bool FAvaTransitionStateViewModel::HasAnyBreakpoint() const
 
 FText FAvaTransitionStateViewModel::GetBreakpointTooltip() const
 {
-#if WITH_STATETREE_DEBUGGER
+#if UE_AVA_WITH_TRANSITION_DEBUG
 	const UStateTreeState* State = GetState();
 	const UStateTreeEditorData* EditorData = GetEditorData();
 	if (State && EditorData)
@@ -240,7 +240,7 @@ void FAvaTransitionStateViewModel::UpdateStateDescription(bool bInSetStateName)
 	}
 }
 
-#if WITH_STATETREE_DEBUGGER
+#if UE_AVA_WITH_TRANSITION_DEBUG
 TSharedRef<SWidget> FAvaTransitionStateViewModel::GetOrCreateDebugIndicatorWidget()
 {
 	if (!DebugIndicatorWidget.IsValid())
@@ -254,7 +254,7 @@ const FAvaTransitionStateDebugInstance* FAvaTransitionStateViewModel::FindDebugI
 {
 	return DebugInstances.FindByKey(InId);
 }
-#endif // WITH_STATETREE_DEBUGGER
+#endif // UE_AVA_WITH_TRANSITION_DEBUG
 
 void FAvaTransitionStateViewModel::GatherChildren(FAvaTransitionViewModelChildren& OutChildren)
 {
@@ -385,7 +385,7 @@ const FGuid& FAvaTransitionStateViewModel::GetGuid() const
 
 void FAvaTransitionStateViewModel::Tick(float InDeltaTime)
 {
-#if WITH_STATETREE_DEBUGGER
+#if UE_AVA_WITH_TRANSITION_DEBUG
 	for (FAvaTransitionStateDebugInstance& StateDebugInstance : DebugInstances)
 	{
 		StateDebugInstance.Tick(InDeltaTime);
@@ -402,10 +402,10 @@ void FAvaTransitionStateViewModel::Tick(float InDeltaTime)
 	{
 		DebugIndicatorWidget->Refresh();
 	}
-#endif // WITH_STATETREE_DEBUGGER
+#endif // UE_AVA_WITH_TRANSITION_DEBUG
 }
 
-#if WITH_STATETREE_DEBUGGER
+#if UE_AVA_WITH_TRANSITION_DEBUG
 void FAvaTransitionStateViewModel::DebugEnter(const FAvaTransitionDebugInfo& InDebugInfo)
 {
 	FAvaTransitionStateDebugInstance* DebugInstance = DebugInstances.FindByKey(InDebugInfo.Id);
@@ -430,6 +430,6 @@ void FAvaTransitionStateViewModel::DebugExit(const FAvaTransitionDebugInfo& InDe
 		DebugInstance->Exit();
 	}
 }
-#endif // WITH_STATETREE_DEBUGGER
+#endif // UE_AVA_WITH_TRANSITION_DEBUG
 
 #undef LOCTEXT_NAMESPACE
diff --git a/Engine/Plugins/VirtualProduction/Avalanche/Source/AvalancheTransitionEditor/Private/ViewModels/State/AvaTransitionStateViewModel.h b/Engine/Plugins/VirtualProduction/Avalanche/Source/AvalancheTransitionEditor/Private/ViewModels/State/AvaTransitionStateViewModel.h
--- a/Engine/Plugins/VirtualProduction/Avalanche/Source/AvalancheTransitionEditor/Private/ViewModels/State/AvaTransitionStateViewModel.h
+++ b/Engine/Plugins/VirtualProduction/Avalanche/Source/AvalancheTransitionEditor/Private/ViewModels/State/AvaTransitionStateViewModel.h
@@ -2,6 +2,7 @@
 
 #pragma once
 
+#include "Debugger/AvaTransitionDebugDefinitions.h"
 #include "Extensions/IAvaTransitionDebuggableExtension.h"
 #include "Extensions/IAvaTransitionDragDropExtension.h"
 #include "Extensions/IAvaTransitionGuidExtension.h"
@@ -15,7 +16,7 @@
 #include "UObject/WeakObjectPtr.h"
 #include "ViewModels/AvaTransitionViewModel.h"
 
-#if WITH_STATETREE_DEBUGGER
+#if UE_AVA_WITH_TRANSITION_DEBUG
 #include "Debugger/AvaTransitionStateDebugInstance.h"
 #endif
 
@@ -27,7 +28,7 @@ class UStateTreeState;
 enum class EStateTreeStateSelectionBehavior : uint8;
 struct FAvaTransitionAutogeneratedStateContext;
 
-#if WITH_STATETREE_DEBUGGER
+#if UE_AVA_WITH_TRANSITION_DEBUG
 class SAvaTransitionStateDebugInstanceContainer;
 struct FStateTreeInstanceDebugId;
 #endif
@@ -116,7 +117,7 @@ public:
 		return TransitionContainer;
 	}
 
-#if WITH_STATETREE_DEBUGGER
+#if UE_AVA_WITH_TRANSITION_DEBUG
 	TConstArrayView<FAvaTransitionStateDebugInstance> GetDebugInstances() const
 	{
 		return DebugInstances;
@@ -163,7 +164,7 @@ public:
 	virtual void Tick(float InDeltaTime) override;
 	//~ End IAvaTransitionTickableExtension
 
-#if WITH_STATETREE_DEBUGGER
+#if UE_AVA_WITH_TRANSITION_DEBUG
 	//~ Begin IAvaTransitionDebuggableExtension
 	virtual void DebugEnter(const FAvaTransitionDebugInfo& InDebugInfo) override;
 	virtual void DebugExit(const FAvaTransitionDebugInfo& InDebugInfo) override;
@@ -182,7 +183,7 @@ private:
 
 	TSharedPtr<FAvaTransitionTransitionContainerViewModel> TransitionContainer;
 
-#if WITH_STATETREE_DEBUGGER
+#if UE_AVA_WITH_TRANSITION_DEBUG
 	TArray<FAvaTransitionStateDebugInstance> DebugInstances;
 
 	TSharedPtr<SAvaTransitionStateDebugInstanceContainer> DebugIndicatorWidget;
diff --git a/Engine/Plugins/VirtualProduction/Avalanche/Source/AvalancheTransitionEditor/Private/ViewModels/Task/AvaTransitionTaskViewModel.cpp b/Engine/Plugins/VirtualProduction/Avalanche/Source/AvalancheTransitionEditor/Private/ViewModels/Task/AvaTransitionTaskViewModel.cpp
--- a/Engine/Plugins/VirtualProduction/Avalanche/Source/AvalancheTransitionEditor/Private/ViewModels/Task/AvaTransitionTaskViewModel.cpp
+++ b/Engine/Plugins/VirtualProduction/Avalanche/Source/AvalancheTransitionEditor/Private/ViewModels/Task/AvaTransitionTaskViewModel.cpp
@@ -88,7 +88,7 @@ bool FAvaTransitionTaskViewModel::IsEnabled() const
 
 EVisibility FAvaTransitionTaskViewModel::GetBreakpointVisibility() const
 {
-#if WITH_STATETREE_DEBUGGER
+#if UE_AVA_WITH_TRANSITION_DEBUG
 	const UAvaTransitionTreeEditorData* EditorData = GetEditorData();
 	if (EditorData && EditorData->HasAnyBreakpoint(GetNodeId()))
 	{
@@ -100,7 +100,7 @@ EVisibility FAvaTransitionTaskViewModel::GetBreakpointVisibility() const
 
 FText FAvaTransitionTaskViewModel::GetBreakpointTooltip() const
 {
-#if WITH_STATETREE_DEBUGGER
+#if UE_AVA_WITH_TRANSITION_DEBUG
 	if (const UAvaTransitionTreeEditorData* EditorData = GetEditorData())
 	{
 		const bool bHasBreakpointOnEnter = EditorData->HasBreakpoint(GetNodeId(), EStateTreeBreakpointType::OnEnter);
@@ -120,7 +120,8 @@ FText FAvaTransitionTaskViewModel::GetBreakpointTooltip() const
 			return LOCTEXT("BreakpointOnExitTooltip","Break when exiting task");
 		}
 	}
-#endif
+#endif // UE_AVA_WITH_TRANSITION_DEBUG
+
 	return FText::GetEmpty();
 }
 
diff --git a/Engine/Plugins/VirtualProduction/Avalanche/Source/AvalancheTransitionEditor/Private/ViewModels/Transition/AvaTransitionTransitionViewModel.cpp b/Engine/Plugins/VirtualProduction/Avalanche/Source/AvalancheTransitionEditor/Private/ViewModels/Transition/AvaTransitionTransitionViewModel.cpp
--- a/Engine/Plugins/VirtualProduction/Avalanche/Source/AvalancheTransitionEditor/Private/ViewModels/Transition/AvaTransitionTransitionViewModel.cpp
+++ b/Engine/Plugins/VirtualProduction/Avalanche/Source/AvalancheTransitionEditor/Private/ViewModels/Transition/AvaTransitionTransitionViewModel.cpp
@@ -128,7 +128,7 @@ FText FAvaTransitionTransitionViewModel::GetDescription() const
 
 EVisibility FAvaTransitionTransitionViewModel::GetBreakpointVisibility() const
 {
-#if WITH_STATETREE_DEBUGGER
+#if UE_AVA_WITH_TRANSITION_DEBUG
 	const UAvaTransitionTreeEditorData* EditorData = GetEditorData();
 	if (EditorData && EditorData->HasBreakpoint(TransitionId, EStateTreeBreakpointType::OnTransition))
 	{
diff --git a/Engine/Plugins/VirtualProduction/Avalanche/Source/AvalancheTransitionEditor/Private/Views/Debugger/SAvaTransitionStateDebugInstance.cpp b/Engine/Plugins/VirtualProduction/Avalanche/Source/AvalancheTransitionEditor/Private/Views/Debugger/SAvaTransitionStateDebugInstance.cpp
--- a/Engine/Plugins/VirtualProduction/Avalanche/Source/AvalancheTransitionEditor/Private/Views/Debugger/SAvaTransitionStateDebugInstance.cpp
+++ b/Engine/Plugins/VirtualProduction/Avalanche/Source/AvalancheTransitionEditor/Private/Views/Debugger/SAvaTransitionStateDebugInstance.cpp
@@ -1,8 +1,8 @@
 // Copyright Epic Games, Inc. All Rights Reserved.
 
-#if WITH_STATETREE_DEBUGGER
-
 #include "SAvaTransitionStateDebugInstance.h"
+
+#if UE_AVA_WITH_TRANSITION_DEBUG
 #include "Debugger/AvaTransitionStateDebugInstance.h"
 #include "Styling/AvaTransitionEditorStyle.h"
 #include "ViewModels/State/AvaTransitionStateViewModel.h"
@@ -117,4 +117,4 @@ int32 SAvaTransitionStateDebugInstance::GetWidgetIndex() const
 	return 0;
 }
 
-#endif // WITH_STATETREE_DEBUGGER
+#endif // UE_AVA_WITH_TRANSITION_DEBUG
diff --git a/Engine/Plugins/VirtualProduction/Avalanche/Source/AvalancheTransitionEditor/Private/Views/Debugger/SAvaTransitionStateDebugInstance.h b/Engine/Plugins/VirtualProduction/Avalanche/Source/AvalancheTransitionEditor/Private/Views/Debugger/SAvaTransitionStateDebugInstance.h
--- a/Engine/Plugins/VirtualProduction/Avalanche/Source/AvalancheTransitionEditor/Private/Views/Debugger/SAvaTransitionStateDebugInstance.h
+++ b/Engine/Plugins/VirtualProduction/Avalanche/Source/AvalancheTransitionEditor/Private/Views/Debugger/SAvaTransitionStateDebugInstance.h
@@ -2,7 +2,9 @@
 
 #pragma once
 
-#if WITH_STATETREE_DEBUGGER
+#include "Debugger/AvaTransitionDebugDefinitions.h"
+
+#if UE_AVA_WITH_TRANSITION_DEBUG
 
 #include "StateTreeExecutionTypes.h"
 #include "Widgets/SCompoundWidget.h"
@@ -32,4 +34,4 @@ private:
 	FStateTreeInstanceDebugId InstanceDebugId;
 };
 
-#endif // WITH_STATETREE_DEBUGGER
+#endif // UE_AVA_WITH_TRANSITION_DEBUG
diff --git a/Engine/Plugins/VirtualProduction/Avalanche/Source/AvalancheTransitionEditor/Private/Views/Debugger/SAvaTransitionStateDebugInstanceContainer.cpp b/Engine/Plugins/VirtualProduction/Avalanche/Source/AvalancheTransitionEditor/Private/Views/Debugger/SAvaTransitionStateDebugInstanceContainer.cpp
--- a/Engine/Plugins/VirtualProduction/Avalanche/Source/AvalancheTransitionEditor/Private/Views/Debugger/SAvaTransitionStateDebugInstanceContainer.cpp
+++ b/Engine/Plugins/VirtualProduction/Avalanche/Source/AvalancheTransitionEditor/Private/Views/Debugger/SAvaTransitionStateDebugInstanceContainer.cpp
@@ -1,8 +1,9 @@
 // Copyright Epic Games, Inc. All Rights Reserved.
 
-#if WITH_STATETREE_DEBUGGER
-
 #include "SAvaTransitionStateDebugInstanceContainer.h"
+
+#if UE_AVA_WITH_TRANSITION_DEBUG
+
 #include "Debugger/AvaTransitionDebugger.h"
 #include "SAvaTransitionStateDebugInstance.h"
 #include "Styling/AvaTransitionEditorStyle.h"
@@ -59,4 +60,4 @@ EVisibility SAvaTransitionStateDebugInstanceContainer::GetDebuggerVisibility() c
 	return EVisibility::Collapsed;
 }
 
-#endif // WITH_STATETREE_DEBUGGER
+#endif // UE_AVA_WITH_TRANSITION_DEBUG
diff --git a/Engine/Plugins/VirtualProduction/Avalanche/Source/AvalancheTransitionEditor/Private/Views/Debugger/SAvaTransitionStateDebugInstanceContainer.h b/Engine/Plugins/VirtualProduction/Avalanche/Source/AvalancheTransitionEditor/Private/Views/Debugger/SAvaTransitionStateDebugInstanceContainer.h
--- a/Engine/Plugins/VirtualProduction/Avalanche/Source/AvalancheTransitionEditor/Private/Views/Debugger/SAvaTransitionStateDebugInstanceContainer.h
+++ b/Engine/Plugins/VirtualProduction/Avalanche/Source/AvalancheTransitionEditor/Private/Views/Debugger/SAvaTransitionStateDebugInstanceContainer.h
@@ -2,7 +2,9 @@
 
 #pragma once
 
-#if WITH_STATETREE_DEBUGGER
+#include "Debugger/AvaTransitionDebugDefinitions.h"
+
+#if UE_AVA_WITH_TRANSITION_DEBUG
 
 #include "StateTreeExecutionTypes.h"
 #include "Widgets/SCompoundWidget.h"
@@ -30,4 +32,4 @@ private:
 	TWeakPtr<FAvaTransitionStateViewModel> StateViewModelWeak;
 };
 
-#endif // WITH_STATETREE_DEBUGGER
+#endif // UE_AVA_WITH_TRANSITION_DEBUG
diff --git a/Engine/Plugins/VirtualProduction/Avalanche/Source/AvalancheTransitionEditor/Private/Views/SAvaTransitionStateView.cpp b/Engine/Plugins/VirtualProduction/Avalanche/Source/AvalancheTransitionEditor/Private/Views/SAvaTransitionStateView.cpp
--- a/Engine/Plugins/VirtualProduction/Avalanche/Source/AvalancheTransitionEditor/Private/Views/SAvaTransitionStateView.cpp
+++ b/Engine/Plugins/VirtualProduction/Avalanche/Source/AvalancheTransitionEditor/Private/Views/SAvaTransitionStateView.cpp
@@ -92,7 +92,7 @@ void SAvaTransitionStateView::Construct(const FArguments& InArgs, const TSharedR
 			];
 	}
 
-#if WITH_STATETREE_DEBUGGER
+#if UE_AVA_WITH_TRANSITION_DEBUG
 	// Debug Info
 	RowContent->AddSlot()
 		.HAlign(HAlign_Left)
diff --git a/Engine/Plugins/VirtualProduction/Avalanche/Source/AvalancheTransitionEditor/Public/AvaTransitionTreeEditorData.h b/Engine/Plugins/VirtualProduction/Avalanche/Source/AvalancheTransitionEditor/Public/AvaTransitionTreeEditorData.h
--- a/Engine/Plugins/VirtualProduction/Avalanche/Source/AvalancheTransitionEditor/Public/AvaTransitionTreeEditorData.h
+++ b/Engine/Plugins/VirtualProduction/Avalanche/Source/AvalancheTransitionEditor/Public/AvaTransitionTreeEditorData.h
@@ -10,7 +10,8 @@
 #include "StateTreeEditorData.h"
 #include "AvaTransitionTreeEditorData.generated.h"
 
-class UAvaTransitionTreeEditorData;
+#define UE_API AVALANCHETRANSITIONEDITOR_API
+
 class UStateTreeState;
 
 namespace UE::AvaTransitionEditor
@@ -28,7 +29,10 @@ class UAvaTransitionTreeEditorData : public UStateTreeEditorData
 public:
 	UAvaTransitionTreeEditorData();
 
-	AVALANCHETRANSITIONEDITOR_API UStateTreeState& CreateState(const UStateTreeState& InSiblingState, bool bInAfter);
+	UE_API UStateTreeState& CreateState(const UStateTreeState& InSiblingState, bool bInAfter);
+
+	/** Returns true if this has the same data as the given editor data */
+	UE_API bool Compare(const UAvaTransitionTreeEditorData* InEditorData) const;
 
 	FAvaTagHandle GetTransitionLayer() const
 	{
@@ -71,3 +75,5 @@ private:
 
 	FSimpleMulticastDelegate OnTreeRequestRefresh;
 };
+
+#undef UE_API
