[AvalancheTests] Automate 'Material Designer - Assign an asset material to a shape'

#rnx

#Synced-CL 45465245

diff --git a/Engine/Plugins/VirtualProduction/Avalanche/Source/AvalancheShapes/Public/DynamicMeshes/AvaShapeDynMeshBase.h b/Engine/Plugins/VirtualProduction/Avalanche/Source/AvalancheShapes/Public/DynamicMeshes/AvaShapeDynMeshBase.h
--- a/Engine/Plugins/VirtualProduction/Avalanche/Source/AvalancheShapes/Public/DynamicMeshes/AvaShapeDynMeshBase.h
+++ b/Engine/Plugins/VirtualProduction/Avalanche/Source/AvalancheShapes/Public/DynamicMeshes/AvaShapeDynMeshBase.h
@@ -221,7 +221,7 @@ public:
 	AVALANCHESHAPES_API TSet<int32> GetMeshesIndexes() const;
 
 	/** find a registered mesh and gets a pointer to it */
-	FAvaShapeMeshData* GetMeshData(int32 InMeshIndex);
+	AVALANCHESHAPES_API FAvaShapeMeshData* GetMeshData(int32 InMeshIndex);
 
 	/** find a registered mesh and gets a pointer to it */
 	const FAvaShapeMeshData* GetMeshData(int32 InMeshIndex) const;
