[Editor TRS] Motion Design editor mode - viewport interaction toggled correctly for gizmos

#rnx

#Synced-CL 50400930

diff --git a/Engine/Plugins/VirtualProduction/Avalanche/Source/AvalancheEditorCore/Private/AvaEditor.cpp b/Engine/Plugins/VirtualProduction/Avalanche/Source/AvalancheEditorCore/Private/AvaEditor.cpp
--- a/Engine/Plugins/VirtualProduction/Avalanche/Source/AvalancheEditorCore/Private/AvaEditor.cpp
+++ b/Engine/Plugins/VirtualProduction/Avalanche/Source/AvalancheEditorCore/Private/AvaEditor.cpp
@@ -26,6 +26,7 @@
 #include "StatusBarSubsystem.h"
 #include "Toolkits/GlobalEditorCommonCommands.h"
 #include "Toolkits/IToolkitHost.h"
+#include "Tools/EdModeInteractiveToolsContext.h"
 #include "UnrealEdGlobals.h"
 #include "Widgets/Docking/SDockTab.h"
 
@@ -286,6 +287,12 @@ void FAvaEditor::Activate(TSharedPtr<IToolkitHost> InOverrideToolkitHost)
 	// Opt out of ITF gizmos and viewport interactions
 	ModeTools.SetSupportsViewportITF(false);
 
+	// Disable viewport interactions to allow legacy gizmos to receive input
+	if (UModeManagerInteractiveToolsContext* ToolsContext = ModeTools.GetInteractiveToolsContext())
+	{
+		ToolsContext->SetViewportInteractionsEnabled(false);
+	}
+
 	RegisterTabSpawners();
 
 	if (UAvaEdMode* const AvaEdMode = Cast<UAvaEdMode>(ModeTools.GetActiveScriptableMode(UAvaEdMode::ModeID)))
@@ -337,6 +344,12 @@ void FAvaEditor::Deactivate()
 
 		// Standard Level Editor default is to support ITF gizmos and interactions
 		ModeTools->SetSupportsViewportITF(true);
+
+		// Re-enable viewport interactions for the Level Editor
+		if (UModeManagerInteractiveToolsContext* ToolsContext = ModeTools->GetInteractiveToolsContext())
+		{
+			ToolsContext->SetViewportInteractionsEnabled(true);
+		}
 	}
 }
 
