Prevent plugin reference to FunctionalTesting except as a DeveloperTool

#rnx

#Synced-CL 50554652

diff --git a/Engine/Plugins/VirtualProduction/Avalanche/Avalanche.uplugin b/Engine/Plugins/VirtualProduction/Avalanche/Avalanche.uplugin
--- a/Engine/Plugins/VirtualProduction/Avalanche/Avalanche.uplugin
+++ b/Engine/Plugins/VirtualProduction/Avalanche/Avalanche.uplugin
@@ -250,7 +250,7 @@
 		},
 		{
 			"Name": "AvalancheFunctionalTest",
-			"Type": "Runtime",
+			"Type": "DeveloperTool",
 			"LoadingPhase": "Default",
 			"TargetDenyList": [ "Server" ]
 		}
