AvalancheFunctionalTest - Added optional actor initialization checking to AScreenshotFunctionalTestAvalanche to fix timing-related screenshot test failures on build machines. New properties allow tests to wait for a specific actor to be fully initialized (BeginPlay, Niagara systems active, components registered) before capturing screenshots, instead of relying on fixed time delays.

#virtualized
#Synced-CL 50706195

diff --git a/Engine/Plugins/VirtualProduction/Avalanche/Source/AvalancheFunctionalTest/AvalancheFunctionalTest.Build.cs b/Engine/Plugins/VirtualProduction/Avalanche/Source/AvalancheFunctionalTest/AvalancheFunctionalTest.Build.cs
--- a/Engine/Plugins/VirtualProduction/Avalanche/Source/AvalancheFunctionalTest/AvalancheFunctionalTest.Build.cs
+++ b/Engine/Plugins/VirtualProduction/Avalanche/Source/AvalancheFunctionalTest/AvalancheFunctionalTest.Build.cs
@@ -11,11 +11,33 @@ public class AvalancheFunctionalTest : ModuleRules
 		PublicDependencyModuleNames.AddRange(
 			new string[]
 			{
+				"Avalanche",
 				"Core",
 				"CoreUObject",
 				"Engine",
 				"FunctionalTesting",
+				"Niagara",
+				"RenderCore",
 			}
 		);
+
+		PrivateDependencyModuleNames.AddRange(
+			new string[]
+			{
+				"Slate",
+				"SlateCore",
+			}
+		);
+
+		if (Target.bBuildEditor)
+		{
+			PrivateDependencyModuleNames.AddRange(
+				new string[]
+				{
+					"UnrealEd",
+					"LevelEditor",
+				}
+			);
+		}
 	}
 }
diff --git a/Engine/Plugins/VirtualProduction/Avalanche/Source/AvalancheFunctionalTest/Private/ScreenshotFunctionalTestAvalanche.cpp b/Engine/Plugins/VirtualProduction/Avalanche/Source/AvalancheFunctionalTest/Private/ScreenshotFunctionalTestAvalanche.cpp
--- a/Engine/Plugins/VirtualProduction/Avalanche/Source/AvalancheFunctionalTest/Private/ScreenshotFunctionalTestAvalanche.cpp
+++ b/Engine/Plugins/VirtualProduction/Avalanche/Source/AvalancheFunctionalTest/Private/ScreenshotFunctionalTestAvalanche.cpp
@@ -6,6 +6,9 @@
 #include "Engine/GameViewportClient.h"
 #include "Tests/AutomationCommon.h"
 #include "UnrealClient.h"
+#include "NiagaraComponent.h"
+#include "NiagaraSystemInstanceController.h"
+#include "Components/PrimitiveComponent.h"
 
 DEFINE_LOG_CATEGORY_STATIC(LogScreenshotAvalanche, Log, All);
 
@@ -32,6 +35,9 @@ void AScreenshotFunctionalTestAvalanche::PrepareTest()
 	// This prevents priority conflicts with Avalanche's custom viewport
 	ApplyAvalancheRenderingOverrides();
 
+	// Record start time for actor readiness timeout
+	ActorWaitStartTime = GetWorld()->GetTimeSeconds();
+
 	// Now call base class PrepareTest which sets up the entire screenshot test
 	Super::PrepareTest();
 }
@@ -115,3 +121,136 @@ void AScreenshotFunctionalTestAvalanche::ApplyAvalancheRenderingOverrides()
 	// Ignore High-DPI settings
 	SecondaryScreenPercentageOverride.SetOverride(100.0f);
 }
+
+bool AScreenshotFunctionalTestAvalanche::IsReady_Implementation()
+{
+	// If actor readiness check is disabled, use base class behavior (fixed delay)
+	if (!bWaitForActorReady)
+	{
+		return Super::IsReady_Implementation();
+	}
+
+	// First check if base class timing requirements are met
+	if (!Super::IsReady_Implementation())
+	{
+		return false;
+	}
+
+	// Check if we've exceeded max wait time
+	const float CurrentTime = GetWorld()->GetTimeSeconds();
+	const float ElapsedWaitTime = CurrentTime - ActorWaitStartTime;
+
+	if (ElapsedWaitTime > MaxActorWaitTime)
+	{
+		UE_LOG(LogScreenshotAvalanche, Warning,
+			TEXT("Screenshot test '%s' timed out waiting for actor to be ready after %.2f seconds. Proceeding with screenshot anyway."),
+			*GetName(), ElapsedWaitTime);
+		return true; // Proceed anyway to avoid hanging
+	}
+
+	// Resolve the actor reference
+	AActor* Actor = ActorToWaitFor.Get();
+	if (!Actor)
+	{
+		// Try to load the actor if it's a soft reference that hasn't been resolved yet
+		Actor = ActorToWaitFor.LoadSynchronous();
+	}
+
+	if (!Actor)
+	{
+		UE_LOG(LogScreenshotAvalanche, Warning,
+			TEXT("Screenshot test '%s' has bWaitForActorReady enabled but ActorToWaitFor is not set or couldn't be resolved. Proceeding with screenshot."),
+			*GetName());
+		return true; // Proceed if actor is not set
+	}
+
+	// Check if the actor is fully initialized
+	if (!IsActorFullyInitialized(Actor))
+	{
+		UE_LOG(LogScreenshotAvalanche, Verbose,
+			TEXT("Screenshot test '%s' waiting for actor '%s' to be fully initialized (elapsed: %.2f / %.2f seconds)"),
+			*GetName(), *Actor->GetName(), ElapsedWaitTime, MaxActorWaitTime);
+		return false;
+	}
+
+	UE_LOG(LogScreenshotAvalanche, Log,
+		TEXT("Screenshot test '%s' confirmed actor '%s' is ready after %.2f seconds"),
+		*GetName(), *Actor->GetName(), ElapsedWaitTime);
+
+	return true;
+}
+
+bool AScreenshotFunctionalTestAvalanche::IsActorFullyInitialized(AActor* Actor) const
+{
+	if (!IsValid(Actor))
+	{
+		return true; // Skip invalid actors
+	}
+
+	// Check if actor has completed BeginPlay
+	if (!Actor->HasActorBegunPlay())
+	{
+		UE_LOG(LogScreenshotAvalanche, VeryVerbose,
+			TEXT("Actor '%s' has not completed BeginPlay yet"), *Actor->GetName());
+		return false;
+	}
+
+	// Check all Niagara components are initialized and active
+	TArray<UNiagaraComponent*> NiagaraComponents;
+	Actor->GetComponents<UNiagaraComponent>(NiagaraComponents);
+
+	for (UNiagaraComponent* NiagaraComp : NiagaraComponents)
+	{
+		if (!IsValid(NiagaraComp))
+		{
+			continue;
+		}
+
+		// Check if Niagara system is active
+		if (!NiagaraComp->IsActive())
+		{
+			UE_LOG(LogScreenshotAvalanche, VeryVerbose,
+				TEXT("Actor '%s': Niagara component '%s' is not active yet"),
+				*Actor->GetName(), *NiagaraComp->GetName());
+			return false;
+		}
+
+		// Check if system instance controller is valid
+		FNiagaraSystemInstanceControllerConstPtr SystemInstanceController = NiagaraComp->GetSystemInstanceController();
+		if (!SystemInstanceController.IsValid())
+		{
+			UE_LOG(LogScreenshotAvalanche, VeryVerbose,
+				TEXT("Actor '%s': Niagara component '%s' does not have a valid system instance controller yet"),
+				*Actor->GetName(), *NiagaraComp->GetName());
+			return false;
+		}
+
+		// Wait for system to have at least one tick to spawn particles
+		// Some systems might need a frame or two to initialize properly
+		const float SystemAge = SystemInstanceController->GetAge();
+		if (SystemAge < 0.1f)
+		{
+			UE_LOG(LogScreenshotAvalanche, VeryVerbose,
+				TEXT("Actor '%s': Niagara component '%s' system is too young (age: %.3f, waiting for at least 0.1s)"),
+				*Actor->GetName(), *NiagaraComp->GetName(), SystemAge);
+			return false;
+		}
+	}
+
+	// Check all primitive components are registered
+	TArray<UPrimitiveComponent*> PrimitiveComponents;
+	Actor->GetComponents<UPrimitiveComponent>(PrimitiveComponents);
+
+	for (UPrimitiveComponent* PrimComp : PrimitiveComponents)
+	{
+		if (IsValid(PrimComp) && !PrimComp->IsRegistered())
+		{
+			UE_LOG(LogScreenshotAvalanche, VeryVerbose,
+				TEXT("Actor '%s': Primitive component '%s' is not registered yet"),
+				*Actor->GetName(), *PrimComp->GetName());
+			return false;
+		}
+	}
+
+	return true;
+}
diff --git a/Engine/Plugins/VirtualProduction/Avalanche/Source/AvalancheFunctionalTest/Public/ScreenshotFunctionalTestAvalanche.h b/Engine/Plugins/VirtualProduction/Avalanche/Source/AvalancheFunctionalTest/Public/ScreenshotFunctionalTestAvalanche.h
--- a/Engine/Plugins/VirtualProduction/Avalanche/Source/AvalancheFunctionalTest/Public/ScreenshotFunctionalTestAvalanche.h
+++ b/Engine/Plugins/VirtualProduction/Avalanche/Source/AvalancheFunctionalTest/Public/ScreenshotFunctionalTestAvalanche.h
@@ -43,14 +43,57 @@ class AScreenshotFunctionalTestAvalanche : public AScreenshotFunctionalTest
 public:
 	AVALANCHEFUNCTIONALTEST_API AScreenshotFunctionalTestAvalanche(const FObjectInitializer& ObjectInitializer);
 
+	/**
+	 * If enabled, the test will wait for the specified actor to be fully initialized
+	 * before taking the screenshot, instead of relying on fixed time/frame delays.
+	 * This ensures actors with Niagara systems, async loading, or complex initialization
+	 * are fully ready before screenshot capture.
+	 */
+	UPROPERTY(EditAnywhere, BlueprintReadWrite, Category = "Screenshot|Actor Readiness", meta = (DisplayName = "Wait For Actor Ready"))
+	bool bWaitForActorReady = false;
+
+	/**
+	 * The actor to wait for before taking the screenshot.
+	 * The test will check if this actor and all its components (especially Niagara systems)
+	 * are fully initialized before proceeding.
+	 * Only used when bWaitForActorReady is enabled.
+	 */
+	UPROPERTY(EditAnywhere, BlueprintReadWrite, Category = "Screenshot|Actor Readiness", meta = (EditCondition = "bWaitForActorReady", DisplayName = "Actor To Wait For"))
+	TSoftObjectPtr<AActor> ActorToWaitFor;
+
+	/**
+	 * Maximum time to wait for the actor to be ready (in seconds).
+	 * If the actor is not ready within this time, the test will proceed anyway
+	 * to avoid hanging indefinitely.
+	 */
+	UPROPERTY(EditAnywhere, BlueprintReadWrite, Category = "Screenshot|Actor Readiness", meta = (EditCondition = "bWaitForActorReady", ClampMin = "0.1", ClampMax = "30.0", UIMin = "0.5", UIMax = "10.0"))
+	float MaxActorWaitTime = 5.0f;
+
 protected:
 	// Override to set up CVars with explicit priority before base class setup
 	AVALANCHEFUNCTIONALTEST_API virtual void PrepareTest() override;
 
+	// Override to check if actor is ready before proceeding
+	AVALANCHEFUNCTIONALTEST_API virtual bool IsReady_Implementation() override;
+
 	// Override to register to both instance and static screenshot delegates for PIE support
 	AVALANCHEFUNCTIONALTEST_API virtual void RequestScreenshot() override;
 	AVALANCHEFUNCTIONALTEST_API virtual void OnScreenShotCaptured(int32 InSizeX, int32 InSizeY, const TArray<FColor>& InImageData) override;
 
+private:
+	/**
+	 * Checks if the specified actor is fully initialized and ready for screenshot.
+	 * Returns true if:
+	 * - Actor has completed BeginPlay
+	 * - All Niagara components are active and have valid system instances
+	 * - All primitive components are registered
+	 * - Niagara systems have had time to spawn particles (at least 0.1s age)
+	 */
+	bool IsActorFullyInitialized(AActor* Actor) const;
+
+	// Time when PrepareTest() was called, used to enforce MaxActorWaitTime
+	float ActorWaitStartTime = 0.0f;
+
 private:
 	/**
 	 * Helper struct to manage CVar overrides at Console priority.
