Motion Design: fix spamming warnings when a shader from the material cache settings is not found

#Synced-CL 50556942

diff --git a/Engine/Plugins/VirtualProduction/Avalanche/Source/AvalancheMaterial/Internal/MaterialCache/AvaMaterialCacheSettings.h b/Engine/Plugins/VirtualProduction/Avalanche/Source/AvalancheMaterial/Internal/MaterialCache/AvaMaterialCacheSettings.h
--- a/Engine/Plugins/VirtualProduction/Avalanche/Source/AvalancheMaterial/Internal/MaterialCache/AvaMaterialCacheSettings.h
+++ b/Engine/Plugins/VirtualProduction/Avalanche/Source/AvalancheMaterial/Internal/MaterialCache/AvaMaterialCacheSettings.h
@@ -117,10 +117,10 @@ private:
 	void RegisterDefaultShaderProfile();
 
 	/** Resolves all the registered shader profiles to its type pointers */
-	void ResolveShaderProfiles();
+	void ResolveShaderProfiles(bool bInLogWarnings);
 
 	/** Resolves a single shader profile to its type pointers */
-	void ResolveShaderProfile(const FAvaShaderProfile& InShaderProfile, TMap<FName, FHashedName>& InOutHashedNames);
+	void ResolveShaderProfile(const FAvaShaderProfile& InShaderProfile, TMap<FName, FHashedName>& InOutHashedNames, bool bInLogWarnings);
 
 	/** Profile to use in realtime (e.g. when playing levels) */
 	UPROPERTY(Config, EditAnywhere, Category="Material Cache")
diff --git a/Engine/Plugins/VirtualProduction/Avalanche/Source/AvalancheMaterial/Private/MaterialCache/AvaMaterialCacheSettings.cpp b/Engine/Plugins/VirtualProduction/Avalanche/Source/AvalancheMaterial/Private/MaterialCache/AvaMaterialCacheSettings.cpp
--- a/Engine/Plugins/VirtualProduction/Avalanche/Source/AvalancheMaterial/Private/MaterialCache/AvaMaterialCacheSettings.cpp
+++ b/Engine/Plugins/VirtualProduction/Avalanche/Source/AvalancheMaterial/Private/MaterialCache/AvaMaterialCacheSettings.cpp
@@ -2,6 +2,7 @@
 
 #include "MaterialCache/AvaMaterialCacheSettings.h"
 #include "MaterialCache/AvaMaterialCacheLog.h"
+#include "UObject/UnrealType.h"
 #include "VertexFactory.h"
 
 const FLazyName UAvaMaterialCacheSettings::DefaultRealtimeProfile(TEXT("BasicSceneRealtime"));
@@ -94,13 +95,17 @@ const UE::Ava::FResolvedShaderProfile* UAvaMaterialCacheSettings::FindResolvedSh
 void UAvaMaterialCacheSettings::PostInitProperties()
 {
 	Super::PostInitProperties();
-	ResolveShaderProfiles();
+
+	constexpr bool bLogWarnings = false;
+	ResolveShaderProfiles(bLogWarnings);
 }
 
 void UAvaMaterialCacheSettings::PostReloadConfig(FProperty* InPropertyThatWasLoaded)
 {
 	Super::PostReloadConfig(InPropertyThatWasLoaded);
-	ResolveShaderProfiles();
+
+	constexpr bool bLogWarnings = false;
+	ResolveShaderProfiles(bLogWarnings);
 }
 
 #if WITH_EDITOR
@@ -110,7 +115,8 @@ void UAvaMaterialCacheSettings::PostEditChangeProperty(FPropertyChangedEvent& In
 
 	if (InPropertyChangedEvent.GetMemberPropertyName() == GET_MEMBER_NAME_CHECKED(UAvaMaterialCacheSettings, ShaderProfiles))
 	{
-		ResolveShaderProfiles();
+		constexpr bool bLogWarnings = true;
+		ResolveShaderProfiles(bLogWarnings);
 	}
 }
 #endif
@@ -134,7 +140,7 @@ void UAvaMaterialCacheSettings::RegisterDefaultShaderProfile()
 		});
 }
 
-void UAvaMaterialCacheSettings::ResolveShaderProfiles()
+void UAvaMaterialCacheSettings::ResolveShaderProfiles(bool bInLogWarnings)
 {
 	UE_LOG(LogAvaMaterialCache, Verbose, TEXT("Resolving Shader Profiles"));
 
@@ -143,15 +149,15 @@ void UAvaMaterialCacheSettings::ResolveShaderProfiles()
 	TMap<FName, FHashedName> HashedNames;
 	for (const FAvaShaderProfile& ShaderProfile : ShaderProfiles)
 	{
-		ResolveShaderProfile(ShaderProfile, HashedNames);
+		ResolveShaderProfile(ShaderProfile, HashedNames, bInLogWarnings);
 	}
 }
 
-void UAvaMaterialCacheSettings::ResolveShaderProfile(const FAvaShaderProfile& InShaderProfile, TMap<FName, FHashedName>& InOutHashedNames)
+void UAvaMaterialCacheSettings::ResolveShaderProfile(const FAvaShaderProfile& InShaderProfile, TMap<FName, FHashedName>& InOutHashedNames, bool bInLogWarnings)
 {
 	if (ResolvedShaderProfiles.Contains(InShaderProfile.ProfileName))
 	{
-		UE_LOG(LogAvaMaterialCache, Warning, TEXT("Shader Profile '%s' already exists. New entry will be ignored."), *InShaderProfile.ProfileName.ToString());
+		UE_CLOG(bInLogWarnings, LogAvaMaterialCache, Warning, TEXT("Shader Profile '%s' already exists. New entry will be ignored."), *InShaderProfile.ProfileName.ToString());
 		return;
 	}
 
@@ -185,7 +191,7 @@ void UAvaMaterialCacheSettings::ResolveShaderProfile(const FAvaShaderProfile& In
 		const FVertexFactoryType* const VertexFactoryType = FVertexFactoryType::GetVFByName(GetHashedName(ShaderCombination.VertexFactoryType));
 		if (!VertexFactoryType)
 		{
-			UE_LOG(LogAvaMaterialCache, Warning, TEXT("Invalid Shader Config. Vertex Factory Type '%s' was not found"), *ShaderCombination.VertexFactoryType.ToString());
+			UE_CLOG(bInLogWarnings, LogAvaMaterialCache, Verbose, TEXT("Invalid Shader Config. Vertex Factory Type '%s' was not found"), *ShaderCombination.VertexFactoryType.ToString());
 			continue;
 		}
 
@@ -194,7 +200,7 @@ void UAvaMaterialCacheSettings::ResolveShaderProfile(const FAvaShaderProfile& In
 			const FShaderPipelineType* PipelineType = FShaderPipelineType::GetShaderPipelineTypeByName(GetHashedName(PipelineTypeName));
 			if (!PipelineType)
 			{
-				UE_LOG(LogAvaMaterialCache, Warning, TEXT("Invalid Shader Config. Pipeline Type '%s' was found"), *PipelineTypeName.ToString());
+				UE_CLOG(bInLogWarnings, LogAvaMaterialCache, Verbose, TEXT("Invalid Shader Config. Pipeline Type '%s' was found"), *PipelineTypeName.ToString());
 				continue;
 			}
 
@@ -208,7 +214,7 @@ void UAvaMaterialCacheSettings::ResolveShaderProfile(const FAvaShaderProfile& In
 			const FShaderType* ShaderType = FindShaderTypeByName(GetHashedName(ShaderTypeName));
 			if (!ShaderType)
 			{
-				UE_LOG(LogAvaMaterialCache, Warning, TEXT("Invalid Shader Config. Shader Type '%s' was not found"), *ShaderTypeName.ToString());
+				UE_CLOG(bInLogWarnings, LogAvaMaterialCache, Verbose, TEXT("Invalid Shader Config. Shader Type '%s' was not found"), *ShaderTypeName.ToString());
 				continue;
 			}
 
