[Backout] - CL50780062
#rnx
Original CL Desc
-----------------------------------------------------------------
Motion Design: fix spamming warnings when a shader from the material cache settings is not found

#Synced-CL 50556942

diff --git a/Engine/Plugins/VirtualProduction/Avalanche/Source/AvalancheMaterial/Internal/MaterialCache/AvaMaterialCacheSettings.h b/Engine/Plugins/VirtualProduction/Avalanche/Source/AvalancheMaterial/Internal/MaterialCache/AvaMaterialCacheSettings.h
--- a/Engine/Plugins/VirtualProduction/Avalanche/Source/AvalancheMaterial/Internal/MaterialCache/AvaMaterialCacheSettings.h
+++ b/Engine/Plugins/VirtualProduction/Avalanche/Source/AvalancheMaterial/Internal/MaterialCache/AvaMaterialCacheSettings.h
@@ -117,10 +117,10 @@ private:
 	void RegisterDefaultShaderProfile();
 
 	/** Resolves all the registered shader profiles to its type pointers */
-	void ResolveShaderProfiles(bool bInLogWarnings);
+	void ResolveShaderProfiles();
 
 	/** Resolves a single shader profile to its type pointers */
-	void ResolveShaderProfile(const FAvaShaderProfile& InShaderProfile, TMap<FName, FHashedName>& InOutHashedNames, bool bInLogWarnings);
+	void ResolveShaderProfile(const FAvaShaderProfile& InShaderProfile, TMap<FName, FHashedName>& InOutHashedNames);
 
 	/** Profile to use in realtime (e.g. when playing levels) */
 	UPROPERTY(Config, EditAnywhere, Category="Material Cache")
diff --git a/Engine/Plugins/VirtualProduction/Avalanche/Source/AvalancheMaterial/Private/MaterialCache/AvaMaterialCacheSettings.cpp b/Engine/Plugins/VirtualProduction/Avalanche/Source/AvalancheMaterial/Private/MaterialCache/AvaMaterialCacheSettings.cpp
--- a/Engine/Plugins/VirtualProduction/Avalanche/Source/AvalancheMaterial/Private/MaterialCache/AvaMaterialCacheSettings.cpp
+++ b/Engine/Plugins/VirtualProduction/Avalanche/Source/AvalancheMaterial/Private/MaterialCache/AvaMaterialCacheSettings.cpp
@@ -2,7 +2,6 @@
 
 #include "MaterialCache/AvaMaterialCacheSettings.h"
 #include "MaterialCache/AvaMaterialCacheLog.h"
-#include "UObject/UnrealType.h"
 #include "VertexFactory.h"
 
 const FLazyName UAvaMaterialCacheSettings::DefaultRealtimeProfile(TEXT("BasicSceneRealtime"));
@@ -95,17 +94,13 @@ const UE::Ava::FResolvedShaderProfile* UAvaMaterialCacheSettings::FindResolvedSh
 void UAvaMaterialCacheSettings::PostInitProperties()
 {
 	Super::PostInitProperties();
-
-	constexpr bool bLogWarnings = false;
-	ResolveShaderProfiles(bLogWarnings);
+	ResolveShaderProfiles();
 }
 
 void UAvaMaterialCacheSettings::PostReloadConfig(FProperty* InPropertyThatWasLoaded)
 {
 	Super::PostReloadConfig(InPropertyThatWasLoaded);
-
-	constexpr bool bLogWarnings = false;
-	ResolveShaderProfiles(bLogWarnings);
+	ResolveShaderProfiles();
 }
 
 #if WITH_EDITOR
@@ -115,8 +110,7 @@ void UAvaMaterialCacheSettings::PostEditChangeProperty(FPropertyChangedEvent& In
 
 	if (InPropertyChangedEvent.GetMemberPropertyName() == GET_MEMBER_NAME_CHECKED(UAvaMaterialCacheSettings, ShaderProfiles))
 	{
-		constexpr bool bLogWarnings = true;
-		ResolveShaderProfiles(bLogWarnings);
+		ResolveShaderProfiles();
 	}
 }
 #endif
@@ -140,7 +134,7 @@ void UAvaMaterialCacheSettings::RegisterDefaultShaderProfile()
 		});
 }
 
-void UAvaMaterialCacheSettings::ResolveShaderProfiles(bool bInLogWarnings)
+void UAvaMaterialCacheSettings::ResolveShaderProfiles()
 {
 	UE_LOG(LogAvaMaterialCache, Verbose, TEXT("Resolving Shader Profiles"));
 
@@ -149,15 +143,15 @@ void UAvaMaterialCacheSettings::ResolveShaderProfiles(bool bInLogWarnings)
 	TMap<FName, FHashedName> HashedNames;
 	for (const FAvaShaderProfile& ShaderProfile : ShaderProfiles)
 	{
-		ResolveShaderProfile(ShaderProfile, HashedNames, bInLogWarnings);
+		ResolveShaderProfile(ShaderProfile, HashedNames);
 	}
 }
 
-void UAvaMaterialCacheSettings::ResolveShaderProfile(const FAvaShaderProfile& InShaderProfile, TMap<FName, FHashedName>& InOutHashedNames, bool bInLogWarnings)
+void UAvaMaterialCacheSettings::ResolveShaderProfile(const FAvaShaderProfile& InShaderProfile, TMap<FName, FHashedName>& InOutHashedNames)
 {
 	if (ResolvedShaderProfiles.Contains(InShaderProfile.ProfileName))
 	{
-		UE_CLOG(bInLogWarnings, LogAvaMaterialCache, Warning, TEXT("Shader Profile '%s' already exists. New entry will be ignored."), *InShaderProfile.ProfileName.ToString());
+		UE_LOG(LogAvaMaterialCache, Warning, TEXT("Shader Profile '%s' already exists. New entry will be ignored."), *InShaderProfile.ProfileName.ToString());
 		return;
 	}
 
@@ -191,7 +185,7 @@ void UAvaMaterialCacheSettings::ResolveShaderProfile(const FAvaShaderProfile& In
 		const FVertexFactoryType* const VertexFactoryType = FVertexFactoryType::GetVFByName(GetHashedName(ShaderCombination.VertexFactoryType));
 		if (!VertexFactoryType)
 		{
-			UE_CLOG(bInLogWarnings, LogAvaMaterialCache, Verbose, TEXT("Invalid Shader Config. Vertex Factory Type '%s' was not found"), *ShaderCombination.VertexFactoryType.ToString());
+			UE_LOG(LogAvaMaterialCache, Warning, TEXT("Invalid Shader Config. Vertex Factory Type '%s' was not found"), *ShaderCombination.VertexFactoryType.ToString());
 			continue;
 		}
 
@@ -200,7 +194,7 @@ void UAvaMaterialCacheSettings::ResolveShaderProfile(const FAvaShaderProfile& In
 			const FShaderPipelineType* PipelineType = FShaderPipelineType::GetShaderPipelineTypeByName(GetHashedName(PipelineTypeName));
 			if (!PipelineType)
 			{
-				UE_CLOG(bInLogWarnings, LogAvaMaterialCache, Verbose, TEXT("Invalid Shader Config. Pipeline Type '%s' was found"), *PipelineTypeName.ToString());
+				UE_LOG(LogAvaMaterialCache, Warning, TEXT("Invalid Shader Config. Pipeline Type '%s' was found"), *PipelineTypeName.ToString());
 				continue;
 			}
 
@@ -214,7 +208,7 @@ void UAvaMaterialCacheSettings::ResolveShaderProfile(const FAvaShaderProfile& In
 			const FShaderType* ShaderType = FindShaderTypeByName(GetHashedName(ShaderTypeName));
 			if (!ShaderType)
 			{
-				UE_CLOG(bInLogWarnings, LogAvaMaterialCache, Verbose, TEXT("Invalid Shader Config. Shader Type '%s' was not found"), *ShaderTypeName.ToString());
+				UE_LOG(LogAvaMaterialCache, Warning, TEXT("Invalid Shader Config. Shader Type '%s' was not found"), *ShaderTypeName.ToString());
 				continue;
 			}
 
