Motion Design: fix spamming warnings when a shader from the material cache settings is not found. These warnings are temporarily set to verbose.

#Synced-CL -1

diff --git a/Engine/Plugins/VirtualProduction/Avalanche/Source/AvalancheMaterial/Private/MaterialCache/AvaMaterialCacheSettings.cpp b/Engine/Plugins/VirtualProduction/Avalanche/Source/AvalancheMaterial/Private/MaterialCache/AvaMaterialCacheSettings.cpp
--- a/Engine/Plugins/VirtualProduction/Avalanche/Source/AvalancheMaterial/Private/MaterialCache/AvaMaterialCacheSettings.cpp
+++ b/Engine/Plugins/VirtualProduction/Avalanche/Source/AvalancheMaterial/Private/MaterialCache/AvaMaterialCacheSettings.cpp
@@ -151,7 +151,7 @@ void UAvaMaterialCacheSettings::ResolveShaderProfile(const FAvaShaderProfile& In
 {
 	if (ResolvedShaderProfiles.Contains(InShaderProfile.ProfileName))
 	{
-		UE_LOG(LogAvaMaterialCache, Warning, TEXT("Shader Profile '%s' already exists. New entry will be ignored."), *InShaderProfile.ProfileName.ToString());
+		UE_LOG(LogAvaMaterialCache, Verbose, TEXT("Shader Profile '%s' already exists. New entry will be ignored."), *InShaderProfile.ProfileName.ToString());
 		return;
 	}
 
@@ -185,7 +185,7 @@ void UAvaMaterialCacheSettings::ResolveShaderProfile(const FAvaShaderProfile& In
 		const FVertexFactoryType* const VertexFactoryType = FVertexFactoryType::GetVFByName(GetHashedName(ShaderCombination.VertexFactoryType));
 		if (!VertexFactoryType)
 		{
-			UE_LOG(LogAvaMaterialCache, Warning, TEXT("Invalid Shader Config. Vertex Factory Type '%s' was not found"), *ShaderCombination.VertexFactoryType.ToString());
+			UE_LOG(LogAvaMaterialCache, Verbose, TEXT("Invalid Shader Config. Vertex Factory Type '%s' was not found"), *ShaderCombination.VertexFactoryType.ToString());
 			continue;
 		}
 
@@ -194,7 +194,7 @@ void UAvaMaterialCacheSettings::ResolveShaderProfile(const FAvaShaderProfile& In
 			const FShaderPipelineType* PipelineType = FShaderPipelineType::GetShaderPipelineTypeByName(GetHashedName(PipelineTypeName));
 			if (!PipelineType)
 			{
-				UE_LOG(LogAvaMaterialCache, Warning, TEXT("Invalid Shader Config. Pipeline Type '%s' was found"), *PipelineTypeName.ToString());
+				UE_LOG(LogAvaMaterialCache, Verbose, TEXT("Invalid Shader Config. Pipeline Type '%s' was found"), *PipelineTypeName.ToString());
 				continue;
 			}
 
@@ -208,7 +208,7 @@ void UAvaMaterialCacheSettings::ResolveShaderProfile(const FAvaShaderProfile& In
 			const FShaderType* ShaderType = FindShaderTypeByName(GetHashedName(ShaderTypeName));
 			if (!ShaderType)
 			{
-				UE_LOG(LogAvaMaterialCache, Warning, TEXT("Invalid Shader Config. Shader Type '%s' was not found"), *ShaderTypeName.ToString());
+				UE_LOG(LogAvaMaterialCache, Verbose, TEXT("Invalid Shader Config. Shader Type '%s' was not found"), *ShaderTypeName.ToString());
 				continue;
 			}
 
