Fixes for CL 51336460 - Remove string duplication in FJsonObject to free memory

#rnx

#Synced-CL 51336291

diff --git a/Engine/Plugins/VirtualProduction/Avalanche/Source/AvalancheMedia/Private/Playable/AvaPlayableRemoteControlValues.cpp b/Engine/Plugins/VirtualProduction/Avalanche/Source/AvalancheMedia/Private/Playable/AvaPlayableRemoteControlValues.cpp
--- a/Engine/Plugins/VirtualProduction/Avalanche/Source/AvalancheMedia/Private/Playable/AvaPlayableRemoteControlValues.cpp
+++ b/Engine/Plugins/VirtualProduction/Avalanche/Source/AvalancheMedia/Private/Playable/AvaPlayableRemoteControlValues.cpp
@@ -496,7 +496,7 @@ void FAvaPlayableRemoteControlValues::CollectReferencedAssetPaths(const TMap<FGu
 
 		// For now, we don't care about the field name.
 		// Just trying to find values that could possibly be an existing asset reference.
-		for (const TPair<FString, TSharedPtr<FJsonValue>>& JsonValue : ValueObject->Values)
+		for (const TPair<UE::FSharedString, TSharedPtr<FJsonValue>>& JsonValue : ValueObject->Values)
 		{
 			UE::AvaPlayableRemoteControlValues::Private::CollectReferencedAssetPaths(JsonValue.Value, OutReferencedPaths);
 		}
