Motion Design: moved rundown worlds to tick with other worlds and tick before end frame sync

#Synced-CL 51441789

diff --git a/Engine/Plugins/VirtualProduction/Avalanche/Source/Avalanche/Private/Framework/AvaGameInstance.cpp b/Engine/Plugins/VirtualProduction/Avalanche/Source/Avalanche/Private/Framework/AvaGameInstance.cpp
--- a/Engine/Plugins/VirtualProduction/Avalanche/Source/Avalanche/Private/Framework/AvaGameInstance.cpp
+++ b/Engine/Plugins/VirtualProduction/Avalanche/Source/Avalanche/Private/Framework/AvaGameInstance.cpp
@@ -15,42 +15,22 @@
 #include "Slate/SceneViewport.h"
 #include "Viewport/AvaViewportQualitySettings.h"
 
-namespace UE::AvaGameInstance::Private
-{
-	/* Frame number marked where a synchronous asset loading occured. */
-	static uint64 MarkedSynchronousAssetLoadingFrame = 0;
-
-	/* Frame number marked where delta seconds are ignored. */
-	static uint64 MarkedIgnoreDeltaSecondsFrame = 0;
-
-	static void UpdateFrameMarkers()
-	{
-		// If the last frame had sync asset loading, mark current frame to ignore delta seconds.
-		if (MarkedSynchronousAssetLoadingFrame != 0 && GFrameCounter > MarkedSynchronousAssetLoadingFrame)
-		{
-			MarkedIgnoreDeltaSecondsFrame = GFrameCounter;
-			MarkedSynchronousAssetLoadingFrame = 0;
-		}
-
-		// Reset marker if stale.
-		if (MarkedIgnoreDeltaSecondsFrame != 0 && GFrameCounter > MarkedIgnoreDeltaSecondsFrame)
-		{
-			MarkedIgnoreDeltaSecondsFrame = 0;
-		}
-	}
-
-	static bool ShouldIgnoreDeltaSecondsForCurrentFrame()
-	{
-		UpdateFrameMarkers();
-		return MarkedIgnoreDeltaSecondsFrame != 0 && GFrameCounter == MarkedIgnoreDeltaSecondsFrame;
-	}
-}
+#if WITH_EDITOR
+#include "Editor.h"
+#endif
 
 UAvaGameInstance::FOnAvaGameInstanceEvent UAvaGameInstance::OnEndPlay;
 UAvaGameInstance::FOnAvaGameInstanceEvent UAvaGameInstance::OnRenderTargetReady;
 
+UAvaGameInstance::UAvaGameInstance()
+	: FTickableGameObject(ETickableTickType::Never) // Register tickable when it is needed
+{
+}
+
 UAvaGameInstance* UAvaGameInstance::Create(UObject* InOuter)
 {
+	checkf(IsInGameThread(), TEXT("Motion Design Game Instance creation can only happen in Game Thread"));
+
 	UAvaGameInstance* GameInstance = NewObject<UAvaGameInstance>(InOuter ? InOuter : GEngine);
 
 	// We need to call Init() (and corresponding Shutdown() later) to get SubsystemCollection to initialize.
@@ -96,9 +76,6 @@ bool UAvaGameInstance::CreateWorld()
 			.SetTransactional(false)
 			//TODO: World Settings
 		);
-
-		// Disallow the Engine to tick this World as this Game Instance will be in charge of ticking it
-		PlayWorld->SetShouldTick(false);
 	}
 
 	WorldContext = &GEngine->CreateNewWorldContext(WorldType);
@@ -165,6 +142,16 @@ bool UAvaGameInstance::BeginPlayWorld(const FAvaInstancePlaySettings& InWorldPla
 		UE_LOG(LogAva, Error, TEXT("Couldn't spawn play actor: %s"), *Error);
 	}
 
+#if WITH_EDITOR
+	if (GEditor && !OnEditorTickHandle.IsValid())
+	{
+		OnEditorTickHandle = GEditor->OnPostEditorTick().AddUObject(this, &UAvaGameInstance::EditorTick);
+	}
+#endif // WITH_EDITOR
+
+	// Register game instance to tick
+	SetTickableTickType(ETickableTickType::Always);
+
 	PlayWorld->BeginPlay();
 
 	FCoreDelegates::OnEndFrame.AddUObject(this, &UAvaGameInstance::OnEndFrameTick);
@@ -238,82 +225,60 @@ void UAvaGameInstance::UpdateSceneViewportSize(const FIntPoint& InViewportSize)
 	}
 }
 
-void UAvaGameInstance::Tick(float DeltaSeconds)
+#if WITH_EDITOR
+void UAvaGameInstance::EditorTick(float InDeltaSeconds)
 {
-	if (bIsTicking)
+	if (!PlayWorld)
 	{
 		return;
 	}
-	TGuardValue TickGuard(bIsTicking, true);
 
-	if (!bWorldPlaying)
+	TRACE_CPUPROFILER_EVENT_SCOPE(UAvaGameInstance::EditorTick);
+	PlayWorld->Tick(ELevelTick::LEVELTICK_All, InDeltaSeconds);
+
+	// Only update reflection captures in game once all 'always loaded' levels have been loaded
+	// This won't work with actual level streaming though
+	if (PlayWorld->AreAlwaysLoadedLevelsLoaded())
 	{
-		return;
+		// Update sky light first because it's considered direct lighting, sky diffuse will be visible in reflection capture indirect specular
+		USkyLightComponent::UpdateSkyCaptureContents(PlayWorld);
+		UReflectionCaptureComponent::UpdateReflectionCaptureContents(PlayWorld);
 	}
+}
+#endif // WITH_EDITOR
 
-	if (!PlayWorld)
+void UAvaGameInstance::TickInternal(float InDeltaSeconds)
+{
+	// Skip if world not playing, or it has already ticked this frame
+	if (!bWorldPlaying || TickedFrame == GFrameCounter)
 	{
-		EndPlayWorld();
 		return;
 	}
-	
-	/*
-	 * Things Removed Here because they're currently unneeded and will possibly never be
-	 * TickWorldTravel
-	 */
 
-	{
-		SCOPE_TIME_GUARD(TEXT("UAvaGameInstance::Tick - World Tick"));
-		PlayWorld->Tick(ELevelTick::LEVELTICK_All, DeltaSeconds);
-	}
+	TRACE_CPUPROFILER_EVENT_SCOPE(UAvaGameInstance::Tick);
+	TickedFrame = GFrameCounter;
 
-	//Removed check for Dedicated Server & Commandlet
+	if (!PlayWorld)
 	{
-		QUICK_SCOPE_CYCLE_COUNTER(STAT_UGameEngine_Tick_CheckCaptures);
-		
-		// Only update reflection captures in game once all 'always loaded' levels have been loaded
-		// This won't work with actual level streaming though
-		if (PlayWorld->AreAlwaysLoadedLevelsLoaded())
-		{
-			// Update sky light first because it's considered direct lighting, sky diffuse will be visible in reflection capture indirect specular
-			USkyLightComponent::UpdateSkyCaptureContents(PlayWorld.Get());
-			UReflectionCaptureComponent::UpdateReflectionCaptureContents(PlayWorld.Get());
-		}
+		RequestEndPlayWorld(/*bForceImmediate*/false);
+		return;
 	}
 
-	//Rest of the stuff between above Lighting Update and GameViewport Tick removed
+	ViewportClient->Tick(InDeltaSeconds);
+	Viewport->Draw();
 
+	if (!bIsRenderTargetReady)
 	{
-		SCOPE_TIME_GUARD(TEXT("UAvaGameInstance::Tick - TickViewport"));
-		QUICK_SCOPE_CYCLE_COUNTER(STAT_MotionDesignGame_Tick_TickViewport);
-		ViewportClient->Tick(DeltaSeconds);
-	}
-	
-	{
-		Viewport->Draw();
-
-		if (!bIsRenderTargetReady)
+		if (!RenderTargetFence.IsValid())
 		{
-			if (!RenderTargetFence.IsValid())
-			{
-				RenderTargetFence = MakeUnique<FRenderCommandFence>();
-				RenderTargetFence->BeginFence();
-			}
-			else
-			{
-				bIsRenderTargetReady = RenderTargetFence->IsFenceComplete();
-				OnRenderTargetReady.Broadcast(this, PlayingChannelName);
-			}
+			RenderTargetFence = MakeUnique<FRenderCommandFence>();
+			RenderTargetFence->BeginFence();
+		}
+		else
+		{
+			bIsRenderTargetReady = RenderTargetFence->IsFenceComplete();
+			OnRenderTargetReady.Broadcast(this, PlayingChannelName);
 		}
-	}
-
-	if (bRequestEndPlayWorld)
-	{
-		EndPlayWorld();
-	}
-	if (bRequestUnloadWorld)
-	{
-		UnloadWorld(/*bShutdown*/true);
 	}
 }
 
@@ -356,6 +321,16 @@ void UAvaGameInstance::EndPlayWorld()
 	bWorldPlaying = false;
 	bRequestEndPlayWorld = false;
 
+	SetTickableTickType(ETickableTickType::Never);
+
+#if WITH_EDITOR
+	if (GEditor)
+	{
+		GEditor->OnPostEditorTick().Remove(OnEditorTickHandle);
+		OnEditorTickHandle.Reset();
+	}
+#endif // WITH_EDITOR
+
 	FCoreDelegates::OnEndFrame.RemoveAll(this);
 
 	OnEndPlay.Broadcast(this, PlayingChannelName);
@@ -388,18 +363,14 @@ void UAvaGameInstance::EndPlayWorld()
 
 void UAvaGameInstance::OnEndFrameTick()
 {
-	double DeltaSeconds = FApp::GetDeltaTime();
-
-	if (UE::AvaGameInstance::Private::ShouldIgnoreDeltaSecondsForCurrentFrame())
+	if (bRequestEndPlayWorld)
 	{
-		DeltaSeconds = LastDeltaSeconds;
+		EndPlayWorld();
 	}
-	else
+	if (bRequestUnloadWorld)
 	{
-		LastDeltaSeconds = DeltaSeconds;
+		UnloadWorld(/*bShutdown*/true);
 	}
-	
-	Tick(DeltaSeconds);
 }
 
 void UAvaGameInstance::OnEnginePreExit()
@@ -416,6 +387,21 @@ void UAvaGameInstance::OnEnginePreExit()
 	}
 }
 
+void UAvaGameInstance::Tick(float InDeltaSeconds)
+{
+	TickInternal(InDeltaSeconds);
+}
+
+TStatId UAvaGameInstance::GetStatId() const
+{
+	RETURN_QUICK_DECLARE_CYCLE_STAT(UAvaGameInstance, STATGROUP_Tickables)
+}
+
+UWorld* UAvaGameInstance::GetTickableGameObjectWorld() const
+{
+	return PlayWorld;
+}
+
 ULocalPlayer* UAvaGameInstance::CreateInitialPlayer(FString& OutError)
 {
 	if (!GetGameViewportClient())
diff --git a/Engine/Plugins/VirtualProduction/Avalanche/Source/Avalanche/Public/Framework/AvaGameInstance.h b/Engine/Plugins/VirtualProduction/Avalanche/Source/Avalanche/Public/Framework/AvaGameInstance.h
--- a/Engine/Plugins/VirtualProduction/Avalanche/Source/Avalanche/Public/Framework/AvaGameInstance.h
+++ b/Engine/Plugins/VirtualProduction/Avalanche/Source/Avalanche/Public/Framework/AvaGameInstance.h
@@ -7,6 +7,8 @@
 #include "Engine/World.h"
 #include "AvaGameInstance.generated.h"
 
+#define UE_API AVALANCHE_API
+
 class FRenderCommandFence;
 class FSceneViewport;
 class UTextureRenderTarget2D;
@@ -22,26 +24,28 @@ struct FAvaInstancePlaySettings
 	const FAvaViewportQualitySettings& QualitySettings;
 };
 
-UCLASS(DisplayName = "Motion Design Game Instance")
-class AVALANCHE_API UAvaGameInstance : public UGameInstance
+UCLASS(MinimalAPI, DisplayName = "Motion Design Game Instance")
+class UAvaGameInstance : public UGameInstance, public FTickableGameObject
 {
 	GENERATED_BODY()
 
 public:
+	UE_API UAvaGameInstance();
+
 	/**
 	 * Creates a new Motion Design Game instance from the given Motion Design Template Asset.
 	 * Supported asset types: UWorld (Level).
 	 */
-	static UAvaGameInstance* Create(UObject* InOuter);
+	UE_API static UAvaGameInstance* Create(UObject* InOuter);
 
 	/** Create the game world */
-	bool CreateWorld();
+	UE_INTERNAL UE_API bool CreateWorld();
 
 	bool IsWorldCreated() const { return bWorldCreated; }
 
 	UWorld* GetPlayWorld() const { return PlayWorld.Get(); }
 
-	bool BeginPlayWorld(const FAvaInstancePlaySettings& InWorldPlaySettings);
+	UE_API bool BeginPlayWorld(const FAvaInstancePlaySettings& InWorldPlaySettings);
 
 	/**
 	 * The end play is normally requested to be done on the next Tick(). This is to
@@ -49,26 +53,26 @@ public:
 	 * However, when shutting down, we need to force the end play to be done immediately because
 	 * there will not a be a next Tick().
 	 */
-	void RequestEndPlayWorld(bool bForceImmediate);
+	UE_API void RequestEndPlayWorld(bool bForceImmediate);
 	
-	void RequestUnloadWorld(bool bForceImmediate);
+	UE_API void RequestUnloadWorld(bool bForceImmediate);
 
 	bool IsWorldPlaying() const { return bWorldPlaying; }
 
 	/**
 	 * Cancel any pending EndPlayWorld or UnloadWorld requests. 
 	 */
-	void CancelWorldRequests();
+	UE_API void CancelWorldRequests();
 
 	/**
 	* The output channels can be resized during playback, if this happens, the
 	* scene viewport needs to be updated to reflect the change.
 	*/
-	void UpdateSceneViewportSize(const FIntPoint& InViewportSize);
+	UE_API void UpdateSceneViewportSize(const FIntPoint& InViewportSize);
 
-	void UpdateRenderTarget(UTextureRenderTarget2D* InRenderTarget);
+	UE_API void UpdateRenderTarget(UTextureRenderTarget2D* InRenderTarget);
 
-	UTextureRenderTarget2D* GetRenderTarget() const;
+	UE_API UTextureRenderTarget2D* GetRenderTarget() const;
 
 	/**
 	 * Returns true if the render target has been rendered and is ready to be used for output.
@@ -101,19 +105,32 @@ public:
 	static FOnAvaGameInstanceEvent& GetOnRenderTargetReady() { return OnRenderTargetReady; }
 
 protected:
-	void Tick(float DeltaSeconds);
+#if WITH_EDITOR
+	void EditorTick(float InDeltaSeconds);
+#endif // WITH_EDITOR
+
+	void TickInternal(float InDeltaSeconds);
 
-	void UnloadWorld(bool bInShutdown);
-	void EndPlayWorld();
-	void OnEndFrameTick();
-	void OnEnginePreExit();
+	UE_API void UnloadWorld(bool bInShutdown);
+	UE_API void EndPlayWorld();
+	UE_API void OnEndFrameTick();
+	UE_API void OnEnginePreExit();
+
+	//~ Begin FTickableObjectBase
+	UE_API virtual void Tick(float InDeltaSeconds) override;
+	UE_API virtual TStatId GetStatId() const override;
+	//~ End FTickableObjectBase
+
+	//~ Begin FTickableGameObject
+	UE_API virtual UWorld* GetTickableGameObjectWorld() const override;
+	//~ End FTickableGameObject
 
 	//~ Begin UGameInstance
-	virtual ULocalPlayer* CreateInitialPlayer(FString& OutError) override;
+	UE_API virtual ULocalPlayer* CreateInitialPlayer(FString& OutError) override;
 	//~ End UGameInstance
 	
 	//~ Begin UObject
-	virtual void BeginDestroy() override;
+	UE_API virtual void BeginDestroy() override;
 	//~ End UObject
 
 	UPROPERTY(Transient)
@@ -134,19 +151,23 @@ protected:
 	bool bRequestUnloadWorld = false;
 	bool bWorldPlaying = false;
 	bool bRequestEndPlayWorld = false;
-
-	/** Tick reentrancy guard. */
-	bool bIsTicking = false;
-	
-	double LastDeltaSeconds = 0.0;
+	bool bIsRenderTargetReady = false;
 
 	/** Render command fence to ensure the render target has been rendered. */
 	TUniquePtr<FRenderCommandFence> RenderTargetFence;
-	bool bIsRenderTargetReady = false;
+
+	/** Last frame that was ticked */
+	uint64 TickedFrame = 0;
 
 	/** Handle to delegate for EnginePreExit event. */
 	FDelegateHandle EnginePreExitHandle;
+#if WITH_EDITOR
+	/** Handle to the UEditorEngine::OnPostEditorTick */
+	FDelegateHandle OnEditorTickHandle;
+#endif
 
-	static FOnAvaGameInstanceEvent OnEndPlay;
-	static FOnAvaGameInstanceEvent OnRenderTargetReady;
+	UE_API static FOnAvaGameInstanceEvent OnEndPlay;
+	UE_API static FOnAvaGameInstanceEvent OnRenderTargetReady;
 };
+
+#undef UE_API
