Motion Design: added option to disable python for local playback.
+ Added a preview string for the command line that will be used when executing local playback server for visibility

#Synced-CL 51441789

diff --git a/Engine/Plugins/VirtualProduction/Avalanche/Source/AvalancheMedia/Private/AvaMediaSettings.cpp b/Engine/Plugins/VirtualProduction/Avalanche/Source/AvalancheMedia/Private/AvaMediaSettings.cpp
--- a/Engine/Plugins/VirtualProduction/Avalanche/Source/AvalancheMedia/Private/AvaMediaSettings.cpp
+++ b/Engine/Plugins/VirtualProduction/Avalanche/Source/AvalancheMedia/Private/AvaMediaSettings.cpp
@@ -3,6 +3,7 @@
 #include "AvaMediaSettings.h"
 #include "IAvaMediaModule.h"
 #include "Misc/MessageDialog.h"
+#include "UObject/UnrealType.h"
 
 #define LOCTEXT_NAMESPACE "AvaMediaSettings"
 
@@ -55,15 +56,76 @@ ELogVerbosity::Type UAvaMediaSettings::ToLogVerbosity(EAvaMediaLogVerbosity InAv
 	}
 }
 
+FString UAvaMediaSettings::GenerateLocalPlaybackServerCommandLine() const
+{
+	const FAvaMediaLocalPlaybackServerSettings& Settings = LocalPlaybackServerSettings;
+
+	FString CommandLine;
+	CommandLine += FString::Printf(TEXT("-game -windowed -ResX=%d -ResY=%d -messaging"), Settings.Resolution.X, Settings.Resolution.Y);
+	if (Settings.bEnableLogConsole)
+	{
+		CommandLine += TEXT(" -log");
+	}
+	if (Settings.bDisablePython)
+	{
+		CommandLine += TEXT(" -DisablePython");
+	}
+
+	// Motion Design arguments
+	CommandLine += FString::Printf(TEXT(" -MotionDesignPlaybackServerStart=\"%s\""), *Settings.ServerName);
+	CommandLine += TEXT(" -MotionDesignPlaybackClientSuppress");
+	CommandLine += FString::Printf(TEXT(" -MotionDesignPlaybackServerLogReplication=%s"), ToString(UAvaMediaSettings::ToLogVerbosity(PlaybackServerLogReplicationVerbosity)));
+
+	// Storm Sync arguments
+	CommandLine += TEXT(" -NoStormSyncServerAutoStart");
+
+	if (!Settings.ExtraCommandLineArguments.IsEmpty())
+	{
+		CommandLine += TEXT(" ");
+		CommandLine += Settings.ExtraCommandLineArguments;
+	}
+
+	FString LogCommands;
+	for (const FAvaPlaybackServerLoggingEntry& LoggingEntry : Settings.Logging)
+	{
+		if (!LoggingEntry.Category.IsNone())
+		{
+			LogCommands += FString::Printf(TEXT("%s%s %s"),
+				LogCommands.IsEmpty() ? TEXT("") : TEXT(", "), *LoggingEntry.Category.ToString(),
+				ToString(UAvaMediaSettings::ToLogVerbosity(LoggingEntry.VerbosityLevel)));
+		}
+	}
+
+	if (!LogCommands.IsEmpty())
+	{
+		CommandLine += FString::Printf(TEXT(" -LogCmds=\"%s\""), *LogCommands);
+	}
+	return CommandLine;
+}
+
+void UAvaMediaSettings::PostInitProperties()
+{
+	Super::PostInitProperties();
+#if WITH_EDITOR
+	UpdateLocalPlaybackServerCommandLinePreview();
+#endif
+}
+
+void UAvaMediaSettings::PostReloadConfig(FProperty* InPropertyThatWasLoaded)
+{
+	Super::PostReloadConfig(InPropertyThatWasLoaded);
+#if WITH_EDITOR
+	UpdateLocalPlaybackServerCommandLinePreview();
+#endif
+}
+
 #if WITH_EDITOR
 void UAvaMediaSettings::PostEditChangeProperty(FPropertyChangedEvent& InPropertyChangedEvent)
 {
 	Super::PostEditChangeProperty(InPropertyChangedEvent);
 
 	// Check if playback client is auto-start.
-	static const FName AutoStartPlaybackClientPropertyName = GET_MEMBER_NAME_CHECKED(UAvaMediaSettings, bAutoStartPlaybackClient);
-
-	if (InPropertyChangedEvent.MemberProperty->GetFName() == AutoStartPlaybackClientPropertyName)
+	if (InPropertyChangedEvent.MemberProperty->GetFName() == GET_MEMBER_NAME_CHECKED(UAvaMediaSettings, bAutoStartPlaybackClient))
 	{
 		IAvaMediaModule& AvaMediaModule = IAvaMediaModule::Get();
 		if (bAutoStartPlaybackClient && !AvaMediaModule.IsPlaybackClientStarted())
@@ -87,7 +149,17 @@ void UAvaMediaSettings::PostEditChangeProperty(FPropertyChangedEvent& InProperty
 			}
 		}
 	}
+	// Update user-facing command line string if server settings has changed 
+	else if (InPropertyChangedEvent.MemberProperty->GetFName() == GET_MEMBER_NAME_CHECKED(UAvaMediaSettings, LocalPlaybackServerSettings))
+	{
+		UpdateLocalPlaybackServerCommandLinePreview();
+	}
 }
-#endif
+
+void UAvaMediaSettings::UpdateLocalPlaybackServerCommandLinePreview()
+{
+	LocalPlaybackServerCommandLinePreview = GenerateLocalPlaybackServerCommandLine();
+}
+#endif // WITH_EDITOR
 
 #undef LOCTEXT_NAMESPACE
diff --git a/Engine/Plugins/VirtualProduction/Avalanche/Source/AvalancheMedia/Private/Playback/AvaPlaybackServerProcess.cpp b/Engine/Plugins/VirtualProduction/Avalanche/Source/AvalancheMedia/Private/Playback/AvaPlaybackServerProcess.cpp
--- a/Engine/Plugins/VirtualProduction/Avalanche/Source/AvalancheMedia/Private/Playback/AvaPlaybackServerProcess.cpp
+++ b/Engine/Plugins/VirtualProduction/Avalanche/Source/AvalancheMedia/Private/Playback/AvaPlaybackServerProcess.cpp
@@ -51,48 +51,10 @@ bool FAvaPlaybackServerProcess::Launch()
 		GameNameOrProjectFile = FApp::GetProjectName();
 	}
 
-	FString CommandLine;
-	CommandLine += GameNameOrProjectFile;
-
 	// Todo: It will likely need a more complete set of parameters. Reference: UDisplayClusterLaunchEditorProjectSettings.
-	const FAvaMediaLocalPlaybackServerSettings& Settings = UAvaMediaSettings::Get().LocalPlaybackServerSettings;
-
-	// Standard arguments
-	const FIntPoint Resolution = Settings.Resolution;
-	CommandLine += FString::Printf(TEXT(" -game -windowed -ResX=%d -ResY=%d -messaging"), Resolution.X, Resolution.Y);
-	if (Settings.bEnableLogConsole)
-	{
-		CommandLine += TEXT(" -log");
-	}
-	// Motion Design arguments
-	CommandLine += FString::Printf(TEXT(" -MotionDesignPlaybackServerStart=%s"), *GetLocalServerName());
-	CommandLine += TEXT(" -MotionDesignPlaybackClientSuppress");
-	CommandLine += FString::Printf(TEXT(" -MotionDesignPlaybackServerLogReplication=%s"), GetPlaybackServerLogReplicationVerbosityString());
-	// Storm Sync arguments
-	CommandLine += TEXT(" -NoStormSyncServerAutoStart");
-
-	if (!Settings.ExtraCommandLineArguments.IsEmpty())
-	{
-		CommandLine += TEXT(" ");
-		CommandLine += Settings.ExtraCommandLineArguments;
-	}
-
-	FString LogCommands;
-	
-	for (const FAvaPlaybackServerLoggingEntry& LoggingEntry : Settings.Logging)
-	{
-		if (!LoggingEntry.Category.IsNone())
-		{
-			LogCommands += FString::Printf(TEXT("%s%s %s"),
-					LogCommands.IsEmpty() ? TEXT("") : TEXT(", "), *LoggingEntry.Category.ToString(),
-					ToString(UAvaMediaSettings::ToLogVerbosity(LoggingEntry.VerbosityLevel)));
-		}
-	}
-
-	if (!LogCommands.IsEmpty())
-	{
-		CommandLine += FString::Printf(TEXT(" -LogCmds=\"%s\""), *LogCommands);		
-	}
+	const FString CommandLine = FString::Printf(TEXT("%s %s")
+		, *GameNameOrProjectFile
+		, *UAvaMediaSettings::Get().GenerateLocalPlaybackServerCommandLine());
 
 	const FString ExecutablePath = FPlatformProcess::ExecutablePath();
 
diff --git a/Engine/Plugins/VirtualProduction/Avalanche/Source/AvalancheMedia/Public/AvaMediaSettings.h b/Engine/Plugins/VirtualProduction/Avalanche/Source/AvalancheMedia/Public/AvaMediaSettings.h
--- a/Engine/Plugins/VirtualProduction/Avalanche/Source/AvalancheMedia/Public/AvaMediaSettings.h
+++ b/Engine/Plugins/VirtualProduction/Avalanche/Source/AvalancheMedia/Public/AvaMediaSettings.h
@@ -18,6 +18,8 @@
 
 class UUserWidget;
 
+#define UE_API AVALANCHEMEDIA_API
+
 /**
  * Defines the verbosity level of the logging system.
  * This enum mirrors ELogVerbosity but can be used directly as a configuration property.
@@ -64,6 +66,13 @@ struct FAvaMediaLocalPlaybackServerSettings
 	UPROPERTY(config, EditAnywhere, Category = Settings)
 	bool bEnableLogConsole = false;
 
+	/**
+	 * Disables python with the -DisablePython command. 
+	 * This is recommended for performance, as Python garbage collection call could impact the overall performance of a garbage collection pass.
+	 */
+	UPROPERTY(config, EditAnywhere, DisplayName="Disable Python (Performance)", Category = Settings)
+	bool bDisablePython = true;
+
 	/** Extra command line arguments. */
 	UPROPERTY(config, EditAnywhere, Category = Settings)
 	FString ExtraCommandLineArguments;
@@ -75,25 +84,30 @@ struct FAvaMediaLocalPlaybackServerSettings
 	TArray<FAvaPlaybackServerLoggingEntry> Logging;
 };
 
-UCLASS(config=Engine, meta=(DisplayName="Playback & Broadcast"))
-class AVALANCHEMEDIA_API UAvaMediaSettings : public UDeveloperSettings
+UCLASS(MinimalAPI, config=Engine, meta=(DisplayName="Playback & Broadcast"))
+class UAvaMediaSettings : public UDeveloperSettings
 {
 	GENERATED_BODY()
 
 public:
-	UAvaMediaSettings();
-	
+	UE_API UAvaMediaSettings();
+
 	static const UAvaMediaSettings& Get() { return *GetSingletonInstance();}
 	static UAvaMediaSettings& GetMutable() {return *GetSingletonInstance();}
 
-	static ELogVerbosity::Type ToLogVerbosity(EAvaMediaLogVerbosity InAvaMediaLogVerbosity);
+	UE_API static ELogVerbosity::Type ToLogVerbosity(EAvaMediaLogVerbosity InAvaMediaLogVerbosity);
+
+	/** Generates the command line arguments for the local playback server */
+	UE_API FString GenerateLocalPlaybackServerCommandLine() const;
 
 	//~ Begin UObject
+	UE_API virtual void PostInitProperties() override;
+	UE_API virtual void PostReloadConfig(FProperty* InPropertyThatWasLoaded) override;
 #if WITH_EDITOR
-	virtual void PostEditChangeProperty(struct FPropertyChangedEvent& InPropertyChangedEvent) override;
+	UE_API virtual void PostEditChangeProperty(FPropertyChangedEvent& InPropertyChangedEvent) override;
 #endif
 	//~ End UObject
-	
+
 	/** Specifies the background clear color for the channel. */
 	UPROPERTY(config, EditAnywhere, Category = "Broadcast")
 	FLinearColor ChannelClearColor = FLinearColor::Black;
@@ -232,8 +246,21 @@ public:
 	uint32 HttpServerPort = 10123;
 
 	/** Default value of the synchronized events feature selection. */
-	static const FName SynchronizedEventsFeatureSelection_Default;
+	UE_API static const FName SynchronizedEventsFeatureSelection_Default;
 
 private:
-	static UAvaMediaSettings* GetSingletonInstance();
+	UE_API static UAvaMediaSettings* GetSingletonInstance();
+
+#if WITH_EDITOR
+	/** Updates the command line preview to reflect the latest local playback server settings */
+	void UpdateLocalPlaybackServerCommandLinePreview();
+#endif
+
+#if WITH_EDITORONLY_DATA
+	/** Preview command line arguments from the playback server settings */
+	UPROPERTY(Transient, VisibleAnywhere, Category = "Playback Server", meta=(MultiLine, DisplayAfter="LocalPlaybackServerSettings"))
+	FString LocalPlaybackServerCommandLinePreview;
+#endif
 };
+
+#undef UE_API
