Motion Design: fix game instance viewport incorrectly showing the level for a frame for a streamed level

#Synced-CL 51679831

diff --git a/Engine/Plugins/VirtualProduction/Avalanche/Source/Avalanche/Private/Framework/AvaGameInstance.cpp b/Engine/Plugins/VirtualProduction/Avalanche/Source/Avalanche/Private/Framework/AvaGameInstance.cpp
--- a/Engine/Plugins/VirtualProduction/Avalanche/Source/Avalanche/Private/Framework/AvaGameInstance.cpp
+++ b/Engine/Plugins/VirtualProduction/Avalanche/Source/Avalanche/Private/Framework/AvaGameInstance.cpp
@@ -22,11 +22,6 @@
 UAvaGameInstance::FOnAvaGameInstanceEvent UAvaGameInstance::OnEndPlay;
 UAvaGameInstance::FOnAvaGameInstanceEvent UAvaGameInstance::OnRenderTargetReady;
 
-UAvaGameInstance::UAvaGameInstance()
-	: FTickableGameObject(ETickableTickType::Never) // Register tickable when it is needed
-{
-}
-
 UAvaGameInstance* UAvaGameInstance::Create(UObject* InOuter)
 {
 	checkf(IsInGameThread(), TEXT("Motion Design Game Instance creation can only happen in Game Thread"));
@@ -149,9 +144,6 @@ bool UAvaGameInstance::BeginPlayWorld(const FAvaInstancePlaySettings& InWorldPla
 	}
 #endif // WITH_EDITOR
 
-	// Register game instance to tick
-	SetTickableTickType(ETickableTickType::Always);
-
 	PlayWorld->BeginPlay();
 
 	FCoreDelegates::OnEndFrame.AddUObject(this, &UAvaGameInstance::OnEndFrameTick);
@@ -247,41 +239,6 @@ void UAvaGameInstance::EditorTick(float InDeltaSeconds)
 }
 #endif // WITH_EDITOR
 
-void UAvaGameInstance::TickInternal(float InDeltaSeconds)
-{
-	// Skip if world not playing, or it has already ticked this frame
-	if (!bWorldPlaying || TickedFrame == GFrameCounter)
-	{
-		return;
-	}
-
-	TRACE_CPUPROFILER_EVENT_SCOPE(UAvaGameInstance::Tick);
-	TickedFrame = GFrameCounter;
-
-	if (!PlayWorld)
-	{
-		RequestEndPlayWorld(/*bForceImmediate*/false);
-		return;
-	}
-
-	ViewportClient->Tick(InDeltaSeconds);
-	Viewport->Draw();
-
-	if (!bIsRenderTargetReady)
-	{
-		if (!RenderTargetFence.IsValid())
-		{
-			RenderTargetFence = MakeUnique<FRenderCommandFence>();
-			RenderTargetFence->BeginFence();
-		}
-		else
-		{
-			bIsRenderTargetReady = RenderTargetFence->IsFenceComplete();
-			OnRenderTargetReady.Broadcast(this, PlayingChannelName);
-		}
-	}
-}
-
 void UAvaGameInstance::UnloadWorld(bool bInShutdown)
 {
 	if (bWorldPlaying)
@@ -318,11 +275,11 @@ void UAvaGameInstance::UnloadWorld(bool bInShutdown)
 
 void UAvaGameInstance::EndPlayWorld()
 {
+	TRACE_CPUPROFILER_EVENT_SCOPE(UAvaGameInstance::EndPlayWorld);
+
 	bWorldPlaying = false;
 	bRequestEndPlayWorld = false;
 
-	SetTickableTickType(ETickableTickType::Never);
-
 #if WITH_EDITOR
 	if (GEditor)
 	{
@@ -363,6 +320,37 @@ void UAvaGameInstance::EndPlayWorld()
 
 void UAvaGameInstance::OnEndFrameTick()
 {
+	// Skip if world not playing, or it has already ticked this frame
+	if (!bWorldPlaying)
+	{
+		return;
+	}
+
+	if (!PlayWorld)
+	{
+		RequestEndPlayWorld(/*bForceImmediate*/false);
+		return;
+	}
+
+	TRACE_CPUPROFILER_EVENT_SCOPE(UAvaGameInstance::OnEndFrameTick);
+
+	const double DeltaTime = FAppTime::GetDeltaTime();
+	ViewportClient->Tick(DeltaTime);
+	Viewport->Draw();
+
+	if (!bIsRenderTargetReady)
+	{
+		if (!RenderTargetFence.IsValid())
+		{
+			RenderTargetFence = MakeUnique<FRenderCommandFence>();
+			RenderTargetFence->BeginFence();
+		}
+		else
+		{
+			bIsRenderTargetReady = RenderTargetFence->IsFenceComplete();
+			OnRenderTargetReady.Broadcast(this, PlayingChannelName);
+		}
+	}
 	if (bRequestEndPlayWorld)
 	{
 		EndPlayWorld();
@@ -387,21 +375,6 @@ void UAvaGameInstance::OnEnginePreExit()
 	}
 }
 
-void UAvaGameInstance::Tick(float InDeltaSeconds)
-{
-	TickInternal(InDeltaSeconds);
-}
-
-TStatId UAvaGameInstance::GetStatId() const
-{
-	RETURN_QUICK_DECLARE_CYCLE_STAT(UAvaGameInstance, STATGROUP_Tickables)
-}
-
-UWorld* UAvaGameInstance::GetTickableGameObjectWorld() const
-{
-	return PlayWorld;
-}
-
 ULocalPlayer* UAvaGameInstance::CreateInitialPlayer(FString& OutError)
 {
 	if (!GetGameViewportClient())
diff --git a/Engine/Plugins/VirtualProduction/Avalanche/Source/Avalanche/Public/Framework/AvaGameInstance.h b/Engine/Plugins/VirtualProduction/Avalanche/Source/Avalanche/Public/Framework/AvaGameInstance.h
--- a/Engine/Plugins/VirtualProduction/Avalanche/Source/Avalanche/Public/Framework/AvaGameInstance.h
+++ b/Engine/Plugins/VirtualProduction/Avalanche/Source/Avalanche/Public/Framework/AvaGameInstance.h
@@ -25,13 +25,11 @@ struct FAvaInstancePlaySettings
 };
 
 UCLASS(MinimalAPI, DisplayName = "Motion Design Game Instance")
-class UAvaGameInstance : public UGameInstance, public FTickableGameObject
+class UAvaGameInstance : public UGameInstance
 {
 	GENERATED_BODY()
 
 public:
-	UE_API UAvaGameInstance();
-
 	/**
 	 * Creates a new Motion Design Game instance from the given Motion Design Template Asset.
 	 * Supported asset types: UWorld (Level).
@@ -116,15 +114,6 @@ protected:
 	UE_API void OnEndFrameTick();
 	UE_API void OnEnginePreExit();
 
-	//~ Begin FTickableObjectBase
-	UE_API virtual void Tick(float InDeltaSeconds) override;
-	UE_API virtual TStatId GetStatId() const override;
-	//~ End FTickableObjectBase
-
-	//~ Begin FTickableGameObject
-	UE_API virtual UWorld* GetTickableGameObjectWorld() const override;
-	//~ End FTickableGameObject
-
 	//~ Begin UGameInstance
 	UE_API virtual ULocalPlayer* CreateInitialPlayer(FString& OutError) override;
 	//~ End UGameInstance
@@ -156,9 +145,6 @@ protected:
 	/** Render command fence to ensure the render target has been rendered. */
 	TUniquePtr<FRenderCommandFence> RenderTargetFence;
 
-	/** Last frame that was ticked */
-	uint64 TickedFrame = 0;
-
 	/** Handle to delegate for EnginePreExit event. */
 	FDelegateHandle EnginePreExitHandle;
 #if WITH_EDITOR
