Motion Design: fixed issue with transition logic running after game viewport draw. This was causing a frame to be visible before desired.

#Synced-CL 51679831

diff --git a/Engine/Plugins/VirtualProduction/Avalanche/Source/Avalanche/Private/Framework/AvaGameInstance.cpp b/Engine/Plugins/VirtualProduction/Avalanche/Source/Avalanche/Private/Framework/AvaGameInstance.cpp
--- a/Engine/Plugins/VirtualProduction/Avalanche/Source/Avalanche/Private/Framework/AvaGameInstance.cpp
+++ b/Engine/Plugins/VirtualProduction/Avalanche/Source/Avalanche/Private/Framework/AvaGameInstance.cpp
@@ -397,9 +397,15 @@ TStatId UAvaGameInstance::GetStatId() const
 	RETURN_QUICK_DECLARE_CYCLE_STAT(UAvaGameInstance, STATGROUP_Tickables)
 }
 
+bool UAvaGameInstance::IsTickableInEditor() const
+{
+	return true;
+}
+
 UWorld* UAvaGameInstance::GetTickableGameObjectWorld() const
 {
-	return PlayWorld;
+	// Explicitly set this to null as the game instance should tick outside of the PlayWorld ticking.
+	return nullptr;
 }
 
 ULocalPlayer* UAvaGameInstance::CreateInitialPlayer(FString& OutError)
diff --git a/Engine/Plugins/VirtualProduction/Avalanche/Source/Avalanche/Public/Framework/AvaGameInstance.h b/Engine/Plugins/VirtualProduction/Avalanche/Source/Avalanche/Public/Framework/AvaGameInstance.h
--- a/Engine/Plugins/VirtualProduction/Avalanche/Source/Avalanche/Public/Framework/AvaGameInstance.h
+++ b/Engine/Plugins/VirtualProduction/Avalanche/Source/Avalanche/Public/Framework/AvaGameInstance.h
@@ -122,6 +122,7 @@ protected:
 	//~ End FTickableObjectBase
 
 	//~ Begin FTickableGameObject
+	UE_API virtual bool IsTickableInEditor() const override;
 	UE_API virtual UWorld* GetTickableGameObjectWorld() const override;
 	//~ End FTickableGameObject
 
diff --git a/Engine/Plugins/VirtualProduction/Avalanche/Source/AvalancheTransition/Private/Execution/AvaTransitionExecutor.cpp b/Engine/Plugins/VirtualProduction/Avalanche/Source/AvalancheTransition/Private/Execution/AvaTransitionExecutor.cpp
--- a/Engine/Plugins/VirtualProduction/Avalanche/Source/AvalancheTransition/Private/Execution/AvaTransitionExecutor.cpp
+++ b/Engine/Plugins/VirtualProduction/Avalanche/Source/AvalancheTransition/Private/Execution/AvaTransitionExecutor.cpp
@@ -8,6 +8,8 @@
 #include "AvaTransitionSubsystem.h"
 #include "AvaTransitionTree.h"
 #include "Behavior/IAvaTransitionBehavior.h"
+#include "Engine/Engine.h"
+#include "Engine/Level.h"
 #include "Engine/World.h"
 #include "Execution/AvaTransitionExecutorBuilder.h"
 #include "Misc/EnumerateRange.h"
@@ -50,6 +52,8 @@ void FAvaTransitionExecutor::Setup()
 			InInstance.Setup();
 		});
 
+	DetermineTickWorld();
+
 	struct FLayerInfo
 	{
 		// Enter Instances found for a given Layer
@@ -147,6 +151,53 @@ void FAvaTransitionExecutor::Setup()
 	}
 }
 
+void FAvaTransitionExecutor::DetermineTickWorld()
+{
+	WorldWeak.Reset();
+
+	TArray<UWorld*, TInlineAllocator<1>> CandidateWorlds;
+
+	// Gather all candidate worlds from the instances
+	ForEachInstance(
+		[&CandidateWorlds](FAvaTransitionBehaviorInstance& InInstance)
+		{
+			const FAvaTransitionScene* const TransitionScene = InInstance.GetTransitionContext().GetTransitionScene();
+			if (!TransitionScene)
+			{
+				return;
+			}
+			const ULevel* const Level = TransitionScene->GetLevel();
+			if (!Level)
+			{
+				return;
+			}
+			if (UWorld* const World = Level->GetWorld())
+			{
+				CandidateWorlds.AddUnique(World);
+			}
+		});
+
+	if (CandidateWorlds.IsEmpty())
+	{
+		return;
+	}
+
+	if (GEngine)
+	{
+		for (const FWorldContext& WorldContext : GEngine->GetWorldContexts())
+		{
+			if (CandidateWorlds.Contains(WorldContext.World()))
+			{
+				WorldWeak = WorldContext.World();
+				return;
+			}
+		}
+	}
+
+	// Fallback, no world in the list was found, use the first candidate.
+	WorldWeak = CandidateWorlds[0];
+}
+
 void FAvaTransitionExecutor::Start()
 {
 	if (!ensureAlways(!IsRunning()))
@@ -236,6 +287,11 @@ bool FAvaTransitionExecutor::IsTickable() const
 	return IsRunning();
 }
 
+UWorld* FAvaTransitionExecutor::GetTickableGameObjectWorld() const
+{
+	return WorldWeak.Get();
+}
+
 FString FAvaTransitionExecutor::GetReferencerName() const
 {
 	return TEXT("FAvaTransitionExecutor");
diff --git a/Engine/Plugins/VirtualProduction/Avalanche/Source/AvalancheTransition/Private/Execution/AvaTransitionExecutor.h b/Engine/Plugins/VirtualProduction/Avalanche/Source/AvalancheTransition/Private/Execution/AvaTransitionExecutor.h
--- a/Engine/Plugins/VirtualProduction/Avalanche/Source/AvalancheTransition/Private/Execution/AvaTransitionExecutor.h
+++ b/Engine/Plugins/VirtualProduction/Avalanche/Source/AvalancheTransition/Private/Execution/AvaTransitionExecutor.h
@@ -28,6 +28,9 @@ public:
 protected:
 	void Setup();
 
+	/** Gets the world to tick with from the behavior instances */
+	void DetermineTickWorld();
+
 	//~ Begin IAvaTransitionExecutor
 	virtual TArray<const FAvaTransitionBehaviorInstance*> GetBehaviorInstances(const FAvaTransitionLayerComparator& InComparator) const;
 	virtual void ForEachBehaviorInstance(TFunctionRef<void(const FAvaTransitionBehaviorInstance&)> InCallable) const override;
@@ -41,6 +44,7 @@ protected:
 	virtual bool IsTickable() const override;
 	virtual bool IsTickableInEditor() const override { return true; }
 	virtual ETickableTickType GetTickableTickType() const override { return ETickableTickType::Conditional; }
+	virtual UWorld* GetTickableGameObjectWorld() const override;
 	//~ End FTickableGameObject
 
 	//~ Begin FGCObject
@@ -54,6 +58,12 @@ private:
 
 	void ConditionallyFinishBehaviors();
 
+	/** 
+	 * First valid world context found.
+	 * The transition will tick with this World Tick if valid.
+	 */
+	TWeakObjectPtr<UWorld> WorldWeak;
+
 	TArray<FAvaTransitionBehaviorInstance> Instances;
 
 	const FAvaTransitionBehaviorInstance NullInstance;
