Motion Design: moved usage of core ticker to custom ts ticker instance to better control timing.

#Synced-CL 51715913

diff --git a/Engine/Plugins/VirtualProduction/ActorModifier/Source/ActorModifier/Private/Extensions/ActorModifierRenderStateUpdateExtension.cpp b/Engine/Plugins/VirtualProduction/ActorModifier/Source/ActorModifier/Private/Extensions/ActorModifierRenderStateUpdateExtension.cpp
--- a/Engine/Plugins/VirtualProduction/ActorModifier/Source/ActorModifier/Private/Extensions/ActorModifierRenderStateUpdateExtension.cpp
+++ b/Engine/Plugins/VirtualProduction/ActorModifier/Source/ActorModifier/Private/Extensions/ActorModifierRenderStateUpdateExtension.cpp
@@ -5,6 +5,7 @@
 #include "Components/PrimitiveComponent.h"
 #include "Containers/Ticker.h"
 #include "Modifiers/ActorModifierCoreBase.h"
+#include "Subsystems/ActorModifierCoreSubsystem.h"
 #include "Utilities/ActorModifierActorUtils.h"
 
 FActorModifierRenderStateUpdateExtension::FActorModifierRenderStateUpdateExtension(IActorModifierRenderStateUpdateHandler* InExtensionHandler)
@@ -126,7 +127,7 @@ void FActorModifierRenderStateUpdateExtension::OnRenderStateDirty(UActorComponen
 
 	// Execute on next tick otherwise visibility data might not be up to date
 	TWeakObjectPtr<AActor> ActorDirtyWeak(ActorDirty);
-	FTSTicker::GetCoreTicker().AddTicker(FTickerDelegate::CreateSPLambda(this, [this, ActorDirtyWeak](float InDeltaSeconds)->bool
+	UActorModifierCoreSubsystem::GetModifierTicker().AddTicker(FTickerDelegate::CreateSPLambda(this, [this, ActorDirtyWeak](float InDeltaSeconds)->bool
 	{
 		AActor* ActorDirty = ActorDirtyWeak.Get();
 
diff --git a/Engine/Plugins/VirtualProduction/ActorModifier/Source/ActorModifier/Private/Extensions/ActorModifierSceneTreeUpdateExtension.cpp b/Engine/Plugins/VirtualProduction/ActorModifier/Source/ActorModifier/Private/Extensions/ActorModifierSceneTreeUpdateExtension.cpp
--- a/Engine/Plugins/VirtualProduction/ActorModifier/Source/ActorModifier/Private/Extensions/ActorModifierSceneTreeUpdateExtension.cpp
+++ b/Engine/Plugins/VirtualProduction/ActorModifier/Source/ActorModifier/Private/Extensions/ActorModifierSceneTreeUpdateExtension.cpp
@@ -6,6 +6,7 @@
 #include "Engine/Level.h"
 #include "Engine/World.h"
 #include "GameFramework/Actor.h"
+#include "Subsystems/ActorModifierCoreSubsystem.h"
 #include "Utilities/ActorModifierActorUtils.h"
 
 FActorModifierSceneTreeUpdateExtension::FOnGetSceneTreeResolver FActorModifierSceneTreeUpdateExtension::OnGetSceneTreeResolverDelegate;
@@ -197,7 +198,7 @@ void FActorModifierSceneTreeUpdateExtension::OnWorldActorDestroyed(AActor* InAct
 	}
 
 	// Delay check from one tick to make sure actor is no longer attached
-	FTSTicker::GetCoreTicker().AddTicker(FTickerDelegate::CreateSPLambda(this, [this](float InDeltaSeconds)->bool
+	UActorModifierCoreSubsystem::GetModifierTicker().AddTicker(FTickerDelegate::CreateSPLambda(this, [this](float InDeltaSeconds)->bool
 	{
 		CheckTrackedActorsUpdate();
 		return false;
diff --git a/Engine/Plugins/VirtualProduction/ActorModifierCore/Source/ActorModifierCore/Private/Modifiers/ActorModifierCoreBase.cpp b/Engine/Plugins/VirtualProduction/ActorModifierCore/Source/ActorModifierCore/Private/Modifiers/ActorModifierCoreBase.cpp
--- a/Engine/Plugins/VirtualProduction/ActorModifierCore/Source/ActorModifierCore/Private/Modifiers/ActorModifierCoreBase.cpp
+++ b/Engine/Plugins/VirtualProduction/ActorModifierCore/Source/ActorModifierCore/Private/Modifiers/ActorModifierCoreBase.cpp
@@ -351,14 +351,12 @@ bool UActorModifierCoreBase::MarkModifierDirty(bool bExecute)
 
 bool UActorModifierCoreBase::IsModifierEnabled() const
 {
-	if (const UActorModifierCoreStack* Stack = GetModifierStack())
+	if (!bModifierEnabled)
 	{
-		if (!Stack->IsModifierEnabled())
-		{
-			return false;
-		}
+		return false;
 	}
-	return bModifierEnabled;
+	const UActorModifierCoreStack* Stack = GetModifierStack();
+	return !Stack || Stack->IsModifierEnabled();
 }
 
 void UActorModifierCoreBase::ProcessLockFunction(TFunctionRef<void()> InFunction)
diff --git a/Engine/Plugins/VirtualProduction/ActorModifierCore/Source/ActorModifierCore/Private/Subsystems/ActorModifierCoreSubsystem.cpp b/Engine/Plugins/VirtualProduction/ActorModifierCore/Source/ActorModifierCore/Private/Subsystems/ActorModifierCoreSubsystem.cpp
--- a/Engine/Plugins/VirtualProduction/ActorModifierCore/Source/ActorModifierCore/Private/Subsystems/ActorModifierCoreSubsystem.cpp
+++ b/Engine/Plugins/VirtualProduction/ActorModifierCore/Source/ActorModifierCore/Private/Subsystems/ActorModifierCoreSubsystem.cpp
@@ -29,6 +29,7 @@ UActorModifierCoreSubsystem::FOnModifierClassRegistered UActorModifierCoreSubsys
 UActorModifierCoreSubsystem::FOnModifierStackRegistered UActorModifierCoreSubsystem::OnModifierStackRegisteredDelegate;
 UActorModifierCoreSubsystem::FOnModifierStackRegistered UActorModifierCoreSubsystem::OnModifierStackUnregisteredDelegate;
 UActorModifierCoreSubsystem::FOnModifierReplaced UActorModifierCoreSubsystem::OnModifierReplacedDelegate;
+FTSTicker UActorModifierCoreSubsystem::ModifierTicker;
 
 UActorModifierCoreSubsystem::UActorModifierCoreSubsystem()
 	: UEngineSubsystem()
@@ -44,6 +45,11 @@ UActorModifierCoreSubsystem* UActorModifierCoreSubsystem::Get()
 	return nullptr;
 }
 
+FTSTicker& UActorModifierCoreSubsystem::GetModifierTicker()
+{
+	return ModifierTicker;
+}
+
 void UActorModifierCoreSubsystem::Initialize(FSubsystemCollectionBase& Collection)
 {
 	Super::Initialize(Collection);
@@ -65,6 +71,8 @@ void UActorModifierCoreSubsystem::Initialize(FSubsystemCollectionBase& Collectio
 
 void UActorModifierCoreSubsystem::Deinitialize()
 {
+	ModifierTicker.Reset();
+
 	// create a copy to avoid issue with iterator and removal
 	TArray<FName> ModifierNames;
 	ModifiersMetadata.GetKeys(ModifierNames);
@@ -538,6 +546,27 @@ void UActorModifierCoreSubsystem::GetSortedModifiers(const TSet<UActorModifierCo
 	});
 }
 
+ETickableTickType UActorModifierCoreSubsystem::GetTickableTickType() const
+{
+	return ETickableTickType::Always;
+}
+
+bool UActorModifierCoreSubsystem::IsTickableInEditor() const
+{
+	return true;
+}
+
+void UActorModifierCoreSubsystem::Tick(float InDeltaTime)
+{
+	TRACE_CPUPROFILER_EVENT_SCOPE(UActorModifierCoreSubsystem::Tick);
+	ModifierTicker.Tick(InDeltaTime);
+}
+
+TStatId UActorModifierCoreSubsystem::GetStatId() const
+{
+	RETURN_QUICK_DECLARE_CYCLE_STAT(UActorModifierCoreSubsystem, STATGROUP_Tickables);
+}
+
 bool UActorModifierCoreSubsystem::MoveModifiers(const TArray<UActorModifierCoreBase*>& InModifiers, UActorModifierCoreStack* InStack, FActorModifierCoreStackMoveOp& InMoveOp) const
 {
 	if (!IsValid(InStack) || InModifiers.IsEmpty())
diff --git a/Engine/Plugins/VirtualProduction/ActorModifierCore/Source/ActorModifierCore/Public/Subsystems/ActorModifierCoreSubsystem.h b/Engine/Plugins/VirtualProduction/ActorModifierCore/Source/ActorModifierCore/Public/Subsystems/ActorModifierCoreSubsystem.h
--- a/Engine/Plugins/VirtualProduction/ActorModifierCore/Source/ActorModifierCore/Public/Subsystems/ActorModifierCoreSubsystem.h
+++ b/Engine/Plugins/VirtualProduction/ActorModifierCore/Source/ActorModifierCore/Public/Subsystems/ActorModifierCoreSubsystem.h
@@ -2,8 +2,10 @@
 
 #pragma once
 
+#include "Containers/Ticker.h"
 #include "Modifiers/ActorModifierCoreDefs.h"
 #include "Subsystems/EngineSubsystem.h"
+#include "Tickable.h"
 #include "ActorModifierCoreSubsystem.generated.h"
 
 class AActor;
@@ -17,7 +19,7 @@ struct FAssetData;
 
 /** This subsystem handle all modifiers stack active in the engine and allows to create modifiers with registered metadata */
 UCLASS()
-class UActorModifierCoreSubsystem : public UEngineSubsystem
+class UActorModifierCoreSubsystem : public UEngineSubsystem, public FTickableGameObject
 {
 	GENERATED_BODY()
 
@@ -50,6 +52,9 @@ public:
 
 	ACTORMODIFIERCORE_API static UActorModifierCoreSubsystem* Get();
 
+	/** Gets the ticker for modifiers to use instead of CoreTicker */
+	ACTORMODIFIERCORE_API static FTSTicker& GetModifierTicker();
+
 	UActorModifierCoreSubsystem();
 
 	/** Register a modifier class or override an already existing one */
@@ -155,6 +160,16 @@ protected:
 	/** Called when a blueprint modifier gets replaced */
 	static FOnModifierReplaced OnModifierReplacedDelegate;
 
+	/** Ticker for modifiers to use instead of CoreTicker */
+	static FTSTicker ModifierTicker;
+
+	//~ Begin FTickableObjectBase
+	virtual ETickableTickType GetTickableTickType() const override;
+	virtual bool IsTickableInEditor() const override;
+	virtual void Tick(float InDeltaTime) override;
+	virtual TStatId GetStatId() const override;
+	//~ End FTickableObjectBase
+
 	//~ Begin UEngineSubsystem
 	virtual void Initialize(FSubsystemCollectionBase& Collection) override;
 	virtual void Deinitialize() override;
diff --git a/Engine/Plugins/VirtualProduction/Avalanche/Source/Avalanche/Private/AvaSceneSubsystem.cpp b/Engine/Plugins/VirtualProduction/Avalanche/Source/Avalanche/Private/AvaSceneSubsystem.cpp
--- a/Engine/Plugins/VirtualProduction/Avalanche/Source/Avalanche/Private/AvaSceneSubsystem.cpp
+++ b/Engine/Plugins/VirtualProduction/Avalanche/Source/Avalanche/Private/AvaSceneSubsystem.cpp
@@ -23,6 +23,11 @@ void UAvaSceneSubsystem::RegisterSceneInterface(ULevel* InLevel, IAvaSceneInterf
 	SceneInterfaces.Add(InLevel, InSceneInterface);
 }
 
+FTSTicker& UAvaSceneSubsystem::GetTicker()
+{
+	return Ticker;
+}
+
 IAvaSceneInterface* UAvaSceneSubsystem::GetSceneInterface() const
 {
 	if (UWorld* const World = GetWorld())
@@ -75,3 +80,19 @@ bool UAvaSceneSubsystem::DoesSupportWorldType(const EWorldType::Type InWorldType
 	const bool bDisallowedWorld = UE::Ava::Private::GUnsupportedWorlds.Contains(InWorldType);
 	return !bDisallowedWorld;
 }
+
+bool UAvaSceneSubsystem::IsTickableInEditor() const
+{
+	return true;
+}
+
+void UAvaSceneSubsystem::Tick(float InDeltaTime)
+{
+	TRACE_CPUPROFILER_EVENT_SCOPE(UAvaSceneSubsystem::Tick);
+	Ticker.Tick(InDeltaTime);
+}
+
+TStatId UAvaSceneSubsystem::GetStatId() const
+{
+	RETURN_QUICK_DECLARE_CYCLE_STAT(UAvaSceneSubsystem, STATGROUP_Tickables);
+}
diff --git a/Engine/Plugins/VirtualProduction/Avalanche/Source/Avalanche/Public/AvaSceneSubsystem.h b/Engine/Plugins/VirtualProduction/Avalanche/Source/Avalanche/Public/AvaSceneSubsystem.h
--- a/Engine/Plugins/VirtualProduction/Avalanche/Source/Avalanche/Public/AvaSceneSubsystem.h
+++ b/Engine/Plugins/VirtualProduction/Avalanche/Source/Avalanche/Public/AvaSceneSubsystem.h
@@ -3,6 +3,7 @@
 #pragma once
 
 #include "Containers/Map.h"
+#include "Containers/Ticker.h"
 #include "Subsystems/WorldSubsystem.h"
 #include "UObject/ObjectPtr.h"
 #include "UObject/WeakInterfacePtr.h"
@@ -12,13 +13,15 @@ class IAvaSceneInterface;
 class ULevel;
 
 UCLASS(MinimalAPI, DisplayName = "Motion Design Scene Subsystem")
-class UAvaSceneSubsystem : public UWorldSubsystem
+class UAvaSceneSubsystem : public UTickableWorldSubsystem
 {
 	GENERATED_BODY()
 
 public:
 	void RegisterSceneInterface(ULevel* InLevel, IAvaSceneInterface* InSceneInterface);
 
+	AVALANCHE_API FTSTicker& GetTicker();
+
 	/** Gets the Scene Interface for the World's Persistent Level */
 	AVALANCHE_API IAvaSceneInterface* GetSceneInterface() const;
 
@@ -34,6 +37,15 @@ protected:
 	AVALANCHE_API virtual bool DoesSupportWorldType(const EWorldType::Type InWorldType) const override;
 	//~ End UWorldSubsystem
 
+	//~ Begin FTickableObjectBase
+	AVALANCHE_API virtual bool IsTickableInEditor() const override;
+	AVALANCHE_API virtual void Tick(float InDeltaTime) override;
+	AVALANCHE_API virtual TStatId GetStatId() const override;
+	//~ End FTickableObjectBase
+
 private:
+	/** Ticker for motion design elements */
+	FTSTicker Ticker;
+
 	TMap<TWeakObjectPtr<ULevel>, TWeakInterfacePtr<IAvaSceneInterface>> SceneInterfaces;
 };
diff --git a/Engine/Plugins/VirtualProduction/Avalanche/Source/AvalancheModifiers/Private/Modifiers/AvaBooleanModifier.cpp b/Engine/Plugins/VirtualProduction/Avalanche/Source/AvalancheModifiers/Private/Modifiers/AvaBooleanModifier.cpp
--- a/Engine/Plugins/VirtualProduction/Avalanche/Source/AvalancheModifiers/Private/Modifiers/AvaBooleanModifier.cpp
+++ b/Engine/Plugins/VirtualProduction/Avalanche/Source/AvalancheModifiers/Private/Modifiers/AvaBooleanModifier.cpp
@@ -16,6 +16,7 @@
 #include "Profilers/AvaBooleanModifierProfiler.h"
 #include "Shared/AvaBooleanModifierShared.h"
 #include "Shared/ActorModifierVisibilityShared.h"
+#include "Subsystems/ActorModifierCoreSubsystem.h"
 
 #define LOCTEXT_NAMESPACE "AvaBooleanModifier"
 
@@ -171,7 +172,7 @@ void UAvaBooleanModifier::OnTransformUpdated(AActor* InActor, bool bInParentMove
 	}
 
 	// Delay to let the transform update propagate to attachment tree
-	FTSTicker::GetCoreTicker().AddTicker(FTickerDelegate::CreateWeakLambda(this, [this, bInParentMoved, InActor](float InDelta)
+	UActorModifierCoreSubsystem::GetModifierTicker().AddTicker(FTickerDelegate::CreateWeakLambda(this, [this, bInParentMoved, InActor](float InDelta)
 	{
 		// Check if colliding modifiers will be affected by this move
 		for (const TWeakObjectPtr<UAvaBooleanModifier>& CollidingModifierWeak : CollidingModifiers)
diff --git a/Engine/Plugins/VirtualProduction/Avalanche/Source/AvalancheShapes/Private/DynamicMeshes/AvaShapeDynMeshBase.cpp b/Engine/Plugins/VirtualProduction/Avalanche/Source/AvalancheShapes/Private/DynamicMeshes/AvaShapeDynMeshBase.cpp
--- a/Engine/Plugins/VirtualProduction/Avalanche/Source/AvalancheShapes/Private/DynamicMeshes/AvaShapeDynMeshBase.cpp
+++ b/Engine/Plugins/VirtualProduction/Avalanche/Source/AvalancheShapes/Private/DynamicMeshes/AvaShapeDynMeshBase.cpp
@@ -4,6 +4,7 @@
 
 #include "Async/Async.h"
 #include "Async/ParallelFor.h"
+#include "AvaSceneSubsystem.h"
 #include "AvaShapeActor.h"
 #include "Components/DynamicMeshComponent.h"
 #include "Containers/Ticker.h"
@@ -55,7 +56,8 @@ namespace UE::Private::AvaShape
 			if (UDynamicMesh* Mesh = MeshComponent->GetDynamicMesh())
 			{
 				const bool bIsSaving = UE::IsSavingPackage(InContext->GetPackage());
-				Mesh->EditMesh(InFunctor, InChangeType, InAttributeChangeType, !bIsSaving ? bInBroadcastChangeEvent : false);
+				const bool bSkipBroadcastChangeEvents = bIsSaving || !bInBroadcastChangeEvent;
+				Mesh->EditMesh(InFunctor, InChangeType, InAttributeChangeType, bSkipBroadcastChangeEvents);
 				return true;
 			}
 		}
@@ -377,41 +379,28 @@ void UAvaShapeDynamicMeshBase::RunUpdate()
 
 	MeshUpdateStatus = EAvaDynamicMeshUpdateState::UpdateInProgress;
 
-	auto ExecUpdate = [this]()->bool
+	auto ExecUpdate = [this](float=0.f)->bool
 	{
-		UDynamicMeshComponent* DynamicMeshComponent = GetShapeMeshComponent();
-
-		if (IsValid(this)
-			&& IsValid(DynamicMeshComponent)
-			&& !HasAnyFlags(RF_BeginDestroyed | RF_FinishDestroyed | RF_ClassDefaultObject))
+		if (UDynamicMeshComponent* DynamicMeshComponent = GetShapeMeshComponent())
 		{
-			if (UE::IsSavingPackage(this))
-			{
-				return false;
-			}
-
 			FActorModifierCoreScopedLock ModifierLock(UActorModifierCoreSubsystem::Get()->GetActorModifierStack(GetShapeActor()));
-
 			UpdateMesh();
 			DynamicMeshComponent->NotifyMeshModified();
 			MeshUpdateStatus = EAvaDynamicMeshUpdateState::UpToDate;
 			OnMeshUpdateFinished();
+			return false; // one time off
 		}
-
+		MeshUpdateStatus = EAvaDynamicMeshUpdateState::UpdateRequired;
 		return true;
 	};
 
 	if (GetRunAsync())
 	{
-		FTSTicker::GetCoreTicker().AddTicker(FTickerDelegate::CreateWeakLambda(this, [this, ExecUpdate](float)
-		{
-			// execute next tick if failed
-			return !ExecUpdate();
-		}));
+		GetTicker().AddTicker(FTickerDelegate::CreateWeakLambda(this, ExecUpdate));
 	}
-	else if (!ExecUpdate())
+	else
 	{
-		MeshUpdateStatus = EAvaDynamicMeshUpdateState::UpdateRequired;
+		ExecUpdate();
 	}
 }
 
@@ -741,6 +730,19 @@ void UAvaShapeDynamicMeshBase::OnMeshUpdateFinished()
 	}
 }
 
+FTSTicker& UAvaShapeDynamicMeshBase::GetTicker() const
+{
+	if (const UWorld* const World = GetWorld())
+	{
+		if (UAvaSceneSubsystem* const SceneSubsystem = World->GetSubsystem<UAvaSceneSubsystem>())
+		{
+			return SceneSubsystem->GetTicker();
+		}
+	}
+	UE_LOG(LogAvaDynamicMesh, Verbose, TEXT("No valid Motion Design Scene Subsystem was found to retrieve ticker. Defaulting to Core Ticker."));
+	return FTSTicker::GetCoreTicker();
+}
+
 FAvaShapeMesh* UAvaShapeDynamicMeshBase::GetMesh(int32 MeshIndex)
 {
 	if (FAvaShapeMeshData* MeshData = GetMeshData(MeshIndex))
diff --git a/Engine/Plugins/VirtualProduction/Avalanche/Source/AvalancheShapes/Public/DynamicMeshes/AvaShapeDynMeshBase.h b/Engine/Plugins/VirtualProduction/Avalanche/Source/AvalancheShapes/Public/DynamicMeshes/AvaShapeDynMeshBase.h
--- a/Engine/Plugins/VirtualProduction/Avalanche/Source/AvalancheShapes/Public/DynamicMeshes/AvaShapeDynMeshBase.h
+++ b/Engine/Plugins/VirtualProduction/Avalanche/Source/AvalancheShapes/Public/DynamicMeshes/AvaShapeDynMeshBase.h
@@ -15,6 +15,7 @@
 #include "AvaShapeDynMeshBase.generated.h"
 
 class AAvaShapeActor;
+class FTSTicker;
 class UDynamicMeshComponent;
 class UStaticMesh;
 struct FAvaShapeParametricMaterial;
@@ -419,6 +420,9 @@ protected:
 	TArray<FAvaSnapPoint> LocalSnapPoints;
 
 private:
+	/** Gets the ticker to use for shapes */
+	FTSTicker& GetTicker() const;
+
 	FAvaShapeMesh* GetMesh(int32 MeshIndex);
 
 	/** Update render state (materials) */
