Motion Design: optimizations for scene tree to have items sorted + translucent sort priority to use object ptrs instead of weak ptrs

#Synced-CL 51764840

diff --git a/Engine/Plugins/VirtualProduction/Avalanche/Source/Avalanche/Private/AvaScene.cpp b/Engine/Plugins/VirtualProduction/Avalanche/Source/Avalanche/Private/AvaScene.cpp
--- a/Engine/Plugins/VirtualProduction/Avalanche/Source/Avalanche/Private/AvaScene.cpp
+++ b/Engine/Plugins/VirtualProduction/Avalanche/Source/Avalanche/Private/AvaScene.cpp
@@ -319,6 +319,12 @@ void AAvaScene::BeginPlay()
 {
 	Super::BeginPlay();
 
+	// Conditionally sort the items for faster lookup.
+	SceneTree.SortItems();
+
+	// Fill the object to node map for fast lookup
+	SceneTree.ResolveObjects(GetTypedOuter<UWorld>());
+
 	if (UAvaCameraSubsystem* CameraSubsystem = UAvaCameraSubsystem::Get(this))
 	{
 		CameraSubsystem->RegisterScene(GetLevel());
diff --git a/Engine/Plugins/VirtualProduction/Avalanche/Source/AvalancheModifiers/Private/Modifiers/AvaTranslucentPriorityModifier.cpp b/Engine/Plugins/VirtualProduction/Avalanche/Source/AvalancheModifiers/Private/Modifiers/AvaTranslucentPriorityModifier.cpp
--- a/Engine/Plugins/VirtualProduction/Avalanche/Source/AvalancheModifiers/Private/Modifiers/AvaTranslucentPriorityModifier.cpp
+++ b/Engine/Plugins/VirtualProduction/Avalanche/Source/AvalancheModifiers/Private/Modifiers/AvaTranslucentPriorityModifier.cpp
@@ -158,9 +158,9 @@ void UAvaTranslucentPriorityModifier::Apply()
 		// Sets all components with the same priority
 		for (const FAvaTranslucentPriorityModifierComponentState& SortedComponentState : CachedSortedComponentStates)
 		{
-			if (SortedComponentState.ModifierWeak == this)
+			if (SortedComponentState.Modifier == this)
 			{
-				if (UPrimitiveComponent* Component = SortedComponentState.PrimitiveComponentWeak.Get())
+				if (UPrimitiveComponent* const Component = SortedComponentState.PrimitiveComponent)
 				{
 					LastSortPriorities.Add(Component, TranslucentSortPriority);
 					Component->SetTranslucentSortPriority(TranslucentSortPriority);
@@ -178,12 +178,8 @@ void UAvaTranslucentPriorityModifier::Apply()
 		AActor* PreviousComponentOwner = nullptr;
 		for (const FAvaTranslucentPriorityModifierComponentState& SortedComponentState : CachedSortedComponentStates)
 		{
-			UPrimitiveComponent* Component = SortedComponentState.PrimitiveComponentWeak.Get();
-			UAvaTranslucentPriorityModifier* ComponentModifier = SortedComponentState.ModifierWeak.Get();
-
-			if (!Component
-				|| !Component->GetOwner()
-				|| !ComponentModifier)
+			UPrimitiveComponent* const Component = SortedComponentState.PrimitiveComponent;
+			if (!Component || !Component->GetOwner() || !SortedComponentState.Modifier)
 			{
 				continue;
 			}
@@ -199,7 +195,7 @@ void UAvaTranslucentPriorityModifier::Apply()
 			PreviousComponentOwner = Component->GetOwner();
 
 			// This modifier handles this component
-			if (ComponentModifier == this)
+			if (SortedComponentState.Modifier == this)
 			{
 				LastSortPriorities.Add(Component, TranslucentSortPriority);
 				Component->SetTranslucentSortPriority(TranslucentSortPriority);
@@ -208,12 +204,11 @@ void UAvaTranslucentPriorityModifier::Apply()
 			else if (Component->TranslucencySortPriority != TranslucentSortPriority)
 			{
 				// Cache to avoid doing the same query for the same result
-				if (ComponentModifier->CachedSortedComponentStates.IsEmpty())
+				if (SortedComponentState.Modifier->CachedSortedComponentStates.IsEmpty())
 				{
-					ComponentModifier->CachedSortedComponentStates = CachedSortedComponentStates;
+					SortedComponentState.Modifier->CachedSortedComponentStates = CachedSortedComponentStates;
 				}
-
-				ComponentModifier->MarkModifierDirty();
+				SortedComponentState.Modifier->MarkModifierDirty();
 			}
 		}
 	}
diff --git a/Engine/Plugins/VirtualProduction/Avalanche/Source/AvalancheModifiers/Private/Shared/AvaTranslucentPriorityModifierShared.cpp b/Engine/Plugins/VirtualProduction/Avalanche/Source/AvalancheModifiers/Private/Shared/AvaTranslucentPriorityModifierShared.cpp
--- a/Engine/Plugins/VirtualProduction/Avalanche/Source/AvalancheModifiers/Private/Shared/AvaTranslucentPriorityModifierShared.cpp
+++ b/Engine/Plugins/VirtualProduction/Avalanche/Source/AvalancheModifiers/Private/Shared/AvaTranslucentPriorityModifierShared.cpp
@@ -10,6 +10,7 @@
 #include "Framework/AvaGameInstance.h"
 #include "IAvaSceneInterface.h"
 #include "Modifiers/AvaTranslucentPriorityModifier.h"
+#include "Version/AvaTranslucentPriorityModifierVersion.h"
 
 #if WITH_EDITOR
 #include "AvaOutlinerUtils.h"
@@ -27,15 +28,14 @@ void UAvaTranslucentPriorityModifierShared::SetComponentsState(UAvaTranslucentPr
 	// Remove old component states
 	for (TSet<FAvaTranslucentPriorityModifierComponentState>::TIterator It(ComponentStates); It; ++It)
 	{
-		UPrimitiveComponent* PrimitiveComponent = It->PrimitiveComponentWeak.Get();
-
+		UPrimitiveComponent* const PrimitiveComponent = It->PrimitiveComponent;
 		if (PrimitiveComponent && InComponents.Contains(PrimitiveComponent))
 		{
-			It->ModifierWeak = InModifierContext;
+			It->Modifier = InModifierContext;
 			continue;
 		}
 
-		const UAvaTranslucentPriorityModifier* Modifier = It->ModifierWeak.Get();
+		const UAvaTranslucentPriorityModifier* Modifier = It->Modifier.Get();
 		if (!PrimitiveComponent || !Modifier || Modifier == InModifierContext)
 		{
 			It.RemoveCurrent();
@@ -61,7 +61,7 @@ void UAvaTranslucentPriorityModifierShared::SaveComponentState(UAvaTranslucentPr
 
 	FAvaTranslucentPriorityModifierComponentState& ComponentState = ComponentStates.FindOrAdd(FAvaTranslucentPriorityModifierComponentState(InComponent));
 
-	if (!ComponentState.ModifierWeak.IsValid())
+	if (!ComponentState.Modifier)
 	{
 		ComponentState.Save();
 		bInOverrideContext = true;
@@ -69,29 +69,30 @@ void UAvaTranslucentPriorityModifierShared::SaveComponentState(UAvaTranslucentPr
 
 	if (bInOverrideContext)
 	{
-		ComponentState.ModifierWeak = InModifierContext;
+		ComponentState.Modifier = InModifierContext;
 	}
 }
 
 void UAvaTranslucentPriorityModifierShared::RestoreComponentState(const UAvaTranslucentPriorityModifier* InModifierContext, UPrimitiveComponent* InComponent, bool bInClearState)
 {
-	if (!IsValid(InComponent))
+	if (!InComponent)
 	{
 		return;
 	}
 
 	if (const FAvaTranslucentPriorityModifierComponentState* ComponentState = ComponentStates.Find(FAvaTranslucentPriorityModifierComponentState(InComponent)))
 	{
-		if (ComponentState->ModifierWeak == InModifierContext)
+		if (ComponentState->Modifier == InModifierContext)
 		{
-			ComponentState->Restore();
-
 			if (bInClearState)
 			{
+				FAvaTranslucentPriorityModifierComponentState ComponentStateCopy = *ComponentState;
 				ComponentStates.Remove(*ComponentState);
-
-				// This is needed in case other parent translucent modifier take ownership of the component
-				UActorComponent::MarkRenderStateDirtyEvent.Broadcast(*InComponent);
+				ComponentStateCopy.Restore();
+			}
+			else
+			{
+				ComponentState->Restore();
 			}
 		}
 	}
@@ -106,9 +107,18 @@ void UAvaTranslucentPriorityModifierShared::RestoreComponentsState(const UAvaTra
 
 	for (TSet<FAvaTranslucentPriorityModifierComponentState>::TIterator It(ComponentStates); It; ++It)
 	{
-		if (InModifierContext == It->ModifierWeak)
+		if (InModifierContext == It->Modifier)
 		{
-			RestoreComponentState(It->ModifierWeak.Get(), It->PrimitiveComponentWeak.Get(), bInClearState);
+			if (bInClearState)
+			{
+				FAvaTranslucentPriorityModifierComponentState ComponentStateCopy = *It;
+				It.RemoveCurrent();
+				ComponentStateCopy.Restore();
+			}
+			else
+			{
+				It->Restore();
+			}
 		}
 	}
 }
@@ -137,7 +147,7 @@ UAvaTranslucentPriorityModifier* UAvaTranslucentPriorityModifierShared::FindModi
 
 	if (const FAvaTranslucentPriorityModifierComponentState* ComponentState = ComponentStates.Find(FAvaTranslucentPriorityModifierComponentState(InComponent)))
 	{
-		ModifierContext = ComponentState->ModifierWeak.Get();
+		ModifierContext = ComponentState->Modifier;
 	}
 
 	return ModifierContext;
@@ -161,15 +171,13 @@ TArray<FAvaTranslucentPriorityModifierComponentState> UAvaTranslucentPriorityMod
 		// Only gather components the modifier is handling
 		for (const FAvaTranslucentPriorityModifierComponentState& ComponentState : ComponentStates)
 		{
-			const UAvaTranslucentPriorityModifier* ModifierContext = ComponentState.ModifierWeak.Get();
-			const UPrimitiveComponent* PrimitiveComponent = ComponentState.PrimitiveComponentWeak.Get();
-
-			if (!ModifierContext || !PrimitiveComponent)
+			const UPrimitiveComponent* const PrimitiveComponent = ComponentState.PrimitiveComponent;
+			if (!ComponentState.Modifier || !PrimitiveComponent)
 			{
 				continue;
 			}
 
-			if (ModifierContext == InModifierContext)
+			if (ComponentState.Modifier == InModifierContext)
 			{
 				SortedComponentStates.Add(ComponentState);
 			}
@@ -187,16 +195,14 @@ TArray<FAvaTranslucentPriorityModifierComponentState> UAvaTranslucentPriorityMod
 
 		for (const FAvaTranslucentPriorityModifierComponentState& ComponentState : ComponentStates)
 		{
-			const UAvaTranslucentPriorityModifier* ModifierContext = ComponentState.ModifierWeak.Get();
-			const UPrimitiveComponent* PrimitiveComponent = ComponentState.PrimitiveComponentWeak.Get();
-
-			if (!ModifierContext || !PrimitiveComponent)
+			const UPrimitiveComponent* const PrimitiveComponent = ComponentState.PrimitiveComponent;
+			if (!ComponentState.Modifier || !PrimitiveComponent)
 			{
 				continue;
 			}
 
-			if (ModifierContext->GetMode() == InModifierContext->GetMode()
-				&& ModifierContext->CameraActorWeak == CameraActor)
+			if (ComponentState.Modifier->GetMode() == InModifierContext->GetMode()
+				&& ComponentState.Modifier->CameraActorWeak == CameraActor)
 			{
 				SortedComponentStates.Add(ComponentState);
 			}
@@ -223,83 +229,49 @@ TArray<FAvaTranslucentPriorityModifierComponentState> UAvaTranslucentPriorityMod
 
 		for (const FAvaTranslucentPriorityModifierComponentState& ComponentState : ComponentStates)
 		{
-			const UAvaTranslucentPriorityModifier* ModifierContext = ComponentState.ModifierWeak.Get();
-			const UPrimitiveComponent* PrimitiveComponent = ComponentState.PrimitiveComponentWeak.Get();
-
-			if (!ModifierContext || !PrimitiveComponent)
+			const UPrimitiveComponent* const PrimitiveComponent = ComponentState.PrimitiveComponent;
+			if (!ComponentState.Modifier || !PrimitiveComponent)
 			{
 				continue;
 			}
 
-			const AActor* OwningActor = ModifierContext->GetModifiedActor();
+			const AActor* OwningActor = ComponentState.Modifier->GetModifiedActor();
 			if (!OwningActor)
 			{
 				continue;
 			}
 
-			if (ModifierContext->GetMode() == InModifierContext->GetMode()
-				&& OwningActor->GetTypedOuter<UWorld>() == World)
+			if (ComponentState.Modifier->GetMode() == InModifierContext->GetMode() && OwningActor->GetTypedOuter<UWorld>() == World)
 			{
 				SortedComponentStates.Add(ComponentState);
 			}
 		}
 
-		// Sort current modifiers by outliner hierarchy
-#if WITH_EDITOR
-		// Use the outliner first if available for sorting
-		TSharedPtr<IAvaOutliner> AvaOutliner = FAvaOutlinerUtils::EditorGetOutliner(World);
-		if (AvaOutliner.IsValid())
+		// Fallback on scene tree sorting
+		if (const IAvaSceneInterface* SceneInterface = FAvaActorUtils::GetSceneInterfaceFromActor(ModifiedActor))
 		{
-			auto FindItem =
-				[&AvaOutliner, World](AActor* InActor)->FAvaOutlinerItemPtr
-				{
-					FAvaOutlinerItemPtr Item;
-					while (InActor && !Item.IsValid())
-					{
-						Item = AvaOutliner->FindItem(InActor);
-						InActor = InActor->GetAttachParentActor();
-					}
-					return Item;
-				};
+			const FAvaSceneTree& SceneTree = SceneInterface->GetSceneTree();
 
-			SortedComponentStates.StableSort([&FindItem, bInvertedTree](const FAvaTranslucentPriorityModifierComponentState& InComponentA, const FAvaTranslucentPriorityModifierComponentState& InComponentB)->bool
-			{
-				const FAvaOutlinerItemPtr OutlinerItemA = FindItem(InComponentA.GetOwningActor());
-				const FAvaOutlinerItemPtr OutlinerItemB = FindItem(InComponentB.GetOwningActor());
-				return UE::AvaOutliner::CompareOutlinerItemOrder(OutlinerItemB, OutlinerItemA) ^ bInvertedTree;
-			});
-		}
-		else
-		{
-#endif
-			// Fallback on scene tree sorting
-			if (const IAvaSceneInterface* SceneInterface = FAvaActorUtils::GetSceneInterfaceFromActor(ModifiedActor))
-			{
-				const FAvaSceneTree& SceneTree = SceneInterface->GetSceneTree();
-
-				auto FindTreeNode =
-					[&SceneTree, World](AActor* InActor)->const FAvaSceneTreeNode*
+			auto FindTreeNode =
+				[&SceneTree, World](AActor* InActor)->const FAvaSceneTreeNode*
 					{
 						const FAvaSceneTreeNode* TreeNode = nullptr;
 						while (InActor && !TreeNode)
 						{
-							UObject* const StopOuter = World && InActor->IsIn(World) ? World : nullptr;
-							TreeNode = SceneTree.FindTreeNode(FAvaSceneItem(InActor, StopOuter));
+							TreeNode = SceneTree.FindObjectTreeNode(InActor, World && InActor->IsIn(World) ? World : nullptr);
 							InActor = InActor->GetAttachParentActor();
 						}
 						return TreeNode;
 					};
 
-				SortedComponentStates.StableSort([&FindTreeNode, bInvertedTree](const FAvaTranslucentPriorityModifierComponentState& InComponentA, const FAvaTranslucentPriorityModifierComponentState& InComponentB)->bool
+			SortedComponentStates.Sort(
+				[&FindTreeNode, bInvertedTree](const FAvaTranslucentPriorityModifierComponentState& InComponentA, const FAvaTranslucentPriorityModifierComponentState& InComponentB)->bool
 				{
 					const FAvaSceneTreeNode* SceneItemA = FindTreeNode(InComponentA.GetOwningActor());
 					const FAvaSceneTreeNode* SceneItemB = FindTreeNode(InComponentB.GetOwningActor());
 					return FAvaSceneTree::CompareTreeItemOrder(SceneItemB, SceneItemA) ^ bInvertedTree;
 				});
-			}
-#if WITH_EDITOR
 		}
-#endif
 	}
 
 	return SortedComponentStates;
@@ -327,14 +299,41 @@ void UAvaTranslucentPriorityModifierShared::SetSortPriorityStep(int32 InStep)
 	OnLevelGlobalsChangedDelegate.Broadcast();
 }
 
+void UAvaTranslucentPriorityModifierShared::Serialize(FArchive& Ar)
+{
+	Super::Serialize(Ar);
+	Ar.UsingCustomVersion(UE::AvaModifiers::FTranslucentPriorityModifierVersion::Guid);
+}
+
 void UAvaTranslucentPriorityModifierShared::PostLoad()
 {
 	Super::PostLoad();
 
+	const int32 ModifierVersion = GetLinkerCustomVersion(UE::AvaModifiers::FTranslucentPriorityModifierVersion::Guid);
+
+	// Upgrade weak object ptrs
+	if (ModifierVersion < UE::AvaModifiers::FTranslucentPriorityModifierVersion::UpgradeFromWeak)
+	{
+		TSet<FAvaTranslucentPriorityModifierComponentState> OldComponentStates = MoveTemp(ComponentStates);
+		ComponentStates.Empty(OldComponentStates.Num());
+
+		for (FAvaTranslucentPriorityModifierComponentState& ComponentState : OldComponentStates)
+		{
+PRAGMA_DISABLE_DEPRECATION_WARNINGS
+			ComponentState.PrimitiveComponent = ComponentState.PrimitiveComponentWeak.Get();
+			ComponentState.Modifier = ComponentState.ModifierWeak.Get();
+PRAGMA_ENABLE_DEPRECATION_WARNINGS
+			if (ComponentState.PrimitiveComponent)
+			{
+				ComponentStates.Emplace(MoveTemp(ComponentState));
+			}
+		}
+	}
+
 	// Clean up invalid components
 	for (TSet<FAvaTranslucentPriorityModifierComponentState>::TIterator It(ComponentStates); It; ++It)
 	{
-		if (!It->PrimitiveComponentWeak.IsValid())
+		if (!It->PrimitiveComponent)
 		{
 			It.RemoveCurrent();
 		}
@@ -343,25 +342,25 @@ void UAvaTranslucentPriorityModifierShared::PostLoad()
 
 void FAvaTranslucentPriorityModifierComponentState::Save()
 {
-	if (const UPrimitiveComponent* Component = PrimitiveComponentWeak.Get())
+	if (PrimitiveComponent)
 	{
-		SortPriority = Component->TranslucencySortPriority;
+		SortPriority = PrimitiveComponent->TranslucencySortPriority;
 	}
 }
 
 void FAvaTranslucentPriorityModifierComponentState::Restore() const
 {
-	if (UPrimitiveComponent* Component = PrimitiveComponentWeak.Get())
+	if (PrimitiveComponent)
 	{
-		Component->SetTranslucentSortPriority(SortPriority);
+		PrimitiveComponent->SetTranslucentSortPriority(SortPriority);
 	}
 }
 
 AActor* FAvaTranslucentPriorityModifierComponentState::GetOwningActor() const
 {
-	if (const UPrimitiveComponent* Component = PrimitiveComponentWeak.Get())
+	if (PrimitiveComponent)
 	{
-		return Component->GetOwner();
+		return PrimitiveComponent->GetOwner();
 	}
 
 	return nullptr;
@@ -369,9 +368,9 @@ AActor* FAvaTranslucentPriorityModifierComponentState::GetOwningActor() const
 
 FVector FAvaTranslucentPriorityModifierComponentState::GetComponentLocation() const
 {
-	if (const UPrimitiveComponent* Component = PrimitiveComponentWeak.Get())
+	if (PrimitiveComponent)
 	{
-		return Component->GetComponentLocation();
+		return PrimitiveComponent->GetComponentLocation();
 	}
 
 	return FVector::ZeroVector;
diff --git a/Engine/Plugins/VirtualProduction/Avalanche/Source/AvalancheModifiers/Private/Version/AvaTranslucentPriorityModifierVersion.cpp b/Engine/Plugins/VirtualProduction/Avalanche/Source/AvalancheModifiers/Private/Version/AvaTranslucentPriorityModifierVersion.cpp
new file mode 100644
--- /dev/null
+++ b/Engine/Plugins/VirtualProduction/Avalanche/Source/AvalancheModifiers/Private/Version/AvaTranslucentPriorityModifierVersion.cpp
@@ -0,0 +1,14 @@
+// Copyright Epic Games, Inc. All Rights Reserved.
+
+#include "Version/AvaTranslucentPriorityModifierVersion.h"
+#include "Serialization/CustomVersion.h"
+
+namespace UE::AvaModifiers
+{
+
+const FGuid FTranslucentPriorityModifierVersion::Guid(0x3A5DCA18, 0x17EF4851, 0xA6156A1B, 0x3A1B0FB2);
+FCustomVersionRegistration GRegisterTranslucentPriorityModifierVersion(FTranslucentPriorityModifierVersion::Guid
+	, FTranslucentPriorityModifierVersion::LatestVersion
+	, TEXT("TranslucentPriorityModifierVersion"));
+
+} // UE::AvaModifiers
diff --git a/Engine/Plugins/VirtualProduction/Avalanche/Source/AvalancheModifiers/Public/Shared/AvaTranslucentPriorityModifierShared.h b/Engine/Plugins/VirtualProduction/Avalanche/Source/AvalancheModifiers/Public/Shared/AvaTranslucentPriorityModifierShared.h
--- a/Engine/Plugins/VirtualProduction/Avalanche/Source/AvalancheModifiers/Public/Shared/AvaTranslucentPriorityModifierShared.h
+++ b/Engine/Plugins/VirtualProduction/Avalanche/Source/AvalancheModifiers/Public/Shared/AvaTranslucentPriorityModifierShared.h
@@ -18,10 +18,17 @@ struct FAvaTranslucentPriorityModifierComponentState
 {
 	GENERATED_BODY()
 
-	FAvaTranslucentPriorityModifierComponentState() {}
+PRAGMA_DISABLE_DEPRECATION_WARNINGS
+	FAvaTranslucentPriorityModifierComponentState() = default;
 	explicit FAvaTranslucentPriorityModifierComponentState(UPrimitiveComponent* InComponent)
-		: PrimitiveComponentWeak(InComponent)
-	{}
+		: PrimitiveComponent(InComponent)
+	{
+	}
+	FAvaTranslucentPriorityModifierComponentState(const FAvaTranslucentPriorityModifierComponentState&) = default;
+	FAvaTranslucentPriorityModifierComponentState(FAvaTranslucentPriorityModifierComponentState&&) = default;
+	FAvaTranslucentPriorityModifierComponentState& operator=(const FAvaTranslucentPriorityModifierComponentState&) = default;
+	FAvaTranslucentPriorityModifierComponentState& operator=(FAvaTranslucentPriorityModifierComponentState&&) = default;
+PRAGMA_ENABLE_DEPRECATION_WARNINGS
 
 	/** Save this component state if valid */
 	void Save();
@@ -37,21 +44,29 @@ struct FAvaTranslucentPriorityModifierComponentState
 
 	friend uint32 GetTypeHash(const FAvaTranslucentPriorityModifierComponentState& InItem)
 	{
-		return GetTypeHash(InItem.PrimitiveComponentWeak);
+		return GetTypeHash(InItem.PrimitiveComponent);
 	}
 
 	bool operator==(const FAvaTranslucentPriorityModifierComponentState& InOther) const
 	{
-		return PrimitiveComponentWeak.IsValid() && PrimitiveComponentWeak == InOther.PrimitiveComponentWeak;
+		return PrimitiveComponent == InOther.PrimitiveComponent;
 	}
 
+	UE_DEPRECATED(5.8,"Use Modifier instead")
+	UPROPERTY(meta=(DeprecatedProperty, DeprecationMessage="Use Modifier instead"))
+	TWeakObjectPtr<UAvaTranslucentPriorityModifier> ModifierWeak;
+
+	UE_DEPRECATED(5.8,"Use PrimitiveComponent instead")
+	UPROPERTY(meta=(DeprecatedProperty, DeprecationMessage="Use PrimitiveComponent instead"))
+	TWeakObjectPtr<UPrimitiveComponent> PrimitiveComponentWeak;
+
 	/** Modifier managing this component state */
 	UPROPERTY()
-	TWeakObjectPtr<UAvaTranslucentPriorityModifier> ModifierWeak;
+	TObjectPtr<UAvaTranslucentPriorityModifier> Modifier;
 
 	/** Primitive component that this state is describing */
 	UPROPERTY()
-	TWeakObjectPtr<UPrimitiveComponent> PrimitiveComponentWeak;
+	TObjectPtr<UPrimitiveComponent> PrimitiveComponent;
 
 	/** Previous sort priority to restore */
 	UPROPERTY()
@@ -107,6 +122,7 @@ public:
 
 private:
 	//~ Begin UObject
+	virtual void Serialize(FArchive& Ar) override;
 	virtual void PostLoad() override;
 	//~ End UObject
 
diff --git a/Engine/Plugins/VirtualProduction/Avalanche/Source/AvalancheModifiers/Public/Version/AvaTranslucentPriorityModifierVersion.h b/Engine/Plugins/VirtualProduction/Avalanche/Source/AvalancheModifiers/Public/Version/AvaTranslucentPriorityModifierVersion.h
new file mode 100644
--- /dev/null
+++ b/Engine/Plugins/VirtualProduction/Avalanche/Source/AvalancheModifiers/Public/Version/AvaTranslucentPriorityModifierVersion.h
@@ -0,0 +1,37 @@
+// Copyright Epic Games, Inc. All Rights Reserved.
+
+#pragma once
+
+#include "HAL/Platform.h"
+#include "Misc/Guid.h"
+
+#define UE_API AVALANCHEMODIFIERS_API
+
+namespace UE::AvaModifiers
+{
+
+/** Describes the versions of the translucent priority modifier */
+struct FTranslucentPriorityModifierVersion
+{
+private:
+	FTranslucentPriorityModifierVersion() = delete;
+
+public:
+	enum Type : uint8
+	{
+		PreVersioning = 0,
+
+		/** Upgraded TWeakObjectPtr to TObjectPtr in UAvaTranslucentPriorityModifierShared */
+		UpgradeFromWeak,
+
+		/* ------------------------------------------------------ */
+		VersionPlusOne,
+		LatestVersion = VersionPlusOne - 1
+	};
+
+	UE_API const static FGuid Guid;
+};
+
+} // UE::AvaModifiers
+
+#undef UE_API
diff --git a/Engine/Plugins/VirtualProduction/Avalanche/Source/AvalancheOutliner/Private/Data/AvaOutlinerSaveState.cpp b/Engine/Plugins/VirtualProduction/Avalanche/Source/AvalancheOutliner/Private/Data/AvaOutlinerSaveState.cpp
--- a/Engine/Plugins/VirtualProduction/Avalanche/Source/AvalancheOutliner/Private/Data/AvaOutlinerSaveState.cpp
+++ b/Engine/Plugins/VirtualProduction/Avalanche/Source/AvalancheOutliner/Private/Data/AvaOutlinerSaveState.cpp
@@ -255,6 +255,8 @@ void FAvaOutlinerSaveState::SaveSceneTree(const FAvaOutliner& InOutliner, bool b
 			SceneTree->Reset();
 		}
 		SaveSceneTreeRecursive(InOutliner.GetTreeRoot(), *SceneTree, InOutliner.GetWorld());
+		SceneTree->SortItems();
+		SceneTree->ResolveObjects(InOutliner.GetWorld());
 	}
 }
 
diff --git a/Engine/Plugins/VirtualProduction/Avalanche/Source/AvalancheSceneTree/Private/AvaSceneTree.cpp b/Engine/Plugins/VirtualProduction/Avalanche/Source/AvalancheSceneTree/Private/AvaSceneTree.cpp
--- a/Engine/Plugins/VirtualProduction/Avalanche/Source/AvalancheSceneTree/Private/AvaSceneTree.cpp
+++ b/Engine/Plugins/VirtualProduction/Avalanche/Source/AvalancheSceneTree/Private/AvaSceneTree.cpp
@@ -4,6 +4,7 @@
 #include "Engine/Level.h"
 #include "Engine/World.h"
 #include "GameFramework/Actor.h"
+#include "Misc/EnumerateRange.h"
 #include "Modules/ModuleManager.h"
 
 IMPLEMENT_MODULE(FDefaultModuleImpl, AvalancheSceneTree)
@@ -30,6 +31,11 @@ FAvaSceneTreeNode* FAvaSceneTree::FindTreeNode(const FAvaSceneItem& InItem)
 	return nullptr;
 }
 
+const FAvaSceneTreeNode* FAvaSceneTree::FindObjectTreeNode(const UObject* InObject, const UObject* InOuter) const
+{
+	return ObjectToNodeMap.Find({ InObject, InOuter });
+}
+
 const FAvaSceneTreeNode* FAvaSceneTree::FindTreeNode(const FAvaSceneItem& InItem) const
 {
 	if (InItem.IsValid())
@@ -71,7 +77,7 @@ FAvaSceneTreeNode& FAvaSceneTree::GetOrAddTreeNode(const FAvaSceneItem& InItem,
 	TreeNode.LocalIndex  = ParentNode->ChildrenIndices.Add(TreeNode.GlobalIndex);
 	TreeNode.ParentIndex = ParentNode->GlobalIndex;
 	TreeNode.OwningTree  = this;
-
+	bSorted = false;
 	return ItemTreeMap.Add(InItem, MoveTemp(TreeNode));
 }
 
@@ -118,10 +124,17 @@ const FAvaSceneTreeNode* FAvaSceneTree::FindLowestCommonAncestor(const TArray<co
 
 bool FAvaSceneTree::CompareTreeItemOrder(const FAvaSceneTreeNode* InA, const FAvaSceneTreeNode* InB)
 {
-	if (!InA || !InB)
+	if (!InA || !InB || InA->OwningTree != InB->OwningTree)
 	{
 		return false;
 	}
+
+	// Fast path: items are sorted by their global order in the tree, so their global index matches the item order. No need for extra computation
+	if (InA->OwningTree->bSorted)
+	{
+		return InA->GlobalIndex < InB->GlobalIndex;
+	}
+
 	if (const FAvaSceneTreeNode* LowestCommonAncestor = FindLowestCommonAncestor({InA, InB}))
 	{
 		const TArray<const FAvaSceneTreeNode*> PathToA = LowestCommonAncestor->FindPath({InA});
@@ -157,6 +170,8 @@ void FAvaSceneTree::Reset()
 	RootNode.Reset();
 	ItemTreeMap.Reset();
 	SceneItems.Reset();
+	ObjectToNodeMap.Reset();
+	bSorted = false;
 }
 
 void FAvaSceneTree::UpdateTreeNodes()
@@ -221,3 +236,70 @@ TArray<AActor*> FAvaSceneTree::GetRootActors(const ULevel* InLevel) const
 	
 	return OutActors;
 }
+
+void FAvaSceneTree::SortItems()
+{
+	if (bSorted)
+	{
+		return;
+	}
+
+	TRACE_CPUPROFILER_EVENT_SCOPE(FAvaSceneTree::SortItems);
+
+	// Edit a new copy of the scene items as the items will be index to the old structure. 
+	TArray<FAvaSceneItem> NewSceneItems = SceneItems;
+	NewSceneItems.Sort(
+		[This=this](const FAvaSceneItem& A, const FAvaSceneItem& B)
+		{
+			const FAvaSceneTreeNode& NodeA = This->ItemTreeMap[A];
+			const FAvaSceneTreeNode& NodeB = This->ItemTreeMap[B];
+			return CompareTreeItemOrder(&NodeA, &NodeB);
+		});
+
+	auto FixIndex = [&NewSceneItems, This=this](int32& InOutIndex)
+		{
+			// index is pointing to root
+			if (InOutIndex == INDEX_NONE)
+			{
+				return;
+			}
+			InOutIndex = NewSceneItems.Find(This->SceneItems[InOutIndex]);
+		};
+
+	// Fix up indices for a given tree node
+	// No need to fix local index as that was the rule for sorting
+	for (TPair<FAvaSceneItem, FAvaSceneTreeNode>& Pair : ItemTreeMap)
+	{
+		FAvaSceneTreeNode& Node = Pair.Value;
+		FixIndex(Node.GlobalIndex);
+		FixIndex(Node.ParentIndex);
+
+		int32 LocalIndex = 0;
+		for (int32& ChildIndex : Node.ChildrenIndices)
+		{
+			FixIndex(ChildIndex);
+		}
+	}
+
+	SceneItems = MoveTemp(NewSceneItems);
+	ObjectToNodeMap.Reset();
+	bSorted = true;
+}
+
+void FAvaSceneTree::ResolveObjects(UObject* InOuter)
+{
+	TRACE_CPUPROFILER_EVENT_SCOPE(FAvaSceneTree::ResolveObjects);
+
+	TPair<FObjectKey, FObjectKey> ObjectPair;
+	ObjectPair.Value = InOuter;
+
+	ObjectToNodeMap.Empty(ItemTreeMap.Num());
+	for (const TPair<FAvaSceneItem, FAvaSceneTreeNode>& Pair : ItemTreeMap)
+	{
+		if (UObject* Object = Pair.Key.Resolve(InOuter))
+		{
+			ObjectPair.Key = Object;
+			ObjectToNodeMap.Add(ObjectPair, Pair.Value);
+		}
+	}
+}
diff --git a/Engine/Plugins/VirtualProduction/Avalanche/Source/AvalancheSceneTree/Public/AvaSceneItem.h b/Engine/Plugins/VirtualProduction/Avalanche/Source/AvalancheSceneTree/Public/AvaSceneItem.h
--- a/Engine/Plugins/VirtualProduction/Avalanche/Source/AvalancheSceneTree/Public/AvaSceneItem.h
+++ b/Engine/Plugins/VirtualProduction/Avalanche/Source/AvalancheSceneTree/Public/AvaSceneItem.h
@@ -23,7 +23,7 @@ struct FAvaSceneItem
 
 	FAvaSceneItem() = default;
 
-	FAvaSceneItem(UObject* InObject, UObject* InStopOuter)
+	FAvaSceneItem(const UObject* InObject, const UObject* InStopOuter)
 	{
 		check(InObject && (!InStopOuter || InObject->IsIn(InStopOuter)));
 		Id     = InObject->GetPathName(InStopOuter);
diff --git a/Engine/Plugins/VirtualProduction/Avalanche/Source/AvalancheSceneTree/Public/AvaSceneTree.h b/Engine/Plugins/VirtualProduction/Avalanche/Source/AvalancheSceneTree/Public/AvaSceneTree.h
--- a/Engine/Plugins/VirtualProduction/Avalanche/Source/AvalancheSceneTree/Public/AvaSceneTree.h
+++ b/Engine/Plugins/VirtualProduction/Avalanche/Source/AvalancheSceneTree/Public/AvaSceneTree.h
@@ -5,6 +5,7 @@
 #include "AvaSceneItem.h"
 #include "AvaSceneTreeNode.h"
 #include "Containers/Map.h"
+#include "UObject/ObjectKey.h"
 #include "AvaSceneTree.generated.h"
 
 class AActor;
@@ -25,6 +26,12 @@ struct AVALANCHESCENETREE_API FAvaSceneTree
 
 	FAvaSceneTreeNode* FindTreeNode(const FAvaSceneItem& InItem);
 
+	/**
+	 * Finds tree node using the object map to avoid FAvaSceneItem allocation
+	 * ResolveObjects() must've been called for the tree for the object map to be up-to-date.
+	 */
+	const FAvaSceneTreeNode* FindObjectTreeNode(const UObject* InObject, const UObject* InOuter) const;
+
 	const FAvaSceneTreeNode* FindTreeNode(const FAvaSceneItem& InItem) const;
 
 	const FAvaSceneItem* GetItemAtIndex(int32 InIndex) const;
@@ -43,6 +50,17 @@ struct AVALANCHESCENETREE_API FAvaSceneTree
 	/** Returns all root actors in this tree */
 	TArray<AActor*> GetRootActors(const ULevel* InLevel) const;
 
+	bool IsSorted() const
+	{
+		return bSorted;
+	}
+
+	/** Sorts items based on tree item ordering. No-op if the tree is already sorted. */
+	void SortItems();
+
+	/** Fills the object to node map for faster look up */
+	void ResolveObjects(UObject* InOuter);
+
 private:
 	void UpdateTreeNodes();
 
@@ -54,6 +72,13 @@ private:
 
 	UPROPERTY()
 	TMap<FAvaSceneItem, FAvaSceneTreeNode> ItemTreeMap;
+
+	/** (Object, Outer) pair to their tree node */
+	TMap<TPair<FObjectKey, FObjectKey>, FAvaSceneTreeNode> ObjectToNodeMap;
+
+	/** Whether scene items have been sorted to their latest structure */
+	UPROPERTY()
+	bool bSorted = false;
 };
 
 template<>
