Motion Design: added wait for text to build before playing in playables.

#Synced-CL 51764840

diff --git a/Engine/Plugins/VirtualProduction/Avalanche/Source/Avalanche/Private/AvaSceneSubsystem.cpp b/Engine/Plugins/VirtualProduction/Avalanche/Source/Avalanche/Private/AvaSceneSubsystem.cpp
--- a/Engine/Plugins/VirtualProduction/Avalanche/Source/Avalanche/Private/AvaSceneSubsystem.cpp
+++ b/Engine/Plugins/VirtualProduction/Avalanche/Source/Avalanche/Private/AvaSceneSubsystem.cpp
@@ -6,6 +6,8 @@
 #include "Engine/Level.h"
 #include "Engine/World.h"
 #include "IAvaSceneInterface.h"
+#include "IText3DBuildSystem.h"
+#include "Utilities/Text3DUtilities.h"
 
 namespace UE::Ava::Private
 {
@@ -37,12 +39,12 @@ IAvaSceneInterface* UAvaSceneSubsystem::GetSceneInterface() const
 	return nullptr;
 }
 
-IAvaSceneInterface* UAvaSceneSubsystem::GetSceneInterface(ULevel* InLevel) const
+IAvaSceneInterface* UAvaSceneSubsystem::GetSceneInterface(const ULevel* InLevel) const
 {
 	return SceneInterfaces.FindRef(InLevel).Get();
 }
 
-IAvaSceneInterface* UAvaSceneSubsystem::FindSceneInterface(ULevel* InLevel)
+IAvaSceneInterface* UAvaSceneSubsystem::FindSceneInterface(const ULevel* InLevel)
 {
 	AAvaScene* SceneActor = nullptr;
 	if (InLevel)
@@ -52,6 +54,12 @@ IAvaSceneInterface* UAvaSceneSubsystem::FindSceneInterface(ULevel* InLevel)
 	return SceneActor;
 }
 
+bool UAvaSceneSubsystem::IsReadyToPlay(const UObject* InContextObject) const
+{
+	const IText3DBuildSystemInterface* const TextBuildSystem = UE::Text3D::Utilities::FindBuildSystem(this);
+	return !TextBuildSystem || !TextBuildSystem->IsBuildInProgress();
+}
+
 void UAvaSceneSubsystem::PostInitialize()
 {
 	Super::PostInitialize();
diff --git a/Engine/Plugins/VirtualProduction/Avalanche/Source/Avalanche/Public/AvaSceneSubsystem.h b/Engine/Plugins/VirtualProduction/Avalanche/Source/Avalanche/Public/AvaSceneSubsystem.h
--- a/Engine/Plugins/VirtualProduction/Avalanche/Source/Avalanche/Public/AvaSceneSubsystem.h
+++ b/Engine/Plugins/VirtualProduction/Avalanche/Source/Avalanche/Public/AvaSceneSubsystem.h
@@ -9,10 +9,12 @@
 #include "UObject/WeakInterfacePtr.h"
 #include "AvaSceneSubsystem.generated.h"
 
+#define UE_API AVALANCHE_API
+
 class IAvaSceneInterface;
 class ULevel;
 
-UCLASS(MinimalAPI, DisplayName = "Motion Design Scene Subsystem")
+UCLASS(MinimalAPI, BlueprintType, DisplayName = "Motion Design Scene Subsystem")
 class UAvaSceneSubsystem : public UTickableWorldSubsystem
 {
 	GENERATED_BODY()
@@ -20,27 +22,31 @@ class UAvaSceneSubsystem : public UTickableWorldSubsystem
 public:
 	void RegisterSceneInterface(ULevel* InLevel, IAvaSceneInterface* InSceneInterface);
 
-	AVALANCHE_API FTSTicker& GetTicker();
+	UE_API FTSTicker& GetTicker();
 
 	/** Gets the Scene Interface for the World's Persistent Level */
-	AVALANCHE_API IAvaSceneInterface* GetSceneInterface() const;
+	UE_API IAvaSceneInterface* GetSceneInterface() const;
 
 	/** Gets the Scene Interface for the provided Level */
-	AVALANCHE_API IAvaSceneInterface* GetSceneInterface(ULevel* InLevel) const;
+	UE_API IAvaSceneInterface* GetSceneInterface(const ULevel* InLevel) const;
 
 	/** Returns the first interface found in the given Level, if any */
-	AVALANCHE_API static IAvaSceneInterface* FindSceneInterface(ULevel* InLevel);
+	UE_API static IAvaSceneInterface* FindSceneInterface(const ULevel* InLevel);
+
+	/** Determines whether the motion design is ready for play */
+	UFUNCTION(BlueprintPure, Category="Motion Design", meta=(DefaultToSelf="InContextObject"))
+	UE_API bool IsReadyToPlay(const UObject* InContextObject) const;
 
 protected:
 	//~ Begin UWorldSubsystem
-	AVALANCHE_API virtual void PostInitialize() override;
-	AVALANCHE_API virtual bool DoesSupportWorldType(const EWorldType::Type InWorldType) const override;
+	UE_API virtual void PostInitialize() override;
+	UE_API virtual bool DoesSupportWorldType(const EWorldType::Type InWorldType) const override;
 	//~ End UWorldSubsystem
 
 	//~ Begin FTickableObjectBase
-	AVALANCHE_API virtual bool IsTickableInEditor() const override;
-	AVALANCHE_API virtual void Tick(float InDeltaTime) override;
-	AVALANCHE_API virtual TStatId GetStatId() const override;
+	UE_API virtual bool IsTickableInEditor() const override;
+	UE_API virtual void Tick(float InDeltaTime) override;
+	UE_API virtual TStatId GetStatId() const override;
 	//~ End FTickableObjectBase
 
 private:
@@ -49,3 +55,5 @@ private:
 
 	TMap<TWeakObjectPtr<ULevel>, TWeakInterfacePtr<IAvaSceneInterface>> SceneInterfaces;
 };
+
+#undef UE_API
diff --git a/Engine/Plugins/VirtualProduction/Avalanche/Source/AvalancheMedia/Private/Playable/Playables/AvaPlayableLevelStreaming.cpp b/Engine/Plugins/VirtualProduction/Avalanche/Source/AvalancheMedia/Private/Playable/Playables/AvaPlayableLevelStreaming.cpp
--- a/Engine/Plugins/VirtualProduction/Avalanche/Source/AvalancheMedia/Private/Playable/Playables/AvaPlayableLevelStreaming.cpp
+++ b/Engine/Plugins/VirtualProduction/Avalanche/Source/AvalancheMedia/Private/Playable/Playables/AvaPlayableLevelStreaming.cpp
@@ -4,6 +4,7 @@
 #include "AvaRemoteControlRebind.h"
 #include "AvaRemoteControlUtils.h"
 #include "AvaScene.h"
+#include "AvaSceneSubsystem.h"
 #include "Engine/Level.h"
 #include "Engine/LevelStreamingDynamic.h"
 #include "Engine/LocalPlayer.h"
@@ -418,6 +419,19 @@ bool UAvaPlayableLevelStreaming::IsReadyToPlay() const
 	{
 		return false;
 	}
+
+	if (const ULevel* LoadedLevel = GetLoadedLevel())
+	{
+		if (UWorld* const World = LoadedLevel->GetWorld())
+		{
+			const UAvaSceneSubsystem* const SceneSubsystem = World->GetSubsystem<UAvaSceneSubsystem>();
+			if (SceneSubsystem && !SceneSubsystem->IsReadyToPlay(LoadedLevel))
+			{
+				return false;
+			}
+		}
+	}
+
 	for (const UAvaPlayableLevelStreaming* SubPlayable : SubPlayables)
 	{
 		if (SubPlayable && !SubPlayable->IsReadyToPlay())
