Motion Design: fix issue where editor preview rundown would show 1 frame.

#Synced-CL 51764840

diff --git a/Engine/Plugins/VirtualProduction/Avalanche/Source/Avalanche/Private/Framework/AvaGameInstance.cpp b/Engine/Plugins/VirtualProduction/Avalanche/Source/Avalanche/Private/Framework/AvaGameInstance.cpp
--- a/Engine/Plugins/VirtualProduction/Avalanche/Source/Avalanche/Private/Framework/AvaGameInstance.cpp
+++ b/Engine/Plugins/VirtualProduction/Avalanche/Source/Avalanche/Private/Framework/AvaGameInstance.cpp
@@ -17,6 +17,44 @@
 
 #if WITH_EDITOR
 #include "Editor.h"
+#include "TickableEditorObject.h"
+#endif
+
+#if WITH_EDITOR
+/**
+ * Editor Ticker Object used to tick the Game Instance World. 
+ * Used instead of GEditor->PostEditorTick so that the world ticks before viewport ticks.
+ * FTickableGameObject::TickObjects(null, ...) is called in Editor Engine conditionally, so cannot be used as a way to tick the Game Instance World in editor.
+ * In addition, other ticking logic could rely on TickObjects(null, ...) to happen after worlds have ticked. 
+ */
+class FAvaGameInstanceEditorTicker : public FTickableEditorObject
+{
+public:
+	explicit FAvaGameInstanceEditorTicker(TDelegate<void(float)>&& InDelegate)
+		: TickDelegate(MoveTemp(InDelegate))
+	{
+	}
+
+	//~ Begin FTickableObjectBase
+	virtual ETickableTickType GetTickableTickType() const override
+	{
+		return ETickableTickType::Always;
+	}
+
+	virtual void Tick(float InDeltaTime) override
+	{
+		TickDelegate.ExecuteIfBound(InDeltaTime);
+	}
+
+	virtual TStatId GetStatId() const override
+	{
+		RETURN_QUICK_DECLARE_CYCLE_STAT(FAvaGameInstanceEditorTicker, STATGROUP_Tickables);
+	}
+	//~ End FTickableObjectBase
+
+private:
+	TDelegate<void(float)> TickDelegate;
+};
 #endif
 
 UAvaGameInstance::FOnAvaGameInstanceEvent UAvaGameInstance::OnEndPlay;
@@ -143,9 +181,13 @@ bool UAvaGameInstance::BeginPlayWorld(const FAvaInstancePlaySettings& InWorldPla
 	}
 
 #if WITH_EDITOR
-	if (GEditor && !OnEditorTickHandle.IsValid())
+	if (GEditor)
+	{
+		WorldTickerEditor = MakePimpl<FAvaGameInstanceEditorTicker>(TDelegate<void(float)>::CreateUObject(this, &UAvaGameInstance::EditorTick));
+	}
+	else
 	{
-		OnEditorTickHandle = GEditor->OnPostEditorTick().AddUObject(this, &UAvaGameInstance::EditorTick);
+		WorldTickerEditor.Reset();
 	}
 #endif // WITH_EDITOR
 
@@ -324,12 +366,8 @@ void UAvaGameInstance::EndPlayWorld()
 	SetTickableTickType(ETickableTickType::Never);
 
 #if WITH_EDITOR
-	if (GEditor)
-	{
-		GEditor->OnPostEditorTick().Remove(OnEditorTickHandle);
-		OnEditorTickHandle.Reset();
-	}
-#endif // WITH_EDITOR
+	WorldTickerEditor.Reset();
+#endif
 
 	FCoreDelegates::OnEndFrame.RemoveAll(this);
 
diff --git a/Engine/Plugins/VirtualProduction/Avalanche/Source/Avalanche/Public/Framework/AvaGameInstance.h b/Engine/Plugins/VirtualProduction/Avalanche/Source/Avalanche/Public/Framework/AvaGameInstance.h
--- a/Engine/Plugins/VirtualProduction/Avalanche/Source/Avalanche/Public/Framework/AvaGameInstance.h
+++ b/Engine/Plugins/VirtualProduction/Avalanche/Source/Avalanche/Public/Framework/AvaGameInstance.h
@@ -157,15 +157,16 @@ protected:
 	/** Render command fence to ensure the render target has been rendered. */
 	TUniquePtr<FRenderCommandFence> RenderTargetFence;
 
+#if WITH_EDITOR
+	/** Ticker implementation for the game instance to tick the world in editor engine */
+	TPimplPtr<class FAvaGameInstanceEditorTicker> WorldTickerEditor; 
+#endif
+
 	/** Last frame that was ticked */
 	uint64 TickedFrame = 0;
 
 	/** Handle to delegate for EnginePreExit event. */
 	FDelegateHandle EnginePreExitHandle;
-#if WITH_EDITOR
-	/** Handle to the UEditorEngine::OnPostEditorTick */
-	FDelegateHandle OnEditorTickHandle;
-#endif
 
 	UE_API static FOnAvaGameInstanceEvent OnEndPlay;
 	UE_API static FOnAvaGameInstanceEvent OnRenderTargetReady;
