Motion Design: fix issue where scene tree resolve objects is called after the first usage of FindObjectTreeNode, preventing the tree node to be found. This change allows FindObjectTreeNode to add the entry itself if it isn't there, while also calling ResolveObjects at a much earlier stage in the motion design scene actor.

#Synced-CL 51764840

diff --git a/Engine/Plugins/VirtualProduction/Avalanche/Source/Avalanche/Private/AvaScene.cpp b/Engine/Plugins/VirtualProduction/Avalanche/Source/Avalanche/Private/AvaScene.cpp
--- a/Engine/Plugins/VirtualProduction/Avalanche/Source/Avalanche/Private/AvaScene.cpp
+++ b/Engine/Plugins/VirtualProduction/Avalanche/Source/Avalanche/Private/AvaScene.cpp
@@ -315,15 +315,20 @@ void AAvaScene::PostActorCreated()
 	RegisterObjects();
 }
 
-void AAvaScene::BeginPlay()
+void AAvaScene::PreInitializeComponents()
 {
-	Super::BeginPlay();
+	Super::PreInitializeComponents();
 
 	// Conditionally sort the items for faster lookup.
 	SceneTree.SortItems();
 
 	// Fill the object to node map for fast lookup
 	SceneTree.ResolveObjects(GetTypedOuter<UWorld>());
+}
+
+void AAvaScene::BeginPlay()
+{
+	Super::BeginPlay();
 
 	if (UAvaCameraSubsystem* CameraSubsystem = UAvaCameraSubsystem::Get(this))
 	{
diff --git a/Engine/Plugins/VirtualProduction/Avalanche/Source/Avalanche/Public/AvaScene.h b/Engine/Plugins/VirtualProduction/Avalanche/Source/Avalanche/Public/AvaScene.h
--- a/Engine/Plugins/VirtualProduction/Avalanche/Source/Avalanche/Public/AvaScene.h
+++ b/Engine/Plugins/VirtualProduction/Avalanche/Source/Avalanche/Public/AvaScene.h
@@ -104,6 +104,7 @@ public:
 
 	//~ Begin AActor
 	virtual void PostActorCreated() override;
+	virtual void PreInitializeComponents() override;
 	virtual void BeginPlay() override;
 	virtual void EndPlay(const EEndPlayReason::Type InEndPlayReason) override;
 #if WITH_EDITOR
diff --git a/Engine/Plugins/VirtualProduction/Avalanche/Source/AvalancheSceneTree/Private/AvaSceneTree.cpp b/Engine/Plugins/VirtualProduction/Avalanche/Source/AvalancheSceneTree/Private/AvaSceneTree.cpp
--- a/Engine/Plugins/VirtualProduction/Avalanche/Source/AvalancheSceneTree/Private/AvaSceneTree.cpp
+++ b/Engine/Plugins/VirtualProduction/Avalanche/Source/AvalancheSceneTree/Private/AvaSceneTree.cpp
@@ -33,7 +33,18 @@ FAvaSceneTreeNode* FAvaSceneTree::FindTreeNode(const FAvaSceneItem& InItem)
 
 const FAvaSceneTreeNode* FAvaSceneTree::FindObjectTreeNode(const UObject* InObject, const UObject* InOuter) const
 {
-	return ObjectToNodeMap.Find({ InObject, InOuter });
+	const FAvaSceneTreeNode* TreeNode = ObjectToNodeMap.Find({ InObject, InOuter });
+	if (TreeNode)
+	{
+		// Fast path: object was already in the map, return the tree node only if it was valid.
+		// If invalid, the object was already looked up and does not exist in the item tree map, so return null.
+		return TreeNode->IsValid() ? TreeNode : nullptr;
+	}
+
+	// Slow path: item was never cached in the object node map.
+	TreeNode = FindTreeNode(FAvaSceneItem(InObject, InObject));
+	ObjectToNodeMap.Emplace({ InObject, InOuter }, TreeNode ? *TreeNode : FAvaSceneTreeNode());
+	return TreeNode;
 }
 
 const FAvaSceneTreeNode* FAvaSceneTree::FindTreeNode(const FAvaSceneItem& InItem) const
@@ -124,7 +135,7 @@ const FAvaSceneTreeNode* FAvaSceneTree::FindLowestCommonAncestor(const TArray<co
 
 bool FAvaSceneTree::CompareTreeItemOrder(const FAvaSceneTreeNode* InA, const FAvaSceneTreeNode* InB)
 {
-	if (!InA || !InB || InA->OwningTree != InB->OwningTree)
+	if (!InA || !InB || !InA->OwningTree || InA->OwningTree != InB->OwningTree)
 	{
 		return false;
 	}
diff --git a/Engine/Plugins/VirtualProduction/Avalanche/Source/AvalancheSceneTree/Public/AvaSceneTree.h b/Engine/Plugins/VirtualProduction/Avalanche/Source/AvalancheSceneTree/Public/AvaSceneTree.h
--- a/Engine/Plugins/VirtualProduction/Avalanche/Source/AvalancheSceneTree/Public/AvaSceneTree.h
+++ b/Engine/Plugins/VirtualProduction/Avalanche/Source/AvalancheSceneTree/Public/AvaSceneTree.h
@@ -74,7 +74,7 @@ private:
 	TMap<FAvaSceneItem, FAvaSceneTreeNode> ItemTreeMap;
 
 	/** (Object, Outer) pair to their tree node */
-	TMap<TPair<FObjectKey, FObjectKey>, FAvaSceneTreeNode> ObjectToNodeMap;
+	mutable TMap<TPair<FObjectKey, FObjectKey>, FAvaSceneTreeNode> ObjectToNodeMap;
 
 	/** Whether scene items have been sorted to their latest structure */
 	UPROPERTY()
diff --git a/Engine/Plugins/VirtualProduction/Avalanche/Source/AvalancheSceneTree/Public/AvaSceneTreeNode.h b/Engine/Plugins/VirtualProduction/Avalanche/Source/AvalancheSceneTree/Public/AvaSceneTreeNode.h
--- a/Engine/Plugins/VirtualProduction/Avalanche/Source/AvalancheSceneTree/Public/AvaSceneTreeNode.h
+++ b/Engine/Plugins/VirtualProduction/Avalanche/Source/AvalancheSceneTree/Public/AvaSceneTreeNode.h
@@ -15,6 +15,11 @@ struct AVALANCHESCENETREE_API FAvaSceneTreeNode
 
 	friend FAvaSceneTree;
 
+	bool IsValid() const
+	{
+		return OwningTree && LocalIndex != INDEX_NONE && GlobalIndex != INDEX_NONE;
+	}
+
 	int32 GetLocalIndex()  const { return LocalIndex; }
 
 	int32 GetGlobalIndex() const { return GlobalIndex; }
