Motion Design: trivial fix to passing the wrong object to the scene item ctor

#Synced-CL 51764840

diff --git a/Engine/Plugins/VirtualProduction/Avalanche/Source/AvalancheSceneTree/Private/AvaSceneTree.cpp b/Engine/Plugins/VirtualProduction/Avalanche/Source/AvalancheSceneTree/Private/AvaSceneTree.cpp
--- a/Engine/Plugins/VirtualProduction/Avalanche/Source/AvalancheSceneTree/Private/AvaSceneTree.cpp
+++ b/Engine/Plugins/VirtualProduction/Avalanche/Source/AvalancheSceneTree/Private/AvaSceneTree.cpp
@@ -42,7 +42,7 @@ const FAvaSceneTreeNode* FAvaSceneTree::FindObjectTreeNode(const UObject* InObje
 	}
 
 	// Slow path: item was never cached in the object node map.
-	TreeNode = FindTreeNode(FAvaSceneItem(InObject, InObject));
+	TreeNode = FindTreeNode(FAvaSceneItem(InObject, InOuter));
 	ObjectToNodeMap.Emplace({ InObject, InOuter }, TreeNode ? *TreeNode : FAvaSceneTreeNode());
 	return TreeNode;
 }
