Motion Design: added build system to update motion desigh shapes in parallel

#Synced-CL 51764840

diff --git a/Engine/Plugins/VirtualProduction/Avalanche/Source/AvalancheShapes/Private/AvaShapeBuildSystem.cpp b/Engine/Plugins/VirtualProduction/Avalanche/Source/AvalancheShapes/Private/AvaShapeBuildSystem.cpp
new file mode 100644
--- /dev/null
+++ b/Engine/Plugins/VirtualProduction/Avalanche/Source/AvalancheShapes/Private/AvaShapeBuildSystem.cpp
@@ -0,0 +1,129 @@
+// Copyright Epic Games, Inc. All Rights Reserved.
+
+#include "AvaShapeBuildSystem.h"
+
+#include "Algo/ForEach.h"
+#include "Async/ParallelFor.h"
+#include "Components/DynamicMeshComponent.h"
+#include "DynamicMeshes/AvaShape2DArrowDynMesh.h"
+#include "DynamicMeshes/AvaShapeDynMeshBase.h"
+#include "Modifiers/ActorModifierCoreDefs.h"
+#include "Modifiers/ActorModifierCoreStack.h"
+#include "Subsystems/ActorModifierCoreSubsystem.h"
+
+UAvaShapeBuildSystem* UAvaShapeBuildSystem::FindBuildSystem(TNotNull<const UObject*> InContextObject)
+{
+	if (UWorld* const World = InContextObject->GetWorld())
+	{
+		return World->GetSubsystem<UAvaShapeBuildSystem>();
+	}
+	return nullptr;
+}
+
+void UAvaShapeBuildSystem::QueueShapeBuild(UAvaShapeDynamicMeshBase* InShape)
+{
+	if (!InShape)
+	{
+		return;
+	}
+
+	// Required: shapes need to be added uniquely
+	QueuedShapes.AddUnique(InShape);
+}
+
+bool UAvaShapeBuildSystem::DoesSupportWorldType(const EWorldType::Type InWorldType) const
+{
+	return InWorldType != EWorldType::None;
+}
+
+bool UAvaShapeBuildSystem::IsTickableInEditor() const
+{
+	return true;
+}
+
+void UAvaShapeBuildSystem::Tick(float InDeltaTime)
+{
+	BuildShapes();
+}
+
+TStatId UAvaShapeBuildSystem::GetStatId() const
+{
+	RETURN_QUICK_DECLARE_CYCLE_STAT(UAvaShapeBuildSubsystem, STATGROUP_Tickables);
+}
+
+UAvaShapeBuildSystem::FShapeBuild::FShapeBuild(TNotNull<UAvaShapeDynamicMeshBase*> InShapeBuilder, const UActorModifierCoreSubsystem* InModifierSystem)
+	: Shape(InShapeBuilder)
+	, ModifierLock(InModifierSystem ? InModifierSystem->GetActorModifierStack(InShapeBuilder->GetShapeActor()) : nullptr)
+{
+}
+
+void UAvaShapeBuildSystem::BuildShapes()
+{
+	if (QueuedShapes.IsEmpty())
+	{
+		return;
+	}
+
+	TRACE_CPUPROFILER_EVENT_SCOPE(UAvaShapeBuildSubsystem::BuildShapes);
+
+	// shape builds that can be built from different threads.
+	TArray<FShapeBuild> ParallelBuilds;
+	ParallelBuilds.Reserve(QueuedShapes.Num());
+
+	// shape builds that can only be built in game thread.
+	TArray<FShapeBuild> GameThreadBuilds;
+
+	const UActorModifierCoreSubsystem* const ModifierSubsystem = UActorModifierCoreSubsystem::Get();
+	for (const TWeakObjectPtr<UAvaShapeDynamicMeshBase>& ShapeWeak : QueuedShapes)
+	{
+		UAvaShapeDynamicMeshBase* const Shape = ShapeWeak.Get();
+		if (!Shape)
+		{
+			continue;
+		}
+
+		if (Shape->CanUpdateMeshInAnyThread())
+		{
+			ParallelBuilds.Emplace(Shape, ModifierSubsystem);	
+		}
+		else
+		{
+			GameThreadBuilds.Emplace(Shape, ModifierSubsystem);
+		}
+	}
+	QueuedShapes.Empty();
+
+	if (GameThreadBuilds.IsEmpty() && ParallelBuilds.IsEmpty())
+	{
+		return;
+	}
+
+	Algo::ForEach(GameThreadBuilds, Projection(&FShapeBuild::Shape, &UAvaShapeDynamicMeshBase::PreUpdateMesh));
+	Algo::ForEach(ParallelBuilds  , Projection(&FShapeBuild::Shape, &UAvaShapeDynamicMeshBase::PreUpdateMesh));
+
+	const UAvaShapeDynamicMeshBase::FMeshUpdateContext MeshParallelUpdateContext
+		{
+			.bParallel = true
+		};
+
+	ParallelForWithPreWork(ParallelBuilds.Num(),
+		[&ParallelBuilds, &MeshParallelUpdateContext](int32 InBuildIndex)
+		{
+			ParallelBuilds[InBuildIndex].Shape->UpdateMeshWithContext(MeshParallelUpdateContext);
+		},
+		[&GameThreadBuilds]()
+		{
+			Algo::ForEach(GameThreadBuilds, Projection(&FShapeBuild::Shape, &UAvaShapeDynamicMeshBase::UpdateMesh));
+		});
+
+	Algo::ForEach(GameThreadBuilds, Projection(&FShapeBuild::Shape, &UAvaShapeDynamicMeshBase::PostUpdateMesh));
+	Algo::ForEach(ParallelBuilds  , Projection(&FShapeBuild::Shape, &UAvaShapeDynamicMeshBase::PostUpdateMesh));
+
+	// Unlock modifier execution by emptying the builds
+	// Done before scope exit for profiling
+	{
+		TRACE_CPUPROFILER_EVENT_SCOPE(UAvaShapeBuildSubsystem::BuildShapes_UnlockModifierExecutions);
+		GameThreadBuilds.Empty();
+		ParallelBuilds.Empty();
+	}
+}
diff --git a/Engine/Plugins/VirtualProduction/Avalanche/Source/AvalancheShapes/Private/AvaShapeBuildSystem.h b/Engine/Plugins/VirtualProduction/Avalanche/Source/AvalancheShapes/Private/AvaShapeBuildSystem.h
new file mode 100644
--- /dev/null
+++ b/Engine/Plugins/VirtualProduction/Avalanche/Source/AvalancheShapes/Private/AvaShapeBuildSystem.h
@@ -0,0 +1,47 @@
+// Copyright Epic Games, Inc. All Rights Reserved.
+
+#pragma once
+
+#include "Modifiers/ActorModifierCoreDefs.h"
+#include "Subsystems/WorldSubsystem.h"
+#include "AvaShapeBuildSystem.generated.h"
+
+class UActorModifierCoreSubsystem;
+class UAvaShapeDynamicMeshBase;
+class UDynamicMeshComponent;
+
+UCLASS()
+class UAvaShapeBuildSystem : public UTickableWorldSubsystem
+{
+	GENERATED_BODY()
+
+public:
+	/** Returns the build system for the given context object */
+	static UAvaShapeBuildSystem* FindBuildSystem(TNotNull<const UObject*> InContextObject);
+
+	void QueueShapeBuild(UAvaShapeDynamicMeshBase* InShape);
+
+	//~ Begin UWorldSubsystem
+	virtual bool DoesSupportWorldType(const EWorldType::Type InWorldType) const override;
+	//~ End UWorldSubsystem
+
+	//~ Begin FTickableGameObject
+	virtual bool IsTickableInEditor() const override;
+	virtual void Tick(float InDeltaTime) override;
+	virtual TStatId GetStatId() const override;
+	//~ End FTickableGameObject
+
+private:
+	struct FShapeBuild
+	{
+		explicit FShapeBuild(TNotNull<UAvaShapeDynamicMeshBase*> InShapeBuilder, const UActorModifierCoreSubsystem* InModifierSystem);
+		/** Shape to build */
+		TNotNull<UAvaShapeDynamicMeshBase*> Shape;
+		/** Keep the modifier locked while build takes place */
+		FActorModifierCoreScopedLock ModifierLock;
+	};
+	void BuildShapes();
+
+	UPROPERTY()
+	TArray<TWeakObjectPtr<UAvaShapeDynamicMeshBase>> QueuedShapes;
+};
diff --git a/Engine/Plugins/VirtualProduction/Avalanche/Source/AvalancheShapes/Private/DynamicMeshes/AvaShape2DArrowDynMesh.cpp b/Engine/Plugins/VirtualProduction/Avalanche/Source/AvalancheShapes/Private/DynamicMeshes/AvaShape2DArrowDynMesh.cpp
--- a/Engine/Plugins/VirtualProduction/Avalanche/Source/AvalancheShapes/Private/DynamicMeshes/AvaShape2DArrowDynMesh.cpp
+++ b/Engine/Plugins/VirtualProduction/Avalanche/Source/AvalancheShapes/Private/DynamicMeshes/AvaShape2DArrowDynMesh.cpp
@@ -206,3 +206,8 @@ bool UAvaShape2DArrowDynamicMesh::CreateMesh(FAvaShapeMesh& InMesh)
 
 	return Super::CreateMesh(InMesh);
 }
+
+bool UAvaShape2DArrowDynamicMesh::CanUpdateMeshInAnyThread() const
+{
+	return GetClass() == UAvaShape2DArrowDynamicMesh::StaticClass();
+}
diff --git a/Engine/Plugins/VirtualProduction/Avalanche/Source/AvalancheShapes/Private/DynamicMeshes/AvaShapeChevronDynMesh.cpp b/Engine/Plugins/VirtualProduction/Avalanche/Source/AvalancheShapes/Private/DynamicMeshes/AvaShapeChevronDynMesh.cpp
--- a/Engine/Plugins/VirtualProduction/Avalanche/Source/AvalancheShapes/Private/DynamicMeshes/AvaShapeChevronDynMesh.cpp
+++ b/Engine/Plugins/VirtualProduction/Avalanche/Source/AvalancheShapes/Private/DynamicMeshes/AvaShapeChevronDynMesh.cpp
@@ -74,3 +74,8 @@ bool UAvaShapeChevronDynamicMesh::CreateMesh(FAvaShapeMesh& InMesh)
 
 	return Super::CreateMesh(InMesh);
 }
+
+bool UAvaShapeChevronDynamicMesh::CanUpdateMeshInAnyThread() const
+{
+	return GetClass() == UAvaShapeChevronDynamicMesh::StaticClass();
+}
diff --git a/Engine/Plugins/VirtualProduction/Avalanche/Source/AvalancheShapes/Private/DynamicMeshes/AvaShapeConeDynMesh.cpp b/Engine/Plugins/VirtualProduction/Avalanche/Source/AvalancheShapes/Private/DynamicMeshes/AvaShapeConeDynMesh.cpp
--- a/Engine/Plugins/VirtualProduction/Avalanche/Source/AvalancheShapes/Private/DynamicMeshes/AvaShapeConeDynMesh.cpp
+++ b/Engine/Plugins/VirtualProduction/Avalanche/Source/AvalancheShapes/Private/DynamicMeshes/AvaShapeConeDynMesh.cpp
@@ -512,6 +512,11 @@ bool UAvaShapeConeDynamicMesh::CreateUVs(FAvaShapeMesh& InMesh, FAvaShapeMateria
 	return Super::CreateUVs(InMesh, InParams);
 }
 
+bool UAvaShapeConeDynamicMesh::CanUpdateMeshInAnyThread() const
+{
+	return GetClass() == UAvaShapeConeDynamicMesh::StaticClass();
+}
+
 bool UAvaShapeConeDynamicMesh::CreateMesh(FAvaShapeMesh& InMesh)
 {
 	if (!Super::CreateMesh(InMesh))
diff --git a/Engine/Plugins/VirtualProduction/Avalanche/Source/AvalancheShapes/Private/DynamicMeshes/AvaShapeCubeDynMesh.cpp b/Engine/Plugins/VirtualProduction/Avalanche/Source/AvalancheShapes/Private/DynamicMeshes/AvaShapeCubeDynMesh.cpp
--- a/Engine/Plugins/VirtualProduction/Avalanche/Source/AvalancheShapes/Private/DynamicMeshes/AvaShapeCubeDynMesh.cpp
+++ b/Engine/Plugins/VirtualProduction/Avalanche/Source/AvalancheShapes/Private/DynamicMeshes/AvaShapeCubeDynMesh.cpp
@@ -113,6 +113,11 @@ void UAvaShapeCubeDynamicMesh::OnSizeChanged()
 	MarkAllMeshesDirty();
 }
 
+bool UAvaShapeCubeDynamicMesh::CanUpdateMeshInAnyThread() const
+{
+	return GetClass() == UAvaShapeCubeDynamicMesh::StaticClass();
+}
+
 void UAvaShapeCubeDynamicMesh::OnSegmentChanged()
 {
 	MarkAllMeshesDirty();
diff --git a/Engine/Plugins/VirtualProduction/Avalanche/Source/AvalancheShapes/Private/DynamicMeshes/AvaShapeDynMeshBase.cpp b/Engine/Plugins/VirtualProduction/Avalanche/Source/AvalancheShapes/Private/DynamicMeshes/AvaShapeDynMeshBase.cpp
--- a/Engine/Plugins/VirtualProduction/Avalanche/Source/AvalancheShapes/Private/DynamicMeshes/AvaShapeDynMeshBase.cpp
+++ b/Engine/Plugins/VirtualProduction/Avalanche/Source/AvalancheShapes/Private/DynamicMeshes/AvaShapeDynMeshBase.cpp
@@ -6,6 +6,8 @@
 #include "Async/ParallelFor.h"
 #include "AvaSceneSubsystem.h"
 #include "AvaShapeActor.h"
+#include "AvaShapeBuildSystem.h"
+#include "DynamicMeshMikkTWrapper.h"
 #include "Components/DynamicMeshComponent.h"
 #include "Containers/Ticker.h"
 #include "DynamicMesh/DynamicMesh3.h"
@@ -312,6 +314,8 @@ void UAvaShapeDynamicMeshBase::ToggleGizmo_Implementation(
 
 void UAvaShapeDynamicMeshBase::UpdateMesh()
 {
+	TRACE_CPUPROFILER_EVENT_SCOPE(UAvaShapeDynamicMeshBase::UpdateMesh);
+
 	RestoreCachedMesh(/** Reset */true);
 
 	bool bUpdateMesh = GenerateMesh();
@@ -331,6 +335,11 @@ void UAvaShapeDynamicMeshBase::UpdateMesh()
 	// Apply Update UV, tangents
 	ApplyMeshSectionChanges();
 
+	SaveCachedMesh();
+}
+
+void UAvaShapeDynamicMeshBase::OffsetActorToAlign()
+{
 	// Offset to match alignments
 	if (bHasNewMeshRegenWorldLocation)
 	{
@@ -349,8 +358,83 @@ void UAvaShapeDynamicMeshBase::UpdateMesh()
 #endif
 		MeshRegenWorldLocation = FVector::ZeroVector;
 	}
+}
 
-	SaveCachedMesh();
+void UAvaShapeDynamicMeshBase::PreUpdateMesh()
+{
+	TRACE_CPUPROFILER_EVENT_SCOPE(UAvaShapeDynamicMeshBase::PreUpdateMesh_GameThread);
+
+	if (!bAllMeshDirty && !bAnyMeshDirty)
+	{
+		return;
+	}
+
+	if (UDynamicMeshComponent* DynamicMeshComponent = GetShapeMeshComponent())
+	{
+		if (UDynamicMesh* DynamicMesh = DynamicMeshComponent->GetDynamicMesh())
+		{
+			DynamicMesh->OnPreMeshChanged().Broadcast(DynamicMesh, FDynamicMeshChangeInfo());
+		}
+	}
+
+	SetupMeshes();
+
+	if (!IsDynamicMeshInitialized())
+	{
+		InitializeDynamicMesh();
+	}
+}
+
+void UAvaShapeDynamicMeshBase::UpdateMeshWithContext(const FMeshUpdateContext& InUpdateContext)
+{
+	MeshUpdateContext = &InUpdateContext;
+	UpdateMesh();
+	MeshUpdateContext = nullptr;
+}
+
+void UAvaShapeDynamicMeshBase::PostUpdateMesh()
+{
+	TRACE_CPUPROFILER_EVENT_SCOPE(UAvaShapeDynamicMeshBase::PostUpdateMesh);
+	OffsetActorToAlign();
+
+	if (UDynamicMeshComponent* MeshComponent = GetShapeMeshComponent())
+	{
+		int32 MeshIndex = INDEX_NONE;
+		UMaterialInterface* MeshMaterial;
+
+		auto UpdateBaseMaterial = [MeshComponent, &MeshIndex, &MeshMaterial]()
+			{
+				check(MeshIndex >= 0);
+				if (MeshIndex >= MeshComponent->BaseMaterials.Num())
+				{
+					MeshComponent->BaseMaterials.SetNum(MeshIndex + 1, EAllowShrinking::No);
+				}
+				MeshComponent->BaseMaterials[MeshIndex] = MeshMaterial;
+			};
+
+		for (const TPair<int32, FAvaShapeMeshData>& MeshData : MeshDatas)
+		{
+			const FAvaShapeMesh& ShapeMesh = MeshData.Value.Mesh;
+			MeshIndex = ShapeMesh.GetMeshIndex();
+			MeshMaterial = GetMaterial(ShapeMesh.GetMeshIndex());
+			UpdateBaseMaterial();
+		}
+
+		if (MeshIndex != INDEX_NONE)
+		{
+			// Call SetMaterial on last index/material to mark render state dirty, update physical materials, etc.
+			MeshComponent->SetMaterial(MeshIndex, MeshMaterial);
+		}
+
+		if (UDynamicMesh* Mesh = MeshComponent->GetDynamicMesh())
+		{
+			Mesh->OnMeshChanged().Broadcast(Mesh, FDynamicMeshChangeInfo());
+		}
+		MeshComponent->NotifyMeshModified();
+
+		MeshUpdateStatus = EAvaDynamicMeshUpdateState::UpToDate;
+		OnMeshUpdateFinished();
+	}
 }
 
 void UAvaShapeDynamicMeshBase::ApplyMeshSectionChanges()
@@ -380,23 +464,30 @@ void UAvaShapeDynamicMeshBase::RunUpdate()
 	MeshUpdateStatus = EAvaDynamicMeshUpdateState::UpdateInProgress;
 
 	auto ExecUpdate = [this](float=0.f)->bool
-	{
-		if (UDynamicMeshComponent* DynamicMeshComponent = GetShapeMeshComponent())
 		{
-			FActorModifierCoreScopedLock ModifierLock(UActorModifierCoreSubsystem::Get()->GetActorModifierStack(GetShapeActor()));
-			UpdateMesh();
-			DynamicMeshComponent->NotifyMeshModified();
-			MeshUpdateStatus = EAvaDynamicMeshUpdateState::UpToDate;
-			OnMeshUpdateFinished();
-			return false; // one time off
-		}
-		MeshUpdateStatus = EAvaDynamicMeshUpdateState::UpdateRequired;
-		return true;
-	};
+			if (GetShapeMeshComponent())
+			{
+				PreUpdateMesh();
+				UpdateMesh();
+				PostUpdateMesh();
+			}
+			else
+			{
+				MeshUpdateStatus = EAvaDynamicMeshUpdateState::UpdateRequired;
+			}
+			return false;
+		};
 
 	if (GetRunAsync())
 	{
-		GetTicker().AddTicker(FTickerDelegate::CreateWeakLambda(this, ExecUpdate));
+		if (UAvaShapeBuildSystem* const BuildSystem = UAvaShapeBuildSystem::FindBuildSystem(this))
+		{
+			BuildSystem->QueueShapeBuild(this);
+		}
+		else
+		{
+			GetTicker().AddTicker(FTickerDelegate::CreateWeakLambda(this, ExecUpdate));	
+		}
 	}
 	else
 	{
@@ -618,8 +709,6 @@ bool UAvaShapeDynamicMeshBase::GenerateMesh()
 		return false;
 	}
 
-	SetupMeshes();
-
 	bAllMeshDirty = false;
 	bAnyMeshDirty = false;
 
@@ -1108,10 +1197,34 @@ bool UAvaShapeDynamicMeshBase::GenerateTangents()
 		return false;
 	}
 
-	FGeometryScriptTangentsOptions TangentOptions;
-	TangentOptions.UVLayer = 0;
+	// Copied from UGeometryScriptLibrary_MeshNormalsFunctions::ComputeTangents to not broadcast any change
+	UE::Private::AvaShape::EditMesh(this,
+		[this](FDynamicMesh3& EditMesh)
+		{
+			constexpr int32 UVLayer = 0;
+			if (!EditMesh.HasAttributes() || EditMesh.Attributes()->NumUVLayers() <= UVLayer || EditMesh.Attributes()->NumNormalLayers() == 0)
+			{
+				return;
+			}
+			if (EditMesh.Attributes()->HasTangentSpace() == false)
+			{
+				EditMesh.Attributes()->EnableTangents();
+			}
 
-	UGeometryScriptLibrary_MeshNormalsFunctions::ComputeTangents(DynMeshComp->GetDynamicMesh(), TangentOptions);
+			FComputeTangentsOptions TangentOptions;
+			TangentOptions.bAveraged = true;
+			TangentOptions.bParallel = !MeshUpdateContext || !MeshUpdateContext->bParallel; // do not do parallel operation if already parallel.
+
+			FMeshTangentsd Tangents(&EditMesh);
+			Tangents.ComputeTriVertexTangents(
+				EditMesh.Attributes()->PrimaryNormals(),
+				EditMesh.Attributes()->GetUVLayer(UVLayer),
+				TangentOptions);
+			Tangents.CopyToOverlays(EditMesh);
+		}
+		, EDynamicMeshChangeType::GeneralEdit
+		, EDynamicMeshAttributeChangeFlags::Unknown
+		, /*Broadcast*/false);
 
 	return true;
 }
@@ -1519,11 +1632,6 @@ bool UAvaShapeDynamicMeshBase::CreateDynamicMesh(FAvaShapeMesh& InMesh)
 		return false;
 	}
 
-	if (!IsDynamicMeshInitialized())
-	{
-		InitializeDynamicMesh();
-	}
-
 	// vertices and triangles should not exists if we create them
 	if (InMesh.VerticeIds.Num() != 0 || InMesh.TriangleIds.Num() != 0)
 	{
@@ -1609,9 +1717,6 @@ bool UAvaShapeDynamicMeshBase::CreateDynamicMesh(FAvaShapeMesh& InMesh)
 		, EDynamicMeshAttributeChangeFlags::MeshTopology
 		, /** Broadcast */false
 	);
-
-	ProcMesh->SetMaterial(InMesh.GetMeshIndex(), GetMaterial(InMesh.GetMeshIndex()));
-
 	return true;
 }
 
@@ -1624,11 +1729,6 @@ bool UAvaShapeDynamicMeshBase::UpdateDynamicMesh(FAvaShapeMesh& InMesh)
 		return false;
 	}
 
-	if (!IsDynamicMeshInitialized())
-	{
-		InitializeDynamicMesh();
-	}
-
 	// different topology, create mesh section
 	if (InMesh.GetMeshIndex() >= ProcMesh->GetMesh()->MaxGroupID() || InMesh.Vertices.Num() != InMesh.VerticeIds.Num())
 	{
@@ -1652,23 +1752,32 @@ bool UAvaShapeDynamicMeshBase::UpdateDynamicMesh(FAvaShapeMesh& InMesh)
 		{
 			FDynamicMeshNormalOverlay* const NormalOverlay = EditMesh.Attributes()->PrimaryNormals();
 
-			ParallelFor(InMesh.Vertices.Num(), [InMesh, &EditMesh, NormalOverlay](int32 Idx)
+			auto DoWork = [InMesh, &EditMesh, NormalOverlay](int32 Idx)
+				{
+					// vertices
+					int32 VId = InMesh.VerticeIds[Idx];
+					EditMesh.SetVertex(VId, InMesh.Vertices[Idx]);
+					// normals
+					int32 NId = InMesh.NormalIds[Idx];
+					NormalOverlay->SetElement(NId, static_cast<FVector3f>(InMesh.Normals[Idx]));
+				};
+
+			if (MeshUpdateContext && MeshUpdateContext->bParallel)
 			{
-				// vertices
-				int32 VId = InMesh.VerticeIds[Idx];
-				EditMesh.SetVertex(VId, InMesh.Vertices[Idx]);
-				// normals
-				int32 NId = InMesh.NormalIds[Idx];
-				NormalOverlay->SetElement(NId, static_cast<FVector3f>(InMesh.Normals[Idx]));
-			});
+				for (int32 Idx = 0; Idx < InMesh.Vertices.Num(); ++Idx)
+				{
+					DoWork(Idx);
+				}
+			}
+			else
+			{
+				ParallelFor(InMesh.Vertices.Num(), DoWork);
+			}
 		}
 		, EDynamicMeshChangeType::MeshVertexChange
 		, EDynamicMeshAttributeChangeFlags::Unknown
 		, /** Broadcast */false
 	);
-
-	ProcMesh->SetMaterial(InMesh.GetMeshIndex(), GetMaterial(InMesh.GetMeshIndex()));
-
 	return true;
 }
 
diff --git a/Engine/Plugins/VirtualProduction/Avalanche/Source/AvalancheShapes/Private/DynamicMeshes/AvaShapeEllipseDynMesh.cpp b/Engine/Plugins/VirtualProduction/Avalanche/Source/AvalancheShapes/Private/DynamicMeshes/AvaShapeEllipseDynMesh.cpp
--- a/Engine/Plugins/VirtualProduction/Avalanche/Source/AvalancheShapes/Private/DynamicMeshes/AvaShapeEllipseDynMesh.cpp
+++ b/Engine/Plugins/VirtualProduction/Avalanche/Source/AvalancheShapes/Private/DynamicMeshes/AvaShapeEllipseDynMesh.cpp
@@ -136,3 +136,8 @@ bool UAvaShapeEllipseDynamicMesh::CreateMesh(FAvaShapeMesh& InMesh)
 
 	return Super::CreateMesh(InMesh);
 }
+
+bool UAvaShapeEllipseDynamicMesh::CanUpdateMeshInAnyThread() const
+{
+	return GetClass() == UAvaShapeEllipseDynamicMesh::StaticClass();
+}
diff --git a/Engine/Plugins/VirtualProduction/Avalanche/Source/AvalancheShapes/Private/DynamicMeshes/AvaShapeIrregularPolygonDynMesh.cpp b/Engine/Plugins/VirtualProduction/Avalanche/Source/AvalancheShapes/Private/DynamicMeshes/AvaShapeIrregularPolygonDynMesh.cpp
--- a/Engine/Plugins/VirtualProduction/Avalanche/Source/AvalancheShapes/Private/DynamicMeshes/AvaShapeIrregularPolygonDynMesh.cpp
+++ b/Engine/Plugins/VirtualProduction/Avalanche/Source/AvalancheShapes/Private/DynamicMeshes/AvaShapeIrregularPolygonDynMesh.cpp
@@ -1024,3 +1024,8 @@ bool UAvaShapeIrregularPolygonDynamicMesh::CreateMesh(FAvaShapeMesh& InMesh)
 
 	return Super::CreateMesh(InMesh);
 }
+
+bool UAvaShapeIrregularPolygonDynamicMesh::CanUpdateMeshInAnyThread() const
+{
+	return GetClass() == UAvaShapeIrregularPolygonDynamicMesh::StaticClass();
+}
diff --git a/Engine/Plugins/VirtualProduction/Avalanche/Source/AvalancheShapes/Private/DynamicMeshes/AvaShapeLineDynMesh.cpp b/Engine/Plugins/VirtualProduction/Avalanche/Source/AvalancheShapes/Private/DynamicMeshes/AvaShapeLineDynMesh.cpp
--- a/Engine/Plugins/VirtualProduction/Avalanche/Source/AvalancheShapes/Private/DynamicMeshes/AvaShapeLineDynMesh.cpp
+++ b/Engine/Plugins/VirtualProduction/Avalanche/Source/AvalancheShapes/Private/DynamicMeshes/AvaShapeLineDynMesh.cpp
@@ -44,6 +44,11 @@ void UAvaShapeLineDynamicMesh::GenerateBorderVertices(TArray<FVector2D>& BorderV
 	BorderVertices.Append(LineVertices);
 }
 
+bool UAvaShapeLineDynamicMesh::CanUpdateMeshInAnyThread() const
+{
+	return GetClass() == UAvaShapeLineDynamicMesh::StaticClass();
+}
+
 void UAvaShapeLineDynamicMesh::OnLineWidthChanged()
 {
 	if (LineWidth < 0.f)
diff --git a/Engine/Plugins/VirtualProduction/Avalanche/Source/AvalancheShapes/Private/DynamicMeshes/AvaShapeNGonDynMesh.cpp b/Engine/Plugins/VirtualProduction/Avalanche/Source/AvalancheShapes/Private/DynamicMeshes/AvaShapeNGonDynMesh.cpp
--- a/Engine/Plugins/VirtualProduction/Avalanche/Source/AvalancheShapes/Private/DynamicMeshes/AvaShapeNGonDynMesh.cpp
+++ b/Engine/Plugins/VirtualProduction/Avalanche/Source/AvalancheShapes/Private/DynamicMeshes/AvaShapeNGonDynMesh.cpp
@@ -57,3 +57,8 @@ void UAvaShapeNGonDynamicMesh::GenerateBorderVertices(TArray<FVector2D>& BorderV
 		BorderVertices.Add(RotatedZero);
 	}
 }
+
+bool UAvaShapeNGonDynamicMesh::CanUpdateMeshInAnyThread() const
+{
+	return GetClass() == UAvaShapeNGonDynamicMesh::StaticClass();
+}
diff --git a/Engine/Plugins/VirtualProduction/Avalanche/Source/AvalancheShapes/Private/DynamicMeshes/AvaShapeRectangleDynMesh.cpp b/Engine/Plugins/VirtualProduction/Avalanche/Source/AvalancheShapes/Private/DynamicMeshes/AvaShapeRectangleDynMesh.cpp
--- a/Engine/Plugins/VirtualProduction/Avalanche/Source/AvalancheShapes/Private/DynamicMeshes/AvaShapeRectangleDynMesh.cpp
+++ b/Engine/Plugins/VirtualProduction/Avalanche/Source/AvalancheShapes/Private/DynamicMeshes/AvaShapeRectangleDynMesh.cpp
@@ -1260,6 +1260,11 @@ void UAvaShapeRectangleDynamicMesh::OnSizeChanged()
 	}
 }
 
+bool UAvaShapeRectangleDynamicMesh::CanUpdateMeshInAnyThread() const
+{
+	return GetClass() == UAvaShapeRectangleDynamicMesh::StaticClass();
+}
+
 bool UAvaShapeRectangleDynamicMesh::CreateMesh(FAvaShapeMesh& InMesh)
 {
 	if (InMesh.GetMeshIndex() == MESH_INDEX_PRIMARY)
diff --git a/Engine/Plugins/VirtualProduction/Avalanche/Source/AvalancheShapes/Private/DynamicMeshes/AvaShapeRingDynMesh.cpp b/Engine/Plugins/VirtualProduction/Avalanche/Source/AvalancheShapes/Private/DynamicMeshes/AvaShapeRingDynMesh.cpp
--- a/Engine/Plugins/VirtualProduction/Avalanche/Source/AvalancheShapes/Private/DynamicMeshes/AvaShapeRingDynMesh.cpp
+++ b/Engine/Plugins/VirtualProduction/Avalanche/Source/AvalancheShapes/Private/DynamicMeshes/AvaShapeRingDynMesh.cpp
@@ -150,3 +150,8 @@ bool UAvaShapeRingDynamicMesh::CreateMesh(FAvaShapeMesh& InMesh)
 
 	return Super::CreateMesh(InMesh);
 }
+
+bool UAvaShapeRingDynamicMesh::CanUpdateMeshInAnyThread() const
+{
+	return GetClass() == UAvaShapeRingDynamicMesh::StaticClass();
+}
diff --git a/Engine/Plugins/VirtualProduction/Avalanche/Source/AvalancheShapes/Private/DynamicMeshes/AvaShapeSphereDynMesh.cpp b/Engine/Plugins/VirtualProduction/Avalanche/Source/AvalancheShapes/Private/DynamicMeshes/AvaShapeSphereDynMesh.cpp
--- a/Engine/Plugins/VirtualProduction/Avalanche/Source/AvalancheShapes/Private/DynamicMeshes/AvaShapeSphereDynMesh.cpp
+++ b/Engine/Plugins/VirtualProduction/Avalanche/Source/AvalancheShapes/Private/DynamicMeshes/AvaShapeSphereDynMesh.cpp
@@ -331,6 +331,11 @@ bool UAvaShapeSphereDynamicMesh::CreateUVs(FAvaShapeMesh& InMesh, FAvaShapeMater
 	return Super::CreateUVs(InMesh, InParams);
 }
 
+bool UAvaShapeSphereDynamicMesh::CanUpdateMeshInAnyThread() const
+{
+	return GetClass() == UAvaShapeSphereDynamicMesh::StaticClass();
+}
+
 bool UAvaShapeSphereDynamicMesh::CreateMesh(FAvaShapeMesh& InMesh)
 {
 	if (!Super::CreateMesh(InMesh))
diff --git a/Engine/Plugins/VirtualProduction/Avalanche/Source/AvalancheShapes/Private/DynamicMeshes/AvaShapeStarDynMesh.cpp b/Engine/Plugins/VirtualProduction/Avalanche/Source/AvalancheShapes/Private/DynamicMeshes/AvaShapeStarDynMesh.cpp
--- a/Engine/Plugins/VirtualProduction/Avalanche/Source/AvalancheShapes/Private/DynamicMeshes/AvaShapeStarDynMesh.cpp
+++ b/Engine/Plugins/VirtualProduction/Avalanche/Source/AvalancheShapes/Private/DynamicMeshes/AvaShapeStarDynMesh.cpp
@@ -99,3 +99,8 @@ bool UAvaShapeStarDynamicMesh::IsMeshVisible(int32 MeshIndex)
 			return true;
 	}
 }
+
+bool UAvaShapeStarDynamicMesh::CanUpdateMeshInAnyThread() const
+{
+	return GetClass() == UAvaShapeStarDynamicMesh::StaticClass();
+}
diff --git a/Engine/Plugins/VirtualProduction/Avalanche/Source/AvalancheShapes/Private/DynamicMeshes/AvaShapeTorusDynMesh.cpp b/Engine/Plugins/VirtualProduction/Avalanche/Source/AvalancheShapes/Private/DynamicMeshes/AvaShapeTorusDynMesh.cpp
--- a/Engine/Plugins/VirtualProduction/Avalanche/Source/AvalancheShapes/Private/DynamicMeshes/AvaShapeTorusDynMesh.cpp
+++ b/Engine/Plugins/VirtualProduction/Avalanche/Source/AvalancheShapes/Private/DynamicMeshes/AvaShapeTorusDynMesh.cpp
@@ -453,6 +453,11 @@ bool UAvaShapeTorusDynamicMesh::CreateUVs(FAvaShapeMesh& InMesh, FAvaShapeMateri
 	return Super::CreateUVs(InMesh, InParams);
 }
 
+bool UAvaShapeTorusDynamicMesh::CanUpdateMeshInAnyThread() const
+{
+	return GetClass() == UAvaShapeTorusDynamicMesh::StaticClass();
+}
+
 bool UAvaShapeTorusDynamicMesh::CreateMesh(FAvaShapeMesh& InMesh)
 {
 	if (!Super::CreateMesh(InMesh))
diff --git a/Engine/Plugins/VirtualProduction/Avalanche/Source/AvalancheShapes/Public/DynamicMeshes/AvaShape2DArrowDynMesh.h b/Engine/Plugins/VirtualProduction/Avalanche/Source/AvalancheShapes/Public/DynamicMeshes/AvaShape2DArrowDynMesh.h
--- a/Engine/Plugins/VirtualProduction/Avalanche/Source/AvalancheShapes/Public/DynamicMeshes/AvaShape2DArrowDynMesh.h
+++ b/Engine/Plugins/VirtualProduction/Avalanche/Source/AvalancheShapes/Public/DynamicMeshes/AvaShape2DArrowDynMesh.h
@@ -58,6 +58,7 @@ protected:
 
 	virtual bool IsMeshVisible(int32 MeshIndex) override;
 	virtual bool CreateMesh(FAvaShapeMesh& InMesh) override;
+	virtual bool CanUpdateMeshInAnyThread() const override;
 
 	// represents the ratio for the arrow and the line, 0.6 means 60% arrow and 40% line
 	UPROPERTY(EditInstanceOnly, BlueprintReadWrite, Category="Shape", Setter, Getter, meta=(ClampMin="0.0", ClampMax="1.0", DisplayName="Ratio Arrow Line", AllowPrivateAccess="true"))
diff --git a/Engine/Plugins/VirtualProduction/Avalanche/Source/AvalancheShapes/Public/DynamicMeshes/AvaShapeChevronDynMesh.h b/Engine/Plugins/VirtualProduction/Avalanche/Source/AvalancheShapes/Public/DynamicMeshes/AvaShapeChevronDynMesh.h
--- a/Engine/Plugins/VirtualProduction/Avalanche/Source/AvalancheShapes/Public/DynamicMeshes/AvaShapeChevronDynMesh.h
+++ b/Engine/Plugins/VirtualProduction/Avalanche/Source/AvalancheShapes/Public/DynamicMeshes/AvaShapeChevronDynMesh.h
@@ -47,6 +47,7 @@ protected:
 
 	virtual bool IsMeshVisible(int32 MeshIndex) override;
 	virtual bool CreateMesh(FAvaShapeMesh& InMesh) override;
+	virtual bool CanUpdateMeshInAnyThread() const override;
 
 	// represents the ratio for the chevron size
 	UPROPERTY(EditInstanceOnly, BlueprintReadWrite, Setter, Getter, Category="Shape", meta=(ClampMin="0.0", ClampMax="0.99", DisplayName="Ratio Chevron", AllowPrivateAccess="true"))
diff --git a/Engine/Plugins/VirtualProduction/Avalanche/Source/AvalancheShapes/Public/DynamicMeshes/AvaShapeConeDynMesh.h b/Engine/Plugins/VirtualProduction/Avalanche/Source/AvalancheShapes/Public/DynamicMeshes/AvaShapeConeDynMesh.h
--- a/Engine/Plugins/VirtualProduction/Avalanche/Source/AvalancheShapes/Public/DynamicMeshes/AvaShapeConeDynMesh.h
+++ b/Engine/Plugins/VirtualProduction/Avalanche/Source/AvalancheShapes/Public/DynamicMeshes/AvaShapeConeDynMesh.h
@@ -107,6 +107,7 @@ protected:
 	virtual bool IsMeshVisible(int32 MeshIndex) override;
 	virtual bool CreateMesh(FAvaShapeMesh& InMesh) override;
 	virtual bool CreateUVs(FAvaShapeMesh& InMesh, FAvaShapeMaterialUVParameters& InParams) override;
+	virtual bool CanUpdateMeshInAnyThread() const override;
 
 	// The number of sides around the base of the cone
 	UPROPERTY(EditInstanceOnly, BlueprintReadWrite, Setter, Getter, Category="Shape", meta=(ClampMin="3.0", ClampMax="128.0", DisplayName="Sides", AllowPrivateAccess="true"))
diff --git a/Engine/Plugins/VirtualProduction/Avalanche/Source/AvalancheShapes/Public/DynamicMeshes/AvaShapeCubeDynMesh.h b/Engine/Plugins/VirtualProduction/Avalanche/Source/AvalancheShapes/Public/DynamicMeshes/AvaShapeCubeDynMesh.h
--- a/Engine/Plugins/VirtualProduction/Avalanche/Source/AvalancheShapes/Public/DynamicMeshes/AvaShapeCubeDynMesh.h
+++ b/Engine/Plugins/VirtualProduction/Avalanche/Source/AvalancheShapes/Public/DynamicMeshes/AvaShapeCubeDynMesh.h
@@ -85,6 +85,7 @@ protected:
 	virtual void RegisterMeshes() override;
 	virtual bool IsMeshVisible(int32 MeshIndex) override;
 	virtual void OnSizeChanged() override;
+	virtual bool CanUpdateMeshInAnyThread() const override;
 
 	virtual bool GenerateBottomMeshSections(FAvaShapeMesh& BottomMesh);
 
diff --git a/Engine/Plugins/VirtualProduction/Avalanche/Source/AvalancheShapes/Public/DynamicMeshes/AvaShapeDynMeshBase.h b/Engine/Plugins/VirtualProduction/Avalanche/Source/AvalancheShapes/Public/DynamicMeshes/AvaShapeDynMeshBase.h
--- a/Engine/Plugins/VirtualProduction/Avalanche/Source/AvalancheShapes/Public/DynamicMeshes/AvaShapeDynMeshBase.h
+++ b/Engine/Plugins/VirtualProduction/Avalanche/Source/AvalancheShapes/Public/DynamicMeshes/AvaShapeDynMeshBase.h
@@ -32,10 +32,11 @@ class UAvaShapeDynamicMeshBase
 {
 	GENERATED_BODY()
 
+	friend class AAvaShapeActor;
 	friend class FAvaDynMeshBaseDetailCustomization;
 	friend class FAvaMeshesDetailCustomization;
 	friend class FAvaShapeDynamicMeshVisualizer;
-	friend class AAvaShapeActor;
+	friend class UAvaShapeBuildSystem;
 
 public:
 	static constexpr float MinSizeValue = UE_SMALL_NUMBER;
@@ -378,9 +379,35 @@ protected:
 	/** Flag a mesh section dirty to regenerate it */
 	void MarkMeshDirty(int32 MeshIndex);
 
-	// Takes all dirty flags, regenerate whatever is needed and updates the DMC, called by RunUpdate or RequireUpdate
+	/**
+	 * Determines whether shape mesh can be done from a different thread.
+	 * All built-in meshes are safe to build in a different thread. This virtual would need to be overriden for any derived class that ensures thread safety.
+	 */
+	virtual bool CanUpdateMeshInAnyThread() const
+	{
+		return false;
+	}
+
+	/** Takes all dirty flags, regenerate whatever is needed and updates the DMC, called by RunUpdate or RequireUpdate */
 	void UpdateMesh();
 
+	/** Offsets actor to account for the mesh generated location */
+	void OffsetActorToAlign();
+
+	/** Part of update flow that can only be performed safely on the game thread. Should be called before UpdateMesh_AnyThread */
+	void PreUpdateMesh();
+
+	struct FMeshUpdateContext
+	{
+		/** Determines whether the update is taking place in parallel (potentially not the game thread) */
+		bool bParallel = false;
+	};
+	/** UpdateMesh() that can be performed on a different thread */
+	void UpdateMeshWithContext(const FMeshUpdateContext& InUpdateContext);
+
+	/** Part of update flow that can only be performed safely on the game thread. Should be called after UpdateMesh_AnyThread */
+	void PostUpdateMesh();
+
 	/** Pushed mesh section data updates to dynamic mesh and creates/updates geometry data */
 	void ApplyMeshSectionChanges();
 
@@ -463,8 +490,8 @@ private:
 	// used to quickly restore mesh that has been modified after OnMeshUpdateFinished was called (eg extrude, masking)
 	TOptional<UE::Geometry::FDynamicMesh3> CachedMesh;
 
-	// only change this value from the base class to avoid thread lock, use RunUpdate() or RequireUpdate()
-	std::atomic<EAvaDynamicMeshUpdateState> MeshUpdateStatus = EAvaDynamicMeshUpdateState::UpToDate;
+	/** the context while mesh is being built */
+	const FMeshUpdateContext* MeshUpdateContext = nullptr;
 
 	/** Do not access directly use getter to load it if nullptr */
 	TWeakObjectPtr<UDynamicMeshComponent> CachedComponent;
@@ -478,6 +505,9 @@ private:
 	TMap<int32, FAvaShapeMeshData> NonGizmoMeshData;
 #endif
 
+	// only change this value from the base class to avoid thread lock, use RunUpdate() or RequireUpdate()
+	EAvaDynamicMeshUpdateState MeshUpdateStatus = EAvaDynamicMeshUpdateState::UpToDate;
+
 	bool bHasNewMeshRegenWorldLocation;
 
 	// flag to register meshes only once
diff --git a/Engine/Plugins/VirtualProduction/Avalanche/Source/AvalancheShapes/Public/DynamicMeshes/AvaShapeEllipseDynMesh.h b/Engine/Plugins/VirtualProduction/Avalanche/Source/AvalancheShapes/Public/DynamicMeshes/AvaShapeEllipseDynMesh.h
--- a/Engine/Plugins/VirtualProduction/Avalanche/Source/AvalancheShapes/Public/DynamicMeshes/AvaShapeEllipseDynMesh.h
+++ b/Engine/Plugins/VirtualProduction/Avalanche/Source/AvalancheShapes/Public/DynamicMeshes/AvaShapeEllipseDynMesh.h
@@ -67,6 +67,7 @@ protected:
 
 	virtual bool IsMeshVisible(int32 MeshIndex) override;
 	virtual bool CreateMesh(FAvaShapeMesh& InMesh) override;
+	virtual bool CanUpdateMeshInAnyThread() const override;
 
 	// The number of sides for the shape
 	UPROPERTY(EditInstanceOnly, BlueprintReadWrite, Setter, Getter, Category="Shape", meta=(ClampMin="3.0", ClampMax="128.0", DisplayName="Sides", AllowPrivateAccess="true"))
diff --git a/Engine/Plugins/VirtualProduction/Avalanche/Source/AvalancheShapes/Public/DynamicMeshes/AvaShapeIrregularPolygonDynMesh.h b/Engine/Plugins/VirtualProduction/Avalanche/Source/AvalancheShapes/Public/DynamicMeshes/AvaShapeIrregularPolygonDynMesh.h
--- a/Engine/Plugins/VirtualProduction/Avalanche/Source/AvalancheShapes/Public/DynamicMeshes/AvaShapeIrregularPolygonDynMesh.h
+++ b/Engine/Plugins/VirtualProduction/Avalanche/Source/AvalancheShapes/Public/DynamicMeshes/AvaShapeIrregularPolygonDynMesh.h
@@ -160,6 +160,7 @@ protected:
 
 	virtual bool ClearMesh() override;
 	virtual bool CreateMesh(FAvaShapeMesh& InMesh) override;
+	virtual bool CanUpdateMeshInAnyThread() const override;
 
 	UPROPERTY(EditInstanceOnly, BlueprintReadWrite, Setter, Getter, Category="Shape", meta=(ClampMin="0.0",ClampMax="1.0", AllowPrivateAccess="true"))
 	float GlobalBevelSize = 0.f;
diff --git a/Engine/Plugins/VirtualProduction/Avalanche/Source/AvalancheShapes/Public/DynamicMeshes/AvaShapeLineDynMesh.h b/Engine/Plugins/VirtualProduction/Avalanche/Source/AvalancheShapes/Public/DynamicMeshes/AvaShapeLineDynMesh.h
--- a/Engine/Plugins/VirtualProduction/Avalanche/Source/AvalancheShapes/Public/DynamicMeshes/AvaShapeLineDynMesh.h
+++ b/Engine/Plugins/VirtualProduction/Avalanche/Source/AvalancheShapes/Public/DynamicMeshes/AvaShapeLineDynMesh.h
@@ -61,6 +61,7 @@ protected:
 
 	virtual bool IsMeshVisible(int32 MeshIndex) override;
 	virtual void GenerateBorderVertices(TArray<FVector2D>& BorderVertices) override;
+	virtual bool CanUpdateMeshInAnyThread() const override;
 
 	UPROPERTY(EditInstanceOnly, BlueprintReadWrite, Setter, Getter, Category="Shape", meta=(ClampMin="0.0", AllowPrivateAccess="true"))
 	float LineWidth;
diff --git a/Engine/Plugins/VirtualProduction/Avalanche/Source/AvalancheShapes/Public/DynamicMeshes/AvaShapeNGonDynMesh.h b/Engine/Plugins/VirtualProduction/Avalanche/Source/AvalancheShapes/Public/DynamicMeshes/AvaShapeNGonDynMesh.h
--- a/Engine/Plugins/VirtualProduction/Avalanche/Source/AvalancheShapes/Public/DynamicMeshes/AvaShapeNGonDynMesh.h
+++ b/Engine/Plugins/VirtualProduction/Avalanche/Source/AvalancheShapes/Public/DynamicMeshes/AvaShapeNGonDynMesh.h
@@ -46,8 +46,8 @@ protected:
 	//~ End UObject
 
 	virtual void OnNumSidesChanged();
-
 	virtual void GenerateBorderVertices(TArray<FVector2D>& BorderVertices) override;
+	virtual bool CanUpdateMeshInAnyThread() const override;
 
 	UPROPERTY(EditInstanceOnly, BlueprintReadWrite, Setter, Getter, Category="Shape", meta=(ClampMin="3.0", ClampMax="128.0", DisplayName="Sides", AllowPrivateAccess="true"))
 	uint8 NumSides;
diff --git a/Engine/Plugins/VirtualProduction/Avalanche/Source/AvalancheShapes/Public/DynamicMeshes/AvaShapeRectangleDynMesh.h b/Engine/Plugins/VirtualProduction/Avalanche/Source/AvalancheShapes/Public/DynamicMeshes/AvaShapeRectangleDynMesh.h
--- a/Engine/Plugins/VirtualProduction/Avalanche/Source/AvalancheShapes/Public/DynamicMeshes/AvaShapeRectangleDynMesh.h
+++ b/Engine/Plugins/VirtualProduction/Avalanche/Source/AvalancheShapes/Public/DynamicMeshes/AvaShapeRectangleDynMesh.h
@@ -232,8 +232,8 @@ protected:
 	bool GenerateBaseMeshSections(FAvaShapeMesh& BaseMesh);
 
 	virtual bool CreateMesh(FAvaShapeMesh& InMesh) override;
-
 	virtual void OnSizeChanged() override;
+	virtual bool CanUpdateMeshInAnyThread() const override;
 
 	UPROPERTY(EditInstanceOnly, BlueprintReadWrite, Setter, Getter, Category="Shape", meta=(AllowPrivateAccess="true"))
 	EAvaHorizontalAlignment HorizontalAlignment = EAvaHorizontalAlignment::Center;
diff --git a/Engine/Plugins/VirtualProduction/Avalanche/Source/AvalancheShapes/Public/DynamicMeshes/AvaShapeRingDynMesh.h b/Engine/Plugins/VirtualProduction/Avalanche/Source/AvalancheShapes/Public/DynamicMeshes/AvaShapeRingDynMesh.h
--- a/Engine/Plugins/VirtualProduction/Avalanche/Source/AvalancheShapes/Public/DynamicMeshes/AvaShapeRingDynMesh.h
+++ b/Engine/Plugins/VirtualProduction/Avalanche/Source/AvalancheShapes/Public/DynamicMeshes/AvaShapeRingDynMesh.h
@@ -75,6 +75,7 @@ protected:
 
 	virtual bool IsMeshVisible(int32 MeshIndex) override;
 	virtual bool CreateMesh(FAvaShapeMesh& InMesh) override;
+	virtual bool CanUpdateMeshInAnyThread() const override;
 
 	UPROPERTY(EditInstanceOnly, BlueprintReadWrite, Setter, Getter, Category="Shape", meta=(ClampMin="3.0", ClampMax="128.0", DisplayName="Sides", AllowPrivateAccess="true"))
 	uint8 NumSides;
diff --git a/Engine/Plugins/VirtualProduction/Avalanche/Source/AvalancheShapes/Public/DynamicMeshes/AvaShapeSphereDynMesh.h b/Engine/Plugins/VirtualProduction/Avalanche/Source/AvalancheShapes/Public/DynamicMeshes/AvaShapeSphereDynMesh.h
--- a/Engine/Plugins/VirtualProduction/Avalanche/Source/AvalancheShapes/Public/DynamicMeshes/AvaShapeSphereDynMesh.h
+++ b/Engine/Plugins/VirtualProduction/Avalanche/Source/AvalancheShapes/Public/DynamicMeshes/AvaShapeSphereDynMesh.h
@@ -115,6 +115,7 @@ protected:
 	virtual bool IsMeshVisible(int32 MeshIndex) override;
 	virtual bool CreateMesh(FAvaShapeMesh& InMesh) override;
 	virtual bool CreateUVs(FAvaShapeMesh& InMesh, FAvaShapeMaterialUVParameters& InParams) override;
+	virtual bool CanUpdateMeshInAnyThread() const override;
 
 	// represents the longitude (Z) angle in degree for the sphere at the start
 	UPROPERTY(EditInstanceOnly, BlueprintReadWrite, Setter, Getter, Category="Shape", meta=(ClampMin="0.0", ClampMax="180.0", DisplayName="Start Longitude degree", AllowPrivateAccess="true"))
diff --git a/Engine/Plugins/VirtualProduction/Avalanche/Source/AvalancheShapes/Public/DynamicMeshes/AvaShapeStarDynMesh.h b/Engine/Plugins/VirtualProduction/Avalanche/Source/AvalancheShapes/Public/DynamicMeshes/AvaShapeStarDynMesh.h
--- a/Engine/Plugins/VirtualProduction/Avalanche/Source/AvalancheShapes/Public/DynamicMeshes/AvaShapeStarDynMesh.h
+++ b/Engine/Plugins/VirtualProduction/Avalanche/Source/AvalancheShapes/Public/DynamicMeshes/AvaShapeStarDynMesh.h
@@ -59,6 +59,7 @@ protected:
 
 	virtual bool UseCenteredVertex() override { return true; }
 	virtual bool IsMeshVisible(int32 MeshIndex) override;
+	virtual bool CanUpdateMeshInAnyThread() const override;
 
 	UPROPERTY(EditInstanceOnly, BlueprintReadWrite, Setter, Getter, Category="Shape", meta=(ClampMin="2.0", ClampMax="128.0", DisplayName="Points", AllowPrivateAccess="true"))
 	uint8 NumPoints = 0;
diff --git a/Engine/Plugins/VirtualProduction/Avalanche/Source/AvalancheShapes/Public/DynamicMeshes/AvaShapeTorusDynMesh.h b/Engine/Plugins/VirtualProduction/Avalanche/Source/AvalancheShapes/Public/DynamicMeshes/AvaShapeTorusDynMesh.h
--- a/Engine/Plugins/VirtualProduction/Avalanche/Source/AvalancheShapes/Public/DynamicMeshes/AvaShapeTorusDynMesh.h
+++ b/Engine/Plugins/VirtualProduction/Avalanche/Source/AvalancheShapes/Public/DynamicMeshes/AvaShapeTorusDynMesh.h
@@ -103,6 +103,7 @@ protected:
 	virtual bool IsMeshVisible(int32 MeshIndex) override;
 	virtual bool CreateMesh(FAvaShapeMesh& InMesh) override;
 	virtual bool CreateUVs(FAvaShapeMesh& InMesh, FAvaShapeMaterialUVParameters& InParams) override;
+	virtual bool CanUpdateMeshInAnyThread() const override;
 
 	// represents the number of slices composing the tube
 	UPROPERTY(EditInstanceOnly, BlueprintReadWrite, Setter, Getter, Category="Shape", meta=(ClampMin="3.0", ClampMax="128.0", DisplayName="Slices", AllowPrivateAccess="true"))
