Motion Design: made update context optional instead of ptr, and removed IsSavingPackage check since the edit mesh will never broadcast the change.

#Synced-CL 51764840

diff --git a/Engine/Plugins/VirtualProduction/Avalanche/Source/AvalancheShapes/Private/AvaShapeBuildSystem.cpp b/Engine/Plugins/VirtualProduction/Avalanche/Source/AvalancheShapes/Private/AvaShapeBuildSystem.cpp
--- a/Engine/Plugins/VirtualProduction/Avalanche/Source/AvalancheShapes/Private/AvaShapeBuildSystem.cpp
+++ b/Engine/Plugins/VirtualProduction/Avalanche/Source/AvalancheShapes/Private/AvaShapeBuildSystem.cpp
@@ -101,19 +101,26 @@ void UAvaShapeBuildSystem::BuildShapes()
 	Algo::ForEach(GameThreadBuilds, Projection(&FShapeBuild::Shape, &UAvaShapeDynamicMeshBase::PreUpdateMesh));
 	Algo::ForEach(ParallelBuilds  , Projection(&FShapeBuild::Shape, &UAvaShapeDynamicMeshBase::PreUpdateMesh));
 
-	const UAvaShapeDynamicMeshBase::FMeshUpdateContext MeshParallelUpdateContext
+	const UAvaShapeDynamicMeshBase::FMeshUpdateContext ParallelUpdateContext
 		{
 			.bParallel = true
 		};
 
 	ParallelForWithPreWork(ParallelBuilds.Num(),
-		[&ParallelBuilds, &MeshParallelUpdateContext](int32 InBuildIndex)
+		[&ParallelBuilds, &ParallelUpdateContext](int32 InBuildIndex)
 		{
-			ParallelBuilds[InBuildIndex].Shape->UpdateMeshWithContext(MeshParallelUpdateContext);
+			ParallelBuilds[InBuildIndex].Shape->UpdateMesh(ParallelUpdateContext);
 		},
 		[&GameThreadBuilds]()
 		{
-			Algo::ForEach(GameThreadBuilds, Projection(&FShapeBuild::Shape, &UAvaShapeDynamicMeshBase::UpdateMesh));
+			const UAvaShapeDynamicMeshBase::FMeshUpdateContext GameThreadUpdateContext
+				{
+					.bParallel = false
+				};
+			for (FShapeBuild& Build : GameThreadBuilds)
+			{
+				Build.Shape->UpdateMesh(GameThreadUpdateContext);
+			}
 		});
 
 	Algo::ForEach(GameThreadBuilds, Projection(&FShapeBuild::Shape, &UAvaShapeDynamicMeshBase::PostUpdateMesh));
diff --git a/Engine/Plugins/VirtualProduction/Avalanche/Source/AvalancheShapes/Private/DynamicMeshes/AvaShapeDynMeshBase.cpp b/Engine/Plugins/VirtualProduction/Avalanche/Source/AvalancheShapes/Private/DynamicMeshes/AvaShapeDynMeshBase.cpp
--- a/Engine/Plugins/VirtualProduction/Avalanche/Source/AvalancheShapes/Private/DynamicMeshes/AvaShapeDynMeshBase.cpp
+++ b/Engine/Plugins/VirtualProduction/Avalanche/Source/AvalancheShapes/Private/DynamicMeshes/AvaShapeDynMeshBase.cpp
@@ -49,7 +49,7 @@ using namespace UE::Geometry;
 namespace UE::Private::AvaShape
 {
 	/** Edit underlying dynamic mesh */
-	bool EditMesh(const UAvaShapeDynamicMeshBase* InContext, const TFunctionRef<void(FDynamicMesh3&)>& InFunctor, EDynamicMeshChangeType InChangeType = EDynamicMeshChangeType::GeneralEdit, EDynamicMeshAttributeChangeFlags InAttributeChangeType = EDynamicMeshAttributeChangeFlags::Unknown, bool bInBroadcastChangeEvent = true)
+	bool EditMesh(const UAvaShapeDynamicMeshBase* InContext, const TFunctionRef<void(FDynamicMesh3&)>& InFunctor, EDynamicMeshChangeType InChangeType = EDynamicMeshChangeType::GeneralEdit, EDynamicMeshAttributeChangeFlags InAttributeChangeType = EDynamicMeshAttributeChangeFlags::Unknown)
 	{
 		checkf(!!InContext, TEXT("Invalid shape context provided to edit mesh"))
 
@@ -57,8 +57,7 @@ namespace UE::Private::AvaShape
 		{
 			if (UDynamicMesh* Mesh = MeshComponent->GetDynamicMesh())
 			{
-				const bool bIsSaving = UE::IsSavingPackage(InContext->GetPackage());
-				const bool bSkipBroadcastChangeEvents = bIsSaving || !bInBroadcastChangeEvent;
+				const bool bSkipBroadcastChangeEvents = true; // EditMesh never broadcasts any event. This is done by the UAvaShapeDynamicMeshBase on post mesh update
 				Mesh->EditMesh(InFunctor, InChangeType, InAttributeChangeType, bSkipBroadcastChangeEvents);
 				return true;
 			}
@@ -312,10 +311,16 @@ void UAvaShapeDynamicMeshBase::ToggleGizmo_Implementation(
 	}
 }
 
-void UAvaShapeDynamicMeshBase::UpdateMesh()
+void UAvaShapeDynamicMeshBase::UpdateMesh(const FMeshUpdateContext& InUpdateContext)
 {
 	TRACE_CPUPROFILER_EVENT_SCOPE(UAvaShapeDynamicMeshBase::UpdateMesh);
 
+	MeshUpdateContext = InUpdateContext;
+	ON_SCOPE_EXIT
+	{
+		MeshUpdateContext.Reset();
+	};
+
 	RestoreCachedMesh(/** Reset */true);
 
 	bool bUpdateMesh = GenerateMesh();
@@ -385,13 +390,6 @@ void UAvaShapeDynamicMeshBase::PreUpdateMesh()
 	}
 }
 
-void UAvaShapeDynamicMeshBase::UpdateMeshWithContext(const FMeshUpdateContext& InUpdateContext)
-{
-	MeshUpdateContext = &InUpdateContext;
-	UpdateMesh();
-	MeshUpdateContext = nullptr;
-}
-
 void UAvaShapeDynamicMeshBase::PostUpdateMesh()
 {
 	TRACE_CPUPROFILER_EVENT_SCOPE(UAvaShapeDynamicMeshBase::PostUpdateMesh);
@@ -468,7 +466,13 @@ void UAvaShapeDynamicMeshBase::RunUpdate()
 			if (GetShapeMeshComponent())
 			{
 				PreUpdateMesh();
-				UpdateMesh();
+
+				const FMeshUpdateContext UpdateContext
+					{
+						.bParallel = false
+					};
+				UpdateMesh(UpdateContext);
+
 				PostUpdateMesh();
 			}
 			else
@@ -512,20 +516,24 @@ AAvaShapeActor* UAvaShapeDynamicMeshBase::GetShapeActor() const
 
 UDynamicMeshComponent* UAvaShapeDynamicMeshBase::GetShapeMeshComponent() const
 {
-	if (!CachedComponent.IsValid())
+	UDynamicMeshComponent* ShapeMeshComponent = CachedComponent.Get();
+
+	// Only resolve shape mesh component when not in an update. Mesh Updates should've already validated this ahead of time in PreUpdateMesh
+	if (!ShapeMeshComponent && !MeshUpdateContext.IsSet())
 	{
 		if (const AAvaShapeActor* ShapeActor = GetShapeActor())
 		{
-			if (UDynamicMeshComponent* DynMeshComp = ShapeActor->GetShapeMeshComponent())
+			ShapeMeshComponent = ShapeActor->GetShapeMeshComponent();
+			if (ShapeMeshComponent)
 			{
 				UAvaShapeDynamicMeshBase* const MutableThis = const_cast<UAvaShapeDynamicMeshBase*>(this);
-				MutableThis->CachedComponent = DynMeshComp;
+				MutableThis->CachedComponent = ShapeMeshComponent;
 				MutableThis->InitializeDynamicMesh();
 			}
 		}
 	}
 
-	return CachedComponent.Get();
+	return ShapeMeshComponent;
 }
 
 const FString& UAvaShapeDynamicMeshBase::GetMeshName() const
@@ -1213,7 +1221,7 @@ bool UAvaShapeDynamicMeshBase::GenerateTangents()
 
 			FComputeTangentsOptions TangentOptions;
 			TangentOptions.bAveraged = true;
-			TangentOptions.bParallel = !MeshUpdateContext || !MeshUpdateContext->bParallel; // do not do parallel operation if already parallel.
+			TangentOptions.bParallel = !MeshUpdateContext.IsSet() || !MeshUpdateContext->bParallel; // do not do parallel operation if already parallel.
 
 			FMeshTangentsd Tangents(&EditMesh);
 			Tangents.ComputeTriVertexTangents(
@@ -1223,8 +1231,7 @@ bool UAvaShapeDynamicMeshBase::GenerateTangents()
 			Tangents.CopyToOverlays(EditMesh);
 		}
 		, EDynamicMeshChangeType::GeneralEdit
-		, EDynamicMeshAttributeChangeFlags::Unknown
-		, /*Broadcast*/false);
+		, EDynamicMeshAttributeChangeFlags::Unknown);
 
 	return true;
 }
@@ -1507,7 +1514,7 @@ bool UAvaShapeDynamicMeshBase::OnMeshSectionUpdated(FAvaShapeMesh& InMesh)
 		return true;
 	}
 
-	if (MeshUpdateStatus != EAvaDynamicMeshUpdateState::UpdateInProgress)
+	if (!MeshUpdateContext.IsSet())
 	{
 		InMesh.bUpdateRequired = true;
 		return true;
@@ -1655,8 +1662,7 @@ bool UAvaShapeDynamicMeshBase::CreateDynamicMesh(FAvaShapeMesh& InMesh)
 
 	InMesh.bUpdateRequired = false;
 
-	UE::Private::AvaShape::EditMesh(
-		this,
+	UE::Private::AvaShape::EditMesh(this,
 		[this, &InMesh](FDynamicMesh3& Mesh)
 		{
 			// clear ids before adding them
@@ -1714,9 +1720,7 @@ bool UAvaShapeDynamicMeshBase::CreateDynamicMesh(FAvaShapeMesh& InMesh)
 			}
 		}
 		, EDynamicMeshChangeType::GeneralEdit
-		, EDynamicMeshAttributeChangeFlags::MeshTopology
-		, /** Broadcast */false
-	);
+		, EDynamicMeshAttributeChangeFlags::MeshTopology);
 	return true;
 }
 
@@ -1746,8 +1750,7 @@ bool UAvaShapeDynamicMeshBase::UpdateDynamicMesh(FAvaShapeMesh& InMesh)
 	InMesh.bUpdateRequired = false;
 
 	// same topology, update vertices only (location, normal, color, uv)
-	UE::Private::AvaShape::EditMesh(
-		this,
+	UE::Private::AvaShape::EditMesh(this,
 		[this, InMesh](FDynamicMesh3& EditMesh)
 		{
 			FDynamicMeshNormalOverlay* const NormalOverlay = EditMesh.Attributes()->PrimaryNormals();
@@ -1762,7 +1765,7 @@ bool UAvaShapeDynamicMeshBase::UpdateDynamicMesh(FAvaShapeMesh& InMesh)
 					NormalOverlay->SetElement(NId, static_cast<FVector3f>(InMesh.Normals[Idx]));
 				};
 
-			if (MeshUpdateContext && MeshUpdateContext->bParallel)
+			if (MeshUpdateContext.IsSet() && MeshUpdateContext->bParallel)
 			{
 				for (int32 Idx = 0; Idx < InMesh.Vertices.Num(); ++Idx)
 				{
@@ -1775,9 +1778,8 @@ bool UAvaShapeDynamicMeshBase::UpdateDynamicMesh(FAvaShapeMesh& InMesh)
 			}
 		}
 		, EDynamicMeshChangeType::MeshVertexChange
-		, EDynamicMeshAttributeChangeFlags::Unknown
-		, /** Broadcast */false
-	);
+		, EDynamicMeshAttributeChangeFlags::Unknown);
+
 	return true;
 }
 
@@ -2061,8 +2063,7 @@ bool UAvaShapeDynamicMeshBase::ApplyUVsManually(FAvaShapeMesh& InMesh)
 
 	InMesh.UVIds.Empty();
 
-	return UE::Private::AvaShape::EditMesh(
-		this,
+	return UE::Private::AvaShape::EditMesh(this,
 		[this, &InMesh](FDynamicMesh3& EditMesh)
 		{
 			FDynamicMeshUVOverlay* const UVOverlay = EditMesh.Attributes()->GetUVLayer(0);
@@ -2098,9 +2099,7 @@ bool UAvaShapeDynamicMeshBase::ApplyUVsManually(FAvaShapeMesh& InMesh)
 			}
 		}
 		, EDynamicMeshChangeType::AttributeEdit
-		, EDynamicMeshAttributeChangeFlags::UVs
-		, /** Broadcast */false
-	);
+		, EDynamicMeshAttributeChangeFlags::UVs);
 }
 
 bool UAvaShapeDynamicMeshBase::ApplyUVsPlanarProjection(FAvaShapeMesh& InMesh, FRotator PlaneRotation, FVector2D PlaneSize)
@@ -2110,8 +2109,7 @@ bool UAvaShapeDynamicMeshBase::ApplyUVsPlanarProjection(FAvaShapeMesh& InMesh, F
 		return false;
 	}
 
-	return UE::Private::AvaShape::EditMesh(
-		this,
+	return UE::Private::AvaShape::EditMesh(this,
 		[this, &InMesh, PlaneRotation, PlaneSize](FDynamicMesh3& EditMesh)
 		{
 			const FTransform PlaneTransform(PlaneRotation, FVector(0), FVector(PlaneSize.X, PlaneSize.Y, 0));
@@ -2136,9 +2134,7 @@ bool UAvaShapeDynamicMeshBase::ApplyUVsPlanarProjection(FAvaShapeMesh& InMesh, F
 			}
 		}
 		, EDynamicMeshChangeType::AttributeEdit
-		, EDynamicMeshAttributeChangeFlags::UVs
-		, /** Broadcast */false
-	);
+		, EDynamicMeshAttributeChangeFlags::UVs);
 }
 
 bool UAvaShapeDynamicMeshBase::ApplyUVsBoxProjection(FAvaShapeMesh& InMesh, FRotator BoxRotation, FVector BoxSize)
@@ -2148,8 +2144,7 @@ bool UAvaShapeDynamicMeshBase::ApplyUVsBoxProjection(FAvaShapeMesh& InMesh, FRot
 		return false;
 	}
 
-	return UE::Private::AvaShape::EditMesh(
-		this,
+	return UE::Private::AvaShape::EditMesh(this,
 		[this, &InMesh, BoxRotation, BoxSize](FDynamicMesh3& EditMesh)
 		{
 			const FTransform PlaneTransform(BoxRotation, FVector(0), BoxSize);
@@ -2172,24 +2167,19 @@ bool UAvaShapeDynamicMeshBase::ApplyUVsBoxProjection(FAvaShapeMesh& InMesh, FRot
 			}
 		}
 		, EDynamicMeshChangeType::AttributeEdit
-		, EDynamicMeshAttributeChangeFlags::UVs
-		, /** Broadcast */false
-	);
+		, EDynamicMeshAttributeChangeFlags::UVs);
 }
 
 bool UAvaShapeDynamicMeshBase::ApplyUVsTransform(FAvaShapeMesh& InMesh, FAvaShapeMaterialUVParameters& InParams, FVector2D ShapeSize, FVector2D UVOffset, float UVFixRotation)
 {
 	bool bResult = true;
-	return UE::Private::AvaShape::EditMesh(
-		this,
+	return UE::Private::AvaShape::EditMesh(this,
 		[InMesh, InParams, ShapeSize, UVOffset, UVFixRotation, &bResult](FDynamicMesh3& EditMesh)
 		{
 			bResult = UE::AvaShapes::TransformMeshUVs(EditMesh, InMesh.UVIds, InParams, ShapeSize, UVOffset, UVFixRotation);
 		}
 		, EDynamicMeshChangeType::AttributeEdit
-		, EDynamicMeshAttributeChangeFlags::UVs
-		, /** Broadcast */false
-	) && bResult;
+		, EDynamicMeshAttributeChangeFlags::UVs) && bResult;
 }
 
 void UAvaShapeDynamicMeshBase::SetRunAsync(bool bInAsync)
@@ -2225,8 +2215,7 @@ void UAvaShapeDynamicMeshBase::RestoreCachedMesh(bool bInReset)
 {
 	if (CachedMesh.IsSet())
 	{
-		UE::Private::AvaShape::EditMesh(
-			this,
+		UE::Private::AvaShape::EditMesh(this,
 			[this, bInReset](FDynamicMesh3& EditMesh)
 			{
 				if (bInReset)
@@ -2240,8 +2229,7 @@ void UAvaShapeDynamicMeshBase::RestoreCachedMesh(bool bInReset)
 				}
 			}
 		, EDynamicMeshChangeType::GeneralEdit
-		,EDynamicMeshAttributeChangeFlags::Unknown
-		, /** Broadcast */false);
+		, EDynamicMeshAttributeChangeFlags::Unknown);
 	}
 }
 
@@ -2406,9 +2394,8 @@ void UAvaShapeDynamicMeshBase::InitializeDynamicMesh()
 		return;
 	}
 
-	UE::Private::AvaShape::EditMesh(
-			this,
-			[this](FDynamicMesh3& EditMesh)
+	UE::Private::AvaShape::EditMesh(this,
+		[this](FDynamicMesh3& EditMesh)
 		{
 			EditMesh.Clear();
 
@@ -2434,9 +2421,7 @@ void UAvaShapeDynamicMeshBase::InitializeDynamicMesh()
 			CachedMesh = EditMesh;
 		}
 		, EDynamicMeshChangeType::GeneralEdit
-		, EDynamicMeshAttributeChangeFlags::Unknown
-		, /** Broadcast */false
-	);
+		, EDynamicMeshAttributeChangeFlags::Unknown);
 
 	MeshComponent->TransformUpdated.AddUObject(this, &UAvaShapeDynamicMeshBase::OnMeshTransformChanged);
 }
@@ -2525,9 +2510,8 @@ bool UAvaShapeDynamicMeshBase::ClearDynamicMeshSection(int32 MeshIndex)
 		return false;
 	}
 
-	return UE::Private::AvaShape::EditMesh(
-			this,
-			[this, &MeshData](FDynamicMesh3& EditMesh)
+	return UE::Private::AvaShape::EditMesh(this,
+		[this, &MeshData](FDynamicMesh3& EditMesh)
 		{
 			for (const int32 TId : MeshData->Mesh.TriangleIds)
 			{
@@ -2541,16 +2525,13 @@ bool UAvaShapeDynamicMeshBase::ClearDynamicMeshSection(int32 MeshIndex)
 			MeshData->Mesh.ClearIds();
 		}
 		, EDynamicMeshChangeType::GeneralEdit
-		, EDynamicMeshAttributeChangeFlags::MeshTopology
-		, /** Broadcast */false
-	);
+		, EDynamicMeshAttributeChangeFlags::MeshTopology);
 }
 
 bool UAvaShapeDynamicMeshBase::ClearDynamicMesh()
 {
-	const bool bResult = UE::Private::AvaShape::EditMesh(
-			this,
-			[this](FDynamicMesh3& EditMesh)
+	const bool bResult = UE::Private::AvaShape::EditMesh(this,
+		[this](FDynamicMesh3& EditMesh)
 		{
 			for (const int32 TId : EditMesh.TriangleIndicesItr())
 			{
@@ -2561,9 +2542,7 @@ bool UAvaShapeDynamicMeshBase::ClearDynamicMesh()
 			}
 		}
 		, EDynamicMeshChangeType::GeneralEdit
-		, EDynamicMeshAttributeChangeFlags::MeshTopology
-		, /** Broadcast */false
-	);
+		, EDynamicMeshAttributeChangeFlags::MeshTopology);
 
 	if (!bResult)
 	{
diff --git a/Engine/Plugins/VirtualProduction/Avalanche/Source/AvalancheShapes/Public/DynamicMeshes/AvaShapeDynMeshBase.h b/Engine/Plugins/VirtualProduction/Avalanche/Source/AvalancheShapes/Public/DynamicMeshes/AvaShapeDynMeshBase.h
--- a/Engine/Plugins/VirtualProduction/Avalanche/Source/AvalancheShapes/Public/DynamicMeshes/AvaShapeDynMeshBase.h
+++ b/Engine/Plugins/VirtualProduction/Avalanche/Source/AvalancheShapes/Public/DynamicMeshes/AvaShapeDynMeshBase.h
@@ -388,8 +388,13 @@ protected:
 		return false;
 	}
 
+	struct FMeshUpdateContext
+	{
+		/** Determines whether the update is taking place in parallel (potentially not the game thread) */
+		bool bParallel = false;
+	};
 	/** Takes all dirty flags, regenerate whatever is needed and updates the DMC, called by RunUpdate or RequireUpdate */
-	void UpdateMesh();
+	void UpdateMesh(const FMeshUpdateContext& InUpdateContext);
 
 	/** Offsets actor to account for the mesh generated location */
 	void OffsetActorToAlign();
@@ -397,14 +402,6 @@ protected:
 	/** Part of update flow that can only be performed safely on the game thread. Should be called before UpdateMesh_AnyThread */
 	void PreUpdateMesh();
 
-	struct FMeshUpdateContext
-	{
-		/** Determines whether the update is taking place in parallel (potentially not the game thread) */
-		bool bParallel = false;
-	};
-	/** UpdateMesh() that can be performed on a different thread */
-	void UpdateMeshWithContext(const FMeshUpdateContext& InUpdateContext);
-
 	/** Part of update flow that can only be performed safely on the game thread. Should be called after UpdateMesh_AnyThread */
 	void PostUpdateMesh();
 
@@ -491,7 +488,7 @@ private:
 	TOptional<UE::Geometry::FDynamicMesh3> CachedMesh;
 
 	/** the context while mesh is being built */
-	const FMeshUpdateContext* MeshUpdateContext = nullptr;
+	TOptional<FMeshUpdateContext> MeshUpdateContext;
 
 	/** Do not access directly use getter to load it if nullptr */
 	TWeakObjectPtr<UDynamicMeshComponent> CachedComponent;
