Motion Design: remove trace bookmarks from mask modifiers (unnecessary noise to insights) + removed  TBasePassPSFNoLightMapPolicySkylight from shader entries as it's no longer present

#Synced-CL 51764840

diff --git a/Engine/Plugins/VirtualProduction/Avalanche/Source/AvalancheMask/Private/Mask2D/AvaMask2DBaseModifier.cpp b/Engine/Plugins/VirtualProduction/Avalanche/Source/AvalancheMask/Private/Mask2D/AvaMask2DBaseModifier.cpp
--- a/Engine/Plugins/VirtualProduction/Avalanche/Source/AvalancheMask/Private/Mask2D/AvaMask2DBaseModifier.cpp
+++ b/Engine/Plugins/VirtualProduction/Avalanche/Source/AvalancheMask/Private/Mask2D/AvaMask2DBaseModifier.cpp
@@ -262,8 +262,6 @@ void UAvaMask2DBaseModifier::OnModifierAdded(EActorModifierCoreEnableReason InRe
 {
 	Super::OnModifierAdded(InReason);
 
-	TRACE_BOOKMARK(TEXT("UAvaMask2DModifier::OnModifierAdded"));
-
 	UpdateCanvas();
 
 	if (InReason == EActorModifierCoreEnableReason::User
@@ -289,8 +287,6 @@ void UAvaMask2DBaseModifier::OnModifierRemoved(EActorModifierCoreDisableReason I
 {
 	Super::OnModifierRemoved(InReason);
 
-	TRACE_BOOKMARK(TEXT("UAvaMask2DModifier::OnModifierRemoved"));
-
 	if (InReason == EActorModifierCoreDisableReason::Destroyed)
 	{
 		CanvasParamsToLocal();
diff --git a/Engine/Plugins/VirtualProduction/Avalanche/Source/AvalancheMaterial/Internal/MaterialCache/AvaMaterialCacheSettings.h b/Engine/Plugins/VirtualProduction/Avalanche/Source/AvalancheMaterial/Internal/MaterialCache/AvaMaterialCacheSettings.h
--- a/Engine/Plugins/VirtualProduction/Avalanche/Source/AvalancheMaterial/Internal/MaterialCache/AvaMaterialCacheSettings.h
+++ b/Engine/Plugins/VirtualProduction/Avalanche/Source/AvalancheMaterial/Internal/MaterialCache/AvaMaterialCacheSettings.h
@@ -14,10 +14,13 @@ class FVertexFactoryType;
 UENUM()
 enum class EAvaShaderProfileType : uint8
 {
-	/** CacheShaders will be used and will gather all the shaders it might need */
-	Default,
 	/**
-	 * Only the specified shaders will be cached. Should yield better performance 
+	 * CacheShaders will be used and will gather all the shaders it might need
+	 * NOTE: this might cause more shaders to be compiled than required.
+	 */
+	Generic,
+	/**
+	 * Only the specified shaders will be cached. Should yield better performance than generic.
 	 * NOTE: Specified shaders only will apply in editor builds. This will default to generic on non-editor builds
 	 */
 	Specific,
@@ -31,7 +34,7 @@ namespace UE::Ava
 /** Shader Profile that has all the type pointers cached */
 struct FResolvedShaderProfile
 {
-	EAvaShaderProfileType ProfileType = EAvaShaderProfileType::Default;
+	EAvaShaderProfileType ProfileType = EAvaShaderProfileType::Generic;
 	TArray<const FVertexFactoryType*> VertexFactoryTypes;
 	TArray<const FShaderPipelineType*> PipelineTypes;
 	TArray<const FShaderType*> ShaderTypes;
@@ -72,7 +75,7 @@ struct FAvaShaderProfile
 	FName ProfileName;
 
 	UPROPERTY(EditAnywhere, Category="Material Cache")
-	EAvaShaderProfileType ProfileType = EAvaShaderProfileType::Default;
+	EAvaShaderProfileType ProfileType = EAvaShaderProfileType::Generic;
 
 	/** Specific Shader Combinations to use */
 	UPROPERTY(EditAnywhere, Category="Material Cache", meta=(EditCondition="ProfileType==EAvaShaderProfileType::Specific"))
diff --git a/Engine/Plugins/VirtualProduction/Avalanche/Source/AvalancheMaterial/Private/MaterialCache/AvaMaterialCacheSettings.cpp b/Engine/Plugins/VirtualProduction/Avalanche/Source/AvalancheMaterial/Private/MaterialCache/AvaMaterialCacheSettings.cpp
--- a/Engine/Plugins/VirtualProduction/Avalanche/Source/AvalancheMaterial/Private/MaterialCache/AvaMaterialCacheSettings.cpp
+++ b/Engine/Plugins/VirtualProduction/Avalanche/Source/AvalancheMaterial/Private/MaterialCache/AvaMaterialCacheSettings.cpp
@@ -3,6 +3,7 @@
 #include "MaterialCache/AvaMaterialCacheSettings.h"
 #include "MaterialCache/AvaMaterialCacheLog.h"
 #include "Shader.h"
+#include "UObject/UnrealType.h"
 #include "VertexFactory.h"
 
 const FLazyName UAvaMaterialCacheSettings::DefaultRealtimeProfile(TEXT("BasicSceneRealtime"));
@@ -24,7 +25,6 @@ TArray<FAvaShaderCombination> MakeDefaultShaderCombinations()
 	// Shaders
 	const FName BasePassVS               = TEXT("TBasePassVSFNoLightMapPolicy");
 	const FName BasePassPS               = TEXT("TBasePassPSFNoLightMapPolicy");
-	const FName BasePassPSSL             = TEXT("TBasePassPSFNoLightMapPolicySkylight");
 	const FName NiagaraShader            = TEXT("FNiagaraShader");
 	const FName MaterialCHS              = TEXT("TMaterialCHSFNoLightMapPolicyFAnyHitShader");
 	const FName LumenCardVS              = TEXT("FLumenCardVS");
@@ -46,7 +46,6 @@ TArray<FAvaShaderCombination> MakeDefaultShaderCombinations()
 			{
 				BasePassVS,
 				BasePassPS,
-				BasePassPSSL,
 				NiagaraShader,
 				MaterialCHS,
 				LumenCardVS,
