Motion Design: set text3d and shape build system parallel for flags to unbalanced as each task build time can be highly variable

#Synced-CL 52272439

diff --git a/Engine/Plugins/VirtualProduction/Avalanche/Source/AvalancheShapes/Private/AvaShapeBuildSystem.cpp b/Engine/Plugins/VirtualProduction/Avalanche/Source/AvalancheShapes/Private/AvaShapeBuildSystem.cpp
--- a/Engine/Plugins/VirtualProduction/Avalanche/Source/AvalancheShapes/Private/AvaShapeBuildSystem.cpp
+++ b/Engine/Plugins/VirtualProduction/Avalanche/Source/AvalancheShapes/Private/AvaShapeBuildSystem.cpp
@@ -106,6 +106,9 @@ void UAvaShapeBuildSystem::BuildShapes()
 			.bParallel = true
 		};
 
+	// Shape building could be highly variable as it depends on the shape being updated
+	constexpr EParallelForFlags ParallelForFlags = EParallelForFlags::Unbalanced;
+
 	ParallelForWithPreWork(ParallelBuilds.Num(),
 		[&ParallelBuilds, &ParallelUpdateContext](int32 InBuildIndex)
 		{
@@ -121,7 +124,7 @@ void UAvaShapeBuildSystem::BuildShapes()
 			{
 				Build.Shape->UpdateMesh(GameThreadUpdateContext);
 			}
-		});
+		}, ParallelForFlags);
 
 	Algo::ForEach(GameThreadBuilds, Projection(&FShapeBuild::Shape, &UAvaShapeDynamicMeshBase::PostUpdateMesh));
 	Algo::ForEach(ParallelBuilds  , Projection(&FShapeBuild::Shape, &UAvaShapeDynamicMeshBase::PostUpdateMesh));
diff --git a/Engine/Plugins/VirtualProduction/Text3D/Source/Text3D/Private/Subsystems/Text3DEngineSubsystem.cpp b/Engine/Plugins/VirtualProduction/Text3D/Source/Text3D/Private/Subsystems/Text3DEngineSubsystem.cpp
--- a/Engine/Plugins/VirtualProduction/Text3D/Source/Text3D/Private/Subsystems/Text3DEngineSubsystem.cpp
+++ b/Engine/Plugins/VirtualProduction/Text3D/Source/Text3D/Private/Subsystems/Text3DEngineSubsystem.cpp
@@ -251,6 +251,9 @@ void UText3DEngineSubsystem::ProcessBuildGlyphMeshes()
 			.DefaultMaterial = GetDefaultMaterial(),
 		};
 
+	// Each mesh build time cost is highly variable, and each task can take enough time to justify the cost of synchronization
+	constexpr EParallelForFlags ParallelForFlags = EParallelForFlags::Unbalanced;
+
 	ParallelFor(MeshDescs.Num(), 
 		[&SharedParams, &MeshDescs, &GlyphParameters](int32 InMeshDescIndex)
 		{
@@ -267,7 +270,7 @@ void UText3DEngineSubsystem::ProcessBuildGlyphMeshes()
 					};
 				FontFaceCache->BuildGlyphMesh(SharedParams, InstanceParams);
 			}
-		});
+		}, ParallelForFlags);
 }
 
 #if WITH_FREETYPE
