Motion Design: added race detection instrumentation for Text3D and shape building.
Moved slot names to be an array of FName to avoid race on FLazyName::Resolve.

#Synced-CL 52400047

diff --git a/Engine/Plugins/VirtualProduction/Avalanche/Source/AvalancheShapes/Private/AvaShapeBuildSystem.cpp b/Engine/Plugins/VirtualProduction/Avalanche/Source/AvalancheShapes/Private/AvaShapeBuildSystem.cpp
--- a/Engine/Plugins/VirtualProduction/Avalanche/Source/AvalancheShapes/Private/AvaShapeBuildSystem.cpp
+++ b/Engine/Plugins/VirtualProduction/Avalanche/Source/AvalancheShapes/Private/AvaShapeBuildSystem.cpp
@@ -11,6 +11,25 @@
 #include "Modifiers/ActorModifierCoreStack.h"
 #include "Subsystems/ActorModifierCoreSubsystem.h"
 
+#if USING_INSTRUMENTATION
+#include "Sanitizer/RaceDetector.h"
+#endif
+
+namespace UE::AvaShape::Private
+{
+
+#if USING_INSTRUMENTATION
+static bool GDetectRaceDuringBuild = false;
+static FAutoConsoleVariableRef CVarDetectRaceDuringBuild(
+	TEXT("MotionDesign.Shapes.DetectRaceDuringBuild"),
+	GDetectRaceDuringBuild,
+	TEXT("Activate the race detector when building shapes in parallel"),
+	ECVF_Default
+);
+#endif
+
+} // UE::AvaShape::Private
+
 UAvaShapeBuildSystem* UAvaShapeBuildSystem::FindBuildSystem(TNotNull<const UObject*> InContextObject)
 {
 	if (UWorld* const World = InContextObject->GetWorld())
@@ -106,25 +125,32 @@ void UAvaShapeBuildSystem::BuildShapes()
 			.bParallel = true
 		};
 
-	// Shape building could be highly variable as it depends on the shape being updated
-	constexpr EParallelForFlags ParallelForFlags = EParallelForFlags::Unbalanced;
+	// Update shapes in parallel
+	{
+		// Shape building could be highly variable as it depends on the shape being updated
+		constexpr EParallelForFlags ParallelForFlags = EParallelForFlags::Unbalanced;
 
-	ParallelForWithPreWork(ParallelBuilds.Num(),
-		[&ParallelBuilds, &ParallelUpdateContext](int32 InBuildIndex)
-		{
-			ParallelBuilds[InBuildIndex].Shape->UpdateMesh(ParallelUpdateContext);
-		},
-		[&GameThreadBuilds]()
-		{
-			const UAvaShapeDynamicMeshBase::FMeshUpdateContext GameThreadUpdateContext
-				{
-					.bParallel = false
-				};
-			for (FShapeBuild& Build : GameThreadBuilds)
+#if USING_INSTRUMENTATION
+		UE::Sanitizer::RaceDetector::FRaceDetectorScope RaceDetectorScope(UE::AvaShape::Private::GDetectRaceDuringBuild);
+#endif
+
+		ParallelForWithPreWork(ParallelBuilds.Num(),
+			[&ParallelBuilds, &ParallelUpdateContext](int32 InBuildIndex)
 			{
-				Build.Shape->UpdateMesh(GameThreadUpdateContext);
-			}
-		}, ParallelForFlags);
+				ParallelBuilds[InBuildIndex].Shape->UpdateMesh(ParallelUpdateContext);
+			},
+			[&GameThreadBuilds]()
+			{
+				const UAvaShapeDynamicMeshBase::FMeshUpdateContext GameThreadUpdateContext
+					{
+						.bParallel = false
+					};
+				for (FShapeBuild& Build : GameThreadBuilds)
+				{
+					Build.Shape->UpdateMesh(GameThreadUpdateContext);
+				}
+			}, ParallelForFlags);
+	}
 
 	Algo::ForEach(GameThreadBuilds, Projection(&FShapeBuild::Shape, &UAvaShapeDynamicMeshBase::PostUpdateMesh));
 	Algo::ForEach(ParallelBuilds  , Projection(&FShapeBuild::Shape, &UAvaShapeDynamicMeshBase::PostUpdateMesh));
diff --git a/Engine/Plugins/VirtualProduction/Text3D/Source/Text3D/Private/GeometryBuilders/Text3DGlyph.cpp b/Engine/Plugins/VirtualProduction/Text3D/Source/Text3D/Private/GeometryBuilders/Text3DGlyph.cpp
--- a/Engine/Plugins/VirtualProduction/Text3D/Source/Text3D/Private/GeometryBuilders/Text3DGlyph.cpp
+++ b/Engine/Plugins/VirtualProduction/Text3D/Source/Text3D/Private/GeometryBuilders/Text3DGlyph.cpp
@@ -36,6 +36,7 @@ void FText3DGlyph::Build(FText3DCachedMesh& InMesh, EText3DMeshType InMeshType,
 	{
 		using namespace UE::Text3D::Material;
 
+		const TConstArrayView<FName> SlotNames = GetSlotNames();
 		check(SlotNames.IsValidIndex(InIndex))
 
 		FName MaterialName = SlotNames[InIndex];
diff --git a/Engine/Plugins/VirtualProduction/Text3D/Source/Text3D/Private/Renderers/DynamicMesh/Text3DDynamicMeshRenderer.cpp b/Engine/Plugins/VirtualProduction/Text3D/Source/Text3D/Private/Renderers/DynamicMesh/Text3DDynamicMeshRenderer.cpp
--- a/Engine/Plugins/VirtualProduction/Text3D/Source/Text3D/Private/Renderers/DynamicMesh/Text3DDynamicMeshRenderer.cpp
+++ b/Engine/Plugins/VirtualProduction/Text3D/Source/Text3D/Private/Renderers/DynamicMesh/Text3DDynamicMeshRenderer.cpp
@@ -142,7 +142,7 @@ void UText3DDynamicMeshRenderer::OnUpdate(const UE::Text3D::Renderer::FUpdatePar
 			for (const int32 MaterialSlot : GlyphMeshData[InIndex].MaterialSlots)
 			{
 				const FName StyleTag = InCharacter->GetStyleTag();
-				const EText3DGroupType GroupType = static_cast<EText3DGroupType>(MaterialSlot - (InIndex * SlotNames.Num()));
+				const EText3DGroupType GroupType = static_cast<EText3DGroupType>(MaterialSlot - (InIndex * GetSlotNames().Num()));
 
 				FMaterialParameters Parameters;
 				Parameters.Group = GroupType;
@@ -333,7 +333,7 @@ void UText3DDynamicMeshRenderer::AppendGlyphMesh(FScopedMeshEditor& InMeshEditor
 				for (const TPair<int32, int32>& FromToPair : FromToMap)
 				{
 					InEditMesh.SetTriangleGroup(FromToPair.Value, InIndex);
-					const int32 MaterialID = (InIndex * UE::Text3D::Material::SlotNames.Num()) + MaterialAttribute->GetValue(FromToPair.Value);
+					const int32 MaterialID = (InIndex * UE::Text3D::Material::GetSlotNames().Num()) + MaterialAttribute->GetValue(FromToPair.Value);
 					MaterialAttribute->SetValue(FromToPair.Value, MaterialID);
 					MaterialSlots.AddUnique(MaterialID);
 				}
diff --git a/Engine/Plugins/VirtualProduction/Text3D/Source/Text3D/Private/Renderers/StaticMeshes/Text3DStaticMeshesRenderer.cpp b/Engine/Plugins/VirtualProduction/Text3D/Source/Text3D/Private/Renderers/StaticMeshes/Text3DStaticMeshesRenderer.cpp
--- a/Engine/Plugins/VirtualProduction/Text3D/Source/Text3D/Private/Renderers/StaticMeshes/Text3DStaticMeshesRenderer.cpp
+++ b/Engine/Plugins/VirtualProduction/Text3D/Source/Text3D/Private/Renderers/StaticMeshes/Text3DStaticMeshesRenderer.cpp
@@ -208,7 +208,7 @@ void UText3DStaticMeshesRenderer::OnUpdate(const UE::Text3D::Renderer::FUpdatePa
 				{
 					for (int32 GroupIndex = 0; GroupIndex < static_cast<int32>(EText3DGroupType::TypeCount); GroupIndex++)
 					{
-						const int32 MaterialIndex = CharacterMeshComponent->GetMaterialIndex(SlotNames[GroupIndex]);
+						const int32 MaterialIndex = CharacterMeshComponent->GetMaterialIndex(GetSlotNames()[GroupIndex]);
 
 						if (MaterialIndex == INDEX_NONE)
 						{
diff --git a/Engine/Plugins/VirtualProduction/Text3D/Source/Text3D/Private/Subsystems/Text3DEngineSubsystem.cpp b/Engine/Plugins/VirtualProduction/Text3D/Source/Text3D/Private/Subsystems/Text3DEngineSubsystem.cpp
--- a/Engine/Plugins/VirtualProduction/Text3D/Source/Text3D/Private/Subsystems/Text3DEngineSubsystem.cpp
+++ b/Engine/Plugins/VirtualProduction/Text3D/Source/Text3D/Private/Subsystems/Text3DEngineSubsystem.cpp
@@ -2,6 +2,7 @@
 
 #include "Text3DEngineSubsystem.h"
 
+#include "Async/ParallelFor.h"
 #include "Containers/Ticker.h"
 #include "Engine/Engine.h"
 #include "Engine/StaticMesh.h"
@@ -13,14 +14,33 @@
 #include "Materials/Material.h"
 #include "Misc/FileHelper.h"
 #include "Settings/Text3DProjectSettings.h"
+#include "Text3DInternalTypes.h"
 #include "UDynamicMesh.h"
-#include "Async/ParallelFor.h"
 #include "UObject/ConstructorHelpers.h"
 
+#if USING_INSTRUMENTATION
+#include "Sanitizer/RaceDetector.h"
+#endif
+
 #if WITH_FREETYPE
 #include "Fonts/FontCacheFreeType.h"
 #endif
 
+namespace UE::Text3D::Private
+{
+
+#if USING_INSTRUMENTATION
+static bool GDetectRaceDuringBuild = false;
+static FAutoConsoleVariableRef CVarDetectRaceDuringBuild(
+	TEXT("Text3D.Build.DetectRace"),
+	GDetectRaceDuringBuild,
+	TEXT("Activate the race detector when building text in parallel"),
+	ECVF_Default
+);
+#endif
+
+} // UE::Text3D::Private
+
 UText3DEngineSubsystem* UText3DEngineSubsystem::Get()
 {
 	if (GEngine)
@@ -254,6 +274,10 @@ void UText3DEngineSubsystem::ProcessBuildGlyphMeshes()
 	// Each mesh build time cost is highly variable, and each task can take enough time to justify the cost of synchronization
 	constexpr EParallelForFlags ParallelForFlags = EParallelForFlags::Unbalanced;
 
+#if USING_INSTRUMENTATION
+	UE::Sanitizer::RaceDetector::FRaceDetectorScope RaceDetectorScope(UE::Text3D::Private::GDetectRaceDuringBuild);
+#endif
+
 	ParallelFor(MeshDescs.Num(), 
 		[&SharedParams, &MeshDescs, &GlyphParameters](int32 InMeshDescIndex)
 		{
diff --git a/Engine/Plugins/VirtualProduction/Text3D/Source/Text3D/Private/Text3DInternalTypes.cpp b/Engine/Plugins/VirtualProduction/Text3D/Source/Text3D/Private/Text3DInternalTypes.cpp
new file mode 100644
--- /dev/null
+++ b/Engine/Plugins/VirtualProduction/Text3D/Source/Text3D/Private/Text3DInternalTypes.cpp
@@ -0,0 +1,21 @@
+// Copyright Epic Games, Inc. All Rights Reserved.
+
+#include "Text3DInternalTypes.h"
+#include "Containers/Array.h"
+#include "UObject/NameTypes.h"
+
+namespace UE::Text3D::Material
+{
+
+TConstArrayView<FName> GetSlotNames()
+{
+	static const TArray<FName> SlotNames = {
+		TEXT("Front"),
+		TEXT("Bevel"),
+		TEXT("Extrude"),
+		TEXT("Back"),
+	};
+	return SlotNames;
+}
+
+} // UE::Text3D::Material
diff --git a/Engine/Plugins/VirtualProduction/Text3D/Source/Text3D/Private/Text3DInternalTypes.h b/Engine/Plugins/VirtualProduction/Text3D/Source/Text3D/Private/Text3DInternalTypes.h
--- a/Engine/Plugins/VirtualProduction/Text3D/Source/Text3D/Private/Text3DInternalTypes.h
+++ b/Engine/Plugins/VirtualProduction/Text3D/Source/Text3D/Private/Text3DInternalTypes.h
@@ -2,16 +2,12 @@
 
 #pragma once
 
-#include "Containers/Array.h"
-#include "UObject/NameTypes.h"
+#include "Containers/ContainersFwd.h"
+
+class FName;
 
 namespace UE::Text3D::Material
 {
-	static const TArray<FLazyName> SlotNames =
-	{
-		TEXT("Front"),
-		TEXT("Bevel"),
-		TEXT("Extrude"),
-		TEXT("Back")
-	};
-}
\ No newline at end of file
+	TConstArrayView<FName> GetSlotNames();
+
+} // UE::Text3D::Material
