AvalancheFunctionalTest - AScreenshotFunctionalTestAvalanche::FinishTest now drains all pending Niagara CPU compilation tasks before actor cleanup.

#Synced-CL 52468003

diff --git a/Engine/Plugins/VirtualProduction/Avalanche/Source/AvalancheFunctionalTest/Private/ScreenshotFunctionalTestAvalanche.cpp b/Engine/Plugins/VirtualProduction/Avalanche/Source/AvalancheFunctionalTest/Private/ScreenshotFunctionalTestAvalanche.cpp
--- a/Engine/Plugins/VirtualProduction/Avalanche/Source/AvalancheFunctionalTest/Private/ScreenshotFunctionalTestAvalanche.cpp
+++ b/Engine/Plugins/VirtualProduction/Avalanche/Source/AvalancheFunctionalTest/Private/ScreenshotFunctionalTestAvalanche.cpp
@@ -7,6 +7,8 @@
 #include "Tests/AutomationCommon.h"
 #include "UnrealClient.h"
 #include "NiagaraComponent.h"
+#include "NiagaraSystem.h"
+#include "UObject/UObjectIterator.h"
 #include "NiagaraSystemInstanceController.h"
 #include "Components/PrimitiveComponent.h"
 
@@ -180,6 +182,64 @@ bool AScreenshotFunctionalTestAvalanche::IsReady_Implementation()
 	return true;
 }
 
+void AScreenshotFunctionalTestAvalanche::FinishTest(EFunctionalTestResult TestResult, const FString& Message)
+{
+	// Drain all pending Niagara CPU compilation tasks before actor cleanup.
+	// FNiagaraPlatformSet::RefreshScalability() can trigger async RequestCompile during the
+	// test tick. If the Cloner actor is destroyed while a FNiagaraSystemCompilationTask is
+	// still in-flight, DigestSystemInfo() dereferences a null TWeakObjectPtr<UNiagaraSystem>
+	// (EmitterHandles at offset 0x390) -> EXCEPTION_ACCESS_VIOLATION.
+	if (UWorld* World = GetWorld())
+	{
+		int32 NumSystemsDrained = 0;
+		int32 NumSystemsChecked = 0;
+
+		for (TObjectIterator<UNiagaraComponent> It; It; ++It)
+		{
+			UNiagaraComponent* NiagaraComp = *It;
+			if (!IsValid(NiagaraComp) || NiagaraComp->GetWorld() != World)
+			{
+				continue;
+			}
+
+			if (UNiagaraSystem* NiagaraSystem = NiagaraComp->GetAsset())
+			{
+				++NumSystemsChecked;
+				const bool bHadPendingCompilation = NiagaraSystem->HasOutstandingCompilationRequests(false);
+
+				if (bHadPendingCompilation)
+				{
+					UE_LOG(LogScreenshotAvalanche, Log,
+						TEXT("[FinishTest] '%s' — NiagaraSystem '%s' (component '%s') has pending CPU compilation, draining..."),
+						*GetName(), *NiagaraSystem->GetName(), *NiagaraComp->GetName());
+
+					NiagaraSystem->WaitForCompilationComplete(false, false);
+					++NumSystemsDrained;
+
+					UE_LOG(LogScreenshotAvalanche, Log,
+						TEXT("[FinishTest] '%s' — NiagaraSystem '%s' compilation drained. Still outstanding after wait: %s"),
+						*GetName(), *NiagaraSystem->GetName(),
+						NiagaraSystem->HasOutstandingCompilationRequests(false) ? TEXT("YES (unexpected)") : TEXT("no"));
+				}
+			}
+		}
+
+		UE_LOG(LogScreenshotAvalanche, Log,
+			TEXT("[FinishTest] '%s' — checked %d NiagaraSystem(s), drained %d with pending compilation."),
+			*GetName(), NumSystemsChecked, NumSystemsDrained);
+
+		if (NumSystemsDrained == 0)
+		{
+			UE_LOG(LogScreenshotAvalanche, Verbose,
+				TEXT("[FinishTest] '%s' — no pending Niagara compilation found. If this test still crashes with "
+				     "EXCEPTION_ACCESS_VIOLATION in DigestSystemInfo, RefreshScalability may have fired after this drain."),
+				*GetName());
+		}
+	}
+
+	Super::FinishTest(TestResult, Message);
+}
+
 bool AScreenshotFunctionalTestAvalanche::IsActorFullyInitialized(AActor* Actor) const
 {
 	if (!IsValid(Actor))
diff --git a/Engine/Plugins/VirtualProduction/Avalanche/Source/AvalancheFunctionalTest/Public/ScreenshotFunctionalTestAvalanche.h b/Engine/Plugins/VirtualProduction/Avalanche/Source/AvalancheFunctionalTest/Public/ScreenshotFunctionalTestAvalanche.h
--- a/Engine/Plugins/VirtualProduction/Avalanche/Source/AvalancheFunctionalTest/Public/ScreenshotFunctionalTestAvalanche.h
+++ b/Engine/Plugins/VirtualProduction/Avalanche/Source/AvalancheFunctionalTest/Public/ScreenshotFunctionalTestAvalanche.h
@@ -80,6 +80,10 @@ protected:
 	AVALANCHEFUNCTIONALTEST_API virtual void RequestScreenshot() override;
 	AVALANCHEFUNCTIONALTEST_API virtual void OnScreenShotCaptured(int32 InSizeX, int32 InSizeY, const TArray<FColor>& InImageData) override;
 
+	// Override to flush pending Niagara shader compilation before cleanup,
+	// preventing async tasks from holding dangling weak pointers to GC'd UNiagaraSystem objects.
+	AVALANCHEFUNCTIONALTEST_API virtual void FinishTest(EFunctionalTestResult TestResult, const FString& Message) override;
+
 private:
 	/**
 	 * Checks if the specified actor is fully initialized and ready for screenshot.
