Motion Design: moved the instance-only properties outside of transition tree into transition behavior.

#Synced-CL 52400047

diff --git a/Engine/Plugins/VirtualProduction/Avalanche/Source/AvalancheEditor/Private/Transition/AvaTransitionExtension.cpp b/Engine/Plugins/VirtualProduction/Avalanche/Source/AvalancheEditor/Private/Transition/AvaTransitionExtension.cpp
--- a/Engine/Plugins/VirtualProduction/Avalanche/Source/AvalancheEditor/Private/Transition/AvaTransitionExtension.cpp
+++ b/Engine/Plugins/VirtualProduction/Avalanche/Source/AvalancheEditor/Private/Transition/AvaTransitionExtension.cpp
@@ -164,18 +164,15 @@ void FAvaTransitionExtension::OnSceneObjectCreated(UObject* InSceneObject)
 		return;
 	}
 
-	const IAvaTransitionBehavior* const TransitionBehavior = TransitionSubsystem->GetOrCreateTransitionBehavior();
+	IAvaTransitionBehavior* const TransitionBehavior = TransitionSubsystem->GetOrCreateTransitionBehavior();
 	if (!TransitionBehavior)
 	{
 		return;
 	}
 
 	// Enable TL if creating the TL default scene
-	if (UAvaTransitionTree* const TransitionTree = TransitionBehavior->GetTransitionTree())
-	{
-		TransitionTree->SetEnabled(true);
-		TransitionTree->SetInstancingMode(EAvaTransitionInstancingMode::Reuse);
-	}
+	TransitionBehavior->SetEnabled(true);
+	TransitionBehavior->SetInstancingMode(EAvaTransitionInstancingMode::Reuse);
 
 	// Spawn Layer 1 Content null
 	FActorSpawnParameters SpawnParameters;
@@ -186,10 +183,10 @@ void FAvaTransitionExtension::OnSceneObjectCreated(UObject* InSceneObject)
 	UActorModifierCoreStack* ModifierStack;
 	if (UActorModifierCoreLibrary::FindModifierStack(Layer1Actor, ModifierStack, /*bCreateIfNone*/true))
 	{
-		constexpr const TCHAR* SubLayerModifierClassPath = TEXT("/Avalanche/TransitionLogic/SubLayer/BP_SubLayerModifier.BP_SubLayerModifier_C";
+		constexpr const TCHAR* SubLayerModifierClassPath = TEXT("/Avalanche/TransitionLogic/SubLayer/BP_SubLayerModifier.BP_SubLayerModifier_C");
 
 		FActorModifierCoreInsertOperation InsertOperation;
-		InsertOperation.ModifierClass = StaticLoadClass(UActorModifierCoreBase::StaticClass(), nullptr, SubLayerModifierClassPath));
+		InsertOperation.ModifierClass = StaticLoadClass(UActorModifierCoreBase::StaticClass(), nullptr, SubLayerModifierClassPath);
 
 		UActorModifierCoreBase* Modifier;
 		if (UActorModifierCoreLibrary::InsertModifier(ModifierStack, InsertOperation, Modifier))
diff --git a/Engine/Plugins/VirtualProduction/Avalanche/Source/AvalancheMedia/Private/Playable/Transition/AvaPlayableLocalTransition.cpp b/Engine/Plugins/VirtualProduction/Avalanche/Source/AvalancheMedia/Private/Playable/Transition/AvaPlayableLocalTransition.cpp
--- a/Engine/Plugins/VirtualProduction/Avalanche/Source/AvalancheMedia/Private/Playable/Transition/AvaPlayableLocalTransition.cpp
+++ b/Engine/Plugins/VirtualProduction/Avalanche/Source/AvalancheMedia/Private/Playable/Transition/AvaPlayableLocalTransition.cpp
@@ -65,10 +65,7 @@ namespace UE::AvaPlayableTransition::Private
 	{
 		if (InBehavior)
 		{
-			if (const UAvaTransitionTree* TransitionTree = InBehavior->GetTransitionTree())
-			{
-				return TransitionTree->GetTransitionLayer();
-			}
+			return InBehavior->GetTransitionLayer();
 		}
 		return FAvaTagHandle();
 	}
@@ -163,15 +160,12 @@ namespace UE::AvaPlayableTransition::Private
 		{
 			if (IAvaTransitionBehavior* TransitionBehavior = GetTransitionBehavior(*InPlayable, *InTransitionSubsystem))
 			{
-				if (const UAvaTransitionTree* TransitionTree = TransitionBehavior->GetTransitionTree())
+				if (TransitionBehavior->IsEnabled() && TransitionBehavior->GetTransitionTree())
 				{
-					if (TransitionTree->IsEnabled())
-					{
-						OutBehaviorInstance = FAvaTransitionBehaviorInstance();
-						OutBehaviorInstance.SetBehavior(TransitionBehavior);
-						OutBehaviorInstance.CreateScene<FAvaPlayableTransitionScene>(InPlayable, InPlayable, PlayableTransition);
-						return true;
-					}
+					OutBehaviorInstance = FAvaTransitionBehaviorInstance();
+					OutBehaviorInstance.SetBehavior(TransitionBehavior);
+					OutBehaviorInstance.CreateScene<FAvaPlayableTransitionScene>(InPlayable, InPlayable, PlayableTransition);
+					return true;
 				}
 			}
 			return false;
diff --git a/Engine/Plugins/VirtualProduction/Avalanche/Source/AvalancheMedia/Private/Rundown/AvaRundownPage.cpp b/Engine/Plugins/VirtualProduction/Avalanche/Source/AvalancheMedia/Private/Rundown/AvaRundownPage.cpp
--- a/Engine/Plugins/VirtualProduction/Avalanche/Source/AvalancheMedia/Private/Rundown/AvaRundownPage.cpp
+++ b/Engine/Plugins/VirtualProduction/Avalanche/Source/AvalancheMedia/Private/Rundown/AvaRundownPage.cpp
@@ -5,6 +5,7 @@
 #include "AssetRegistry/AssetData.h"
 #include "AvaScene.h"
 #include "AvaTransitionTree.h"
+#include "Behavior/IAvaTransitionBehavior.h"
 #include "Broadcast/AvaBroadcast.h"
 #include "IAvaMediaModule.h"
 #include "Playable/AvaPlayableRemoteControl.h"
@@ -198,12 +199,12 @@ bool FAvaRundownPage::UpdatePageSummary(const UAvaRundown* InRundown, const TArr
 		InitialDescription.Reserve(OrderedDescription.Num());
 
 		constexpr int32 FinishDescriptionWhenExceedLimit = 50;
-		int32 CurrentLenght = 0;
+		int32 CurrentLength = 0;
 		for (const TPair<int32, FText>& OrderedValue : OrderedDescription)
 		{
 			InitialDescription.Add(OrderedValue.Value);
-			CurrentLenght += OrderedValue.Value.ToString().Len();
-			if (CurrentLenght >= FinishDescriptionWhenExceedLimit)
+			CurrentLength += OrderedValue.Value.ToString().Len();
+			if (CurrentLength >= FinishDescriptionWhenExceedLimit)
 			{
 				InitialDescription.Add(LOCTEXT("EndOfDescription", "..."));
 				break;
@@ -224,10 +225,10 @@ bool FAvaRundownPage::UpdateTransitionLogic()
 	{
 		if (const IAvaSceneInterface* SceneInterface = FAvaRundownPageAssetUtils::GetSceneInterface(LoadedSourceAsset))
 		{
-			const UAvaTransitionTree* TransitionTree = FAvaRundownPageAssetUtils::FindTransitionTree(SceneInterface);
-			bHasTransitionLogic = TransitionTree ? TransitionTree->IsEnabled() : false;
-			TransitionLayerTag = FAvaRundownPageAssetUtils::GetTransitionLayerTag(TransitionTree);
-			TransitionMode = TransitionTree ? TransitionTree->GetInstancingMode() : EAvaTransitionInstancingMode::New;
+			const IAvaTransitionBehavior* TransitionBehavior = FAvaRundownPageAssetUtils::FindTransitionBehavior(SceneInterface);
+			bHasTransitionLogic = TransitionBehavior ? TransitionBehavior->IsEnabled() : false;
+			TransitionLayerTag = TransitionBehavior ? TransitionBehavior->GetTransitionLayer() : FAvaTagHandle();
+			TransitionMode = TransitionBehavior ? TransitionBehavior->GetInstancingMode() : EAvaTransitionInstancingMode::New;
 		}
 	}
 	return false;
diff --git a/Engine/Plugins/VirtualProduction/Avalanche/Source/AvalancheMedia/Private/Rundown/AvaRundownPageAssetUtils.cpp b/Engine/Plugins/VirtualProduction/Avalanche/Source/AvalancheMedia/Private/Rundown/AvaRundownPageAssetUtils.cpp
--- a/Engine/Plugins/VirtualProduction/Avalanche/Source/AvalancheMedia/Private/Rundown/AvaRundownPageAssetUtils.cpp
+++ b/Engine/Plugins/VirtualProduction/Avalanche/Source/AvalancheMedia/Private/Rundown/AvaRundownPageAssetUtils.cpp
@@ -37,7 +37,7 @@ const IAvaSceneInterface* FAvaRundownPageAssetUtils::GetSceneInterface(const UOb
 	return GetInterfaceFromAsset<IAvaSceneInterface>(InLoadedAsset);
 }
 
-const UAvaTransitionTree* FAvaRundownPageAssetUtils::FindTransitionTree(const IAvaSceneInterface* InSceneInterface)
+const IAvaTransitionBehavior* FAvaRundownPageAssetUtils::FindTransitionBehavior(const IAvaSceneInterface* InSceneInterface)
 {
 	if (!InSceneInterface)
 	{
@@ -56,40 +56,11 @@ const UAvaTransitionTree* FAvaRundownPageAssetUtils::FindTransitionTree(const IA
 	{
 		if (const UAvaTransitionSubsystem* TransitionSubsystem = SceneWorld->GetSubsystem<UAvaTransitionSubsystem>())
 		{
-			if (const IAvaTransitionBehavior* TransitionBehavior = TransitionSubsystem->GetTransitionBehavior(SceneLevel))
-			{
-				if (const UAvaTransitionTree* TransitionTree = TransitionBehavior->GetTransitionTree())
-				{
-					return TransitionTree;
-				}
-			}
-			// If we had a subsystem and it didn't find the behavior then there was no state tree either and we are done.
-			return nullptr;
+			return TransitionSubsystem->GetTransitionBehavior(SceneLevel);
 		}
 	}
-	
+
 	// Method 2 - Fallback using direct lookup of the private actor class.
 	// This is needed to lookup in managed instances and source asset.
-
-	static FString AvaTransitionBehaviorActor(TEXT("AvaTransitionBehaviorActor"));
-	for (const AActor* Actor : SceneLevel->Actors)
-	{
-		const UClass* ActorClass = IsValid(Actor) ? Actor->GetClass() : nullptr;
-		if (ActorClass && ActorClass->GetName() == AvaTransitionBehaviorActor)
-		{
-			const FObjectPropertyBase* const Property = FindFProperty<FObjectPropertyBase>(ActorClass, TEXT("TransitionTree"));
-			if (Property
-				&& Property->PropertyClass
-				&& Property->PropertyClass->IsChildOf(UAvaTransitionTree::StaticClass()))
-			{
-				return Cast<UAvaTransitionTree>(Property->GetObjectPropertyValue_InContainer(Actor));
-			}
-		}
-	}
-	return nullptr;
+	return UAvaTransitionSubsystem::FindTransitionBehavior(SceneLevel);
 }
-
-FAvaTagHandle FAvaRundownPageAssetUtils::GetTransitionLayerTag(const UAvaTransitionTree* InTransitionTree)
-{
-	return InTransitionTree ? InTransitionTree->GetTransitionLayer() : FAvaTagHandle();
-}
\ No newline at end of file
diff --git a/Engine/Plugins/VirtualProduction/Avalanche/Source/AvalancheMedia/Private/Rundown/AvaRundownPageAssetUtils.h b/Engine/Plugins/VirtualProduction/Avalanche/Source/AvalancheMedia/Private/Rundown/AvaRundownPageAssetUtils.h
--- a/Engine/Plugins/VirtualProduction/Avalanche/Source/AvalancheMedia/Private/Rundown/AvaRundownPageAssetUtils.h
+++ b/Engine/Plugins/VirtualProduction/Avalanche/Source/AvalancheMedia/Private/Rundown/AvaRundownPageAssetUtils.h
@@ -6,6 +6,7 @@
 
 class IAvaSceneInterface;
 class IAvaSequenceProvider;
+class IAvaTransitionBehavior;
 class UAvaTransitionTree;
 class UObject;
 struct FAvaTagHandle;
@@ -20,6 +21,5 @@ class FAvaRundownPageAssetUtils
 {
 public:
 	static const IAvaSceneInterface* GetSceneInterface(const UObject* InLoadedAsset);
-	static const UAvaTransitionTree* FindTransitionTree(const IAvaSceneInterface* InSceneInterface);
-	static FAvaTagHandle GetTransitionLayerTag(const UAvaTransitionTree* InTransitionTree);
+	static const IAvaTransitionBehavior* FindTransitionBehavior(const IAvaSceneInterface* InSceneInterface);
 };
\ No newline at end of file
diff --git a/Engine/Plugins/VirtualProduction/Avalanche/Source/AvalancheTransition/Private/AvaTransitionLibrary.cpp b/Engine/Plugins/VirtualProduction/Avalanche/Source/AvalancheTransition/Private/AvaTransitionLibrary.cpp
--- a/Engine/Plugins/VirtualProduction/Avalanche/Source/AvalancheTransition/Private/AvaTransitionLibrary.cpp
+++ b/Engine/Plugins/VirtualProduction/Avalanche/Source/AvalancheTransition/Private/AvaTransitionLibrary.cpp
@@ -310,22 +310,36 @@ bool UAvaTransitionLibrary::AreScenesTransitioning(UObject* InContextObject
 	return false;
 }
 
-const UAvaTransitionTree* UAvaTransitionLibrary::GetTransitionTree(UObject* InContextObject)
+TScriptInterface<IAvaTransitionBehavior> UAvaTransitionLibrary::GetTransitionBehavior(UObject* InContextObject)
 {
 	using namespace UE::AvaTransition;
 
 	if (const IAvaTransitionNodeInterface* NodeInterface = Cast<IAvaTransitionNodeInterface>(InContextObject))
 	{
-		return NodeInterface->GetBehaviorInstanceCache().GetTransitionTree();
+		if (IAvaTransitionBehavior* TransitionBehavior = NodeInterface->GetBehaviorInstanceCache().GetTransitionBehavior())
+		{
+			return &TransitionBehavior->AsUObject();
+		}
 	}
 
 	if (ULevel* Level = Private::GetLevelFromContextObject(InContextObject))
 	{
 		if (UAvaTransitionSubsystem* const TransitionSubsystem = Private::GetSubsystemFromLevel(Level))
 		{
-			const IAvaTransitionBehavior* TransitionBehavior = TransitionSubsystem->GetTransitionBehavior(Level);
-			return TransitionBehavior ? TransitionBehavior->GetTransitionTree() : nullptr;
+			if (IAvaTransitionBehavior* TransitionBehavior = TransitionSubsystem->GetTransitionBehavior(Level))
+			{
+				return &TransitionBehavior->AsUObject();
+			}
 		}
 	}
 	return nullptr;
 }
+
+const UAvaTransitionTree* UAvaTransitionLibrary::GetTransitionTree(UObject* InContextObject)
+{
+	if (const IAvaTransitionBehavior* TransitionBehavior = UAvaTransitionLibrary::GetTransitionBehavior(InContextObject).GetInterface())
+	{
+		return TransitionBehavior->GetTransitionTree();
+	}
+	return nullptr;
+}
diff --git a/Engine/Plugins/VirtualProduction/Avalanche/Source/AvalancheTransition/Private/AvaTransitionTree.cpp b/Engine/Plugins/VirtualProduction/Avalanche/Source/AvalancheTransition/Private/AvaTransitionTree.cpp
--- a/Engine/Plugins/VirtualProduction/Avalanche/Source/AvalancheTransition/Private/AvaTransitionTree.cpp
+++ b/Engine/Plugins/VirtualProduction/Avalanche/Source/AvalancheTransition/Private/AvaTransitionTree.cpp
@@ -1,27 +1,34 @@
 // Copyright Epic Games, Inc. All Rights Reserved.
 
 #include "AvaTransitionTree.h"
-
 #include "StateTreeTaskBase.h"
 
 FAvaTagHandle UAvaTransitionTree::GetTransitionLayer() const
 {
-	return TransitionLayer;
+PRAGMA_DISABLE_DEPRECATION_WARNINGS
+	return TransitionLayer_DEPRECATED;
+PRAGMA_ENABLE_DEPRECATION_WARNINGS
 }
 
 void UAvaTransitionTree::SetTransitionLayer(FAvaTagHandle InTransitionLayer)
 {
-	TransitionLayer = InTransitionLayer;
+PRAGMA_DISABLE_DEPRECATION_WARNINGS
+	TransitionLayer_DEPRECATED = InTransitionLayer;
+PRAGMA_ENABLE_DEPRECATION_WARNINGS
 }
 
 bool UAvaTransitionTree::IsEnabled() const
 {
-	return bEnabled;
+PRAGMA_DISABLE_DEPRECATION_WARNINGS
+	return bEnabled_DEPRECATED;
+PRAGMA_ENABLE_DEPRECATION_WARNINGS
 }
 
 void UAvaTransitionTree::SetEnabled(bool bInEnabled)
 {
-	bEnabled = bInEnabled;
+PRAGMA_DISABLE_DEPRECATION_WARNINGS
+	bEnabled_DEPRECATED = bInEnabled;
+PRAGMA_ENABLE_DEPRECATION_WARNINGS
 }
 
 bool UAvaTransitionTree::ContainsTask(const UScriptStruct* InTaskStruct) const
@@ -51,3 +58,24 @@ bool UAvaTransitionTree::ContainsTask(const UScriptStruct* InTaskStruct) const
 
 	return false;
 }
+
+void UAvaTransitionTree::SetInstancingMode(EAvaTransitionInstancingMode InInstancingMode)
+{
+PRAGMA_DISABLE_DEPRECATION_WARNINGS
+	InstancingMode_DEPRECATED = InInstancingMode;
+PRAGMA_ENABLE_DEPRECATION_WARNINGS
+}
+
+EAvaTransitionInstancingMode UAvaTransitionTree::GetInstancingMode() const
+{
+PRAGMA_DISABLE_DEPRECATION_WARNINGS
+	return InstancingMode_DEPRECATED;
+PRAGMA_ENABLE_DEPRECATION_WARNINGS
+}
+
+FName UAvaTransitionTree::GetEnabledPropertyName()
+{
+PRAGMA_DISABLE_DEPRECATION_WARNINGS
+	return GET_MEMBER_NAME_CHECKED(UAvaTransitionTree, bEnabled_DEPRECATED);
+PRAGMA_ENABLE_DEPRECATION_WARNINGS
+}
diff --git a/Engine/Plugins/VirtualProduction/Avalanche/Source/AvalancheTransition/Private/AvaTransitionVersion.cpp b/Engine/Plugins/VirtualProduction/Avalanche/Source/AvalancheTransition/Private/AvaTransitionVersion.cpp
new file mode 100644
--- /dev/null
+++ b/Engine/Plugins/VirtualProduction/Avalanche/Source/AvalancheTransition/Private/AvaTransitionVersion.cpp
@@ -0,0 +1,11 @@
+// Copyright Epic Games, Inc. All Rights Reserved.
+
+#include "AvaTransitionVersion.h"
+#include "Serialization/CustomVersion.h"
+
+namespace UE::AvaTransition
+{
+	const FGuid FBehaviorVersion::Guid(0x57DE9A66, 0x3FB04693, 0x8ED244A4, 0x2FD1D9BD);
+	FCustomVersionRegistration GRegisterBehaviorVersion(FBehaviorVersion::Guid, FBehaviorVersion::LatestVersion, TEXT("AvaTransitionBehaviorVersion"));
+
+} // UE::AvaTransition
diff --git a/Engine/Plugins/VirtualProduction/Avalanche/Source/AvalancheTransition/Private/Behavior/AvaTransitionBehaviorActor.cpp b/Engine/Plugins/VirtualProduction/Avalanche/Source/AvalancheTransition/Private/Behavior/AvaTransitionBehaviorActor.cpp
--- a/Engine/Plugins/VirtualProduction/Avalanche/Source/AvalancheTransition/Private/Behavior/AvaTransitionBehaviorActor.cpp
+++ b/Engine/Plugins/VirtualProduction/Avalanche/Source/AvalancheTransition/Private/Behavior/AvaTransitionBehaviorActor.cpp
@@ -2,11 +2,13 @@
 
 #include "AvaTransitionBehaviorActor.h"
 #include "AvaTransitionContext.h"
+#include "AvaTransitionDefaultLayerTags.h"
 #include "AvaTransitionEnums.h"
 #include "AvaTransitionScene.h"
 #include "AvaTransitionSubsystem.h"
 #include "AvaTransitionTree.h"
 #include "AvaTransitionTreeSchema.h"
+#include "AvaTransitionVersion.h"
 #include "Engine/World.h"
 #include "IAvaTransitionModule.h"
 #include "StateTreeDelegates.h"
@@ -26,6 +28,46 @@ AAvaTransitionBehaviorActor::AAvaTransitionBehaviorActor()
 	StateTreeReference.SetStateTree(TransitionTree);
 }
 
+UAvaTransitionTree* AAvaTransitionBehaviorActor::GetTransitionTree() const
+{
+	return TransitionTree;
+}
+
+FAvaTagHandle AAvaTransitionBehaviorActor::GetTransitionLayer() const
+{
+	return TransitionLayer;
+}
+
+void AAvaTransitionBehaviorActor::SetTransitionLayer(FAvaTagHandle InTransitionLayer)
+{
+	TransitionLayer = InTransitionLayer;
+}
+
+bool AAvaTransitionBehaviorActor::IsEnabled() const
+{
+	return bEnabled;
+}
+
+void AAvaTransitionBehaviorActor::SetEnabled(bool bInEnabled)
+{
+	bEnabled = bInEnabled;
+}
+
+EAvaTransitionInstancingMode AAvaTransitionBehaviorActor::GetInstancingMode() const
+{
+	return InstancingMode;
+}
+
+void AAvaTransitionBehaviorActor::SetInstancingMode(EAvaTransitionInstancingMode InInstancingMode)
+{
+	InstancingMode = InInstancingMode;
+}
+
+const FStateTreeReference& AAvaTransitionBehaviorActor::GetStateTreeReference() const
+{
+	return StateTreeReference;
+}
+
 void AAvaTransitionBehaviorActor::PostActorCreated()
 {
 	Super::PostActorCreated();
@@ -38,9 +80,28 @@ void AAvaTransitionBehaviorActor::PostActorCreated()
 	}
 }
 
+void AAvaTransitionBehaviorActor::Serialize(FArchive& Ar)
+{
+	Super::Serialize(Ar);
+	Ar.UsingCustomVersion(UE::AvaTransition::FBehaviorVersion::Guid);
+}
+
 void AAvaTransitionBehaviorActor::PostLoad()
 {
 	Super::PostLoad();
+
+	if (GetLinkerCustomVersion(UE::AvaTransition::FBehaviorVersion::Guid) < UE::AvaTransition::FBehaviorVersion::InstanceProperties)
+	{
+		if (TransitionTree)
+		{
+PRAGMA_DISABLE_DEPRECATION_WARNINGS
+			TransitionLayer = TransitionTree->GetTransitionLayer();
+			InstancingMode = TransitionTree->GetInstancingMode();
+			bEnabled = TransitionTree->IsEnabled();
+PRAGMA_ENABLE_DEPRECATION_WARNINGS
+		}
+	}
+
 	ValidateTransitionTree();
 }
 
diff --git a/Engine/Plugins/VirtualProduction/Avalanche/Source/AvalancheTransition/Private/Behavior/AvaTransitionBehaviorActor.h b/Engine/Plugins/VirtualProduction/Avalanche/Source/AvalancheTransition/Private/Behavior/AvaTransitionBehaviorActor.h
--- a/Engine/Plugins/VirtualProduction/Avalanche/Source/AvalancheTransition/Private/Behavior/AvaTransitionBehaviorActor.h
+++ b/Engine/Plugins/VirtualProduction/Avalanche/Source/AvalancheTransition/Private/Behavior/AvaTransitionBehaviorActor.h
@@ -25,8 +25,14 @@ public:
 protected:
 	//~ Begin IAvaTransitionBehavior
 	virtual UObject& AsUObject() override { return *this; }
-	virtual UAvaTransitionTree* GetTransitionTree() const override { return TransitionTree; }
-	virtual const FStateTreeReference& GetStateTreeReference() const override { return StateTreeReference; }
+	virtual UAvaTransitionTree* GetTransitionTree() const override;
+	virtual FAvaTagHandle GetTransitionLayer() const override;
+	virtual void SetTransitionLayer(FAvaTagHandle InTransitionLayer) override;
+	virtual bool IsEnabled() const override;
+	virtual void SetEnabled(bool bInEnabled) override;
+	virtual EAvaTransitionInstancingMode GetInstancingMode() const override;
+	virtual void SetInstancingMode(EAvaTransitionInstancingMode InInstancingMode) override;
+	virtual const FStateTreeReference& GetStateTreeReference() const override;
 	//~ End IAvaTransitionBehavior
 
 	//~ Begin AActor
@@ -38,6 +44,7 @@ protected:
 	//~ End AActor
 
 	//~ Begin UObject
+	virtual void Serialize(FArchive& Ar) override;
 	virtual void PostLoad() override;
 #if WITH_EDITOR
 	virtual bool CanEditChange(const FProperty* InProperty) const override;
@@ -55,4 +62,18 @@ private:
 
 	UPROPERTY(EditAnywhere, Category="Transition Logic")
 	FStateTreeReference StateTreeReference;
+
+	/** Transition Layer to use */
+	UPROPERTY(EditAnywhere, Category="Transition Logic")
+	FAvaTagHandle TransitionLayer;
+
+	UPROPERTY(EditAnywhere, Category="Transition Logic")
+	EAvaTransitionInstancingMode InstancingMode = EAvaTransitionInstancingMode::New;
+
+	/**
+	 * Determines whether this Transition Logic is enabled.
+	 * Can be overriden by a Transition Instance to force the logic to run regardless
+	 */
+	UPROPERTY(EditAnywhere, Category="Transition Logic", meta=(DisplayPriority=0))
+	bool bEnabled = false;
 };
diff --git a/Engine/Plugins/VirtualProduction/Avalanche/Source/AvalancheTransition/Private/Behavior/AvaTransitionBehaviorInstance.cpp b/Engine/Plugins/VirtualProduction/Avalanche/Source/AvalancheTransition/Private/Behavior/AvaTransitionBehaviorInstance.cpp
--- a/Engine/Plugins/VirtualProduction/Avalanche/Source/AvalancheTransition/Private/Behavior/AvaTransitionBehaviorInstance.cpp
+++ b/Engine/Plugins/VirtualProduction/Avalanche/Source/AvalancheTransition/Private/Behavior/AvaTransitionBehaviorInstance.cpp
@@ -20,18 +20,7 @@ FAvaTransitionBehaviorInstance& FAvaTransitionBehaviorInstance::SetBehavior(IAva
 bool FAvaTransitionBehaviorInstance::IsEnabled() const
 {
 	const IAvaTransitionBehavior* TransitionBehavior = GetBehavior();
-	if (!TransitionBehavior)
-	{
-		return false;
-	}
-
-	const UAvaTransitionTree* TransitionTree = TransitionBehavior->GetTransitionTree();
-	if (!TransitionTree)
-	{
-		return false;
-	}
-
-	return TransitionTree->IsEnabled();
+	return TransitionBehavior && TransitionBehavior->IsEnabled();
 }
 
 IAvaTransitionBehavior* FAvaTransitionBehaviorInstance::GetBehavior() const
@@ -213,16 +202,7 @@ bool FAvaTransitionBehaviorInstance::ValidateTransitionScene()
 void FAvaTransitionBehaviorInstance::UpdateTransitionLayers(const IAvaTransitionBehavior* InBehavior)
 {
 	FAvaTransitionScene& TransitionScene = *TransitionContext.GetTransitionScene();
-
-	const UAvaTransitionTree* TransitionTree = InBehavior ? InBehavior->GetTransitionTree() : nullptr;
-	if (TransitionTree)
-	{
-		TransitionContext.TransitionLayer = TransitionTree->GetTransitionLayer();
-	}
-	else
-	{
-		TransitionContext.TransitionLayer = FAvaTagHandle();
-	}
+	TransitionContext.TransitionLayer = InBehavior ? InBehavior->GetTransitionLayer() : FAvaTagHandle();
 
 	// Give opportunity for the Scene to override the layer (if not already overriden)
 	if (!OverrideLayer.IsSet())
diff --git a/Engine/Plugins/VirtualProduction/Avalanche/Source/AvalancheTransition/Private/Behavior/AvaTransitionInstanceCache.cpp b/Engine/Plugins/VirtualProduction/Avalanche/Source/AvalancheTransition/Private/Behavior/AvaTransitionInstanceCache.cpp
--- a/Engine/Plugins/VirtualProduction/Avalanche/Source/AvalancheTransition/Private/Behavior/AvaTransitionInstanceCache.cpp
+++ b/Engine/Plugins/VirtualProduction/Avalanche/Source/AvalancheTransition/Private/Behavior/AvaTransitionInstanceCache.cpp
@@ -13,11 +13,11 @@ const FAvaTransitionContext* FAvaTransitionBehaviorInstanceCache::GetTransitionC
 	return nullptr;
 }
 
-const UAvaTransitionTree* FAvaTransitionBehaviorInstanceCache::GetTransitionTree() const
+IAvaTransitionBehavior* FAvaTransitionBehaviorInstanceCache::GetTransitionBehavior() const
 {
 	if (CachedBehaviorInstance && CachedSceneOwner.IsValid())
 	{
-		return CachedBehaviorInstance->GetTransitionTree();
+		return CachedBehaviorInstance->GetBehavior();
 	}
 	return nullptr;
 }
diff --git a/Engine/Plugins/VirtualProduction/Avalanche/Source/AvalancheTransition/Private/Execution/AvaTransitionExecutor.cpp b/Engine/Plugins/VirtualProduction/Avalanche/Source/AvalancheTransition/Private/Execution/AvaTransitionExecutor.cpp
--- a/Engine/Plugins/VirtualProduction/Avalanche/Source/AvalancheTransition/Private/Execution/AvaTransitionExecutor.cpp
+++ b/Engine/Plugins/VirtualProduction/Avalanche/Source/AvalancheTransition/Private/Execution/AvaTransitionExecutor.cpp
@@ -129,15 +129,15 @@ void FAvaTransitionExecutor::Setup()
 			{
 				bool bSceneReused = false;
 
-				const UAvaTransitionTree* TransitionTree = ExitInstance.GetTransitionTree();
-				if (TransitionTree && TransitionTree->GetInstancingMode() == EAvaTransitionInstancingMode::Reuse)
+				const IAvaTransitionBehavior* TransitionBehavior = ExitInstance.GetBehavior();
+				if (TransitionBehavior && TransitionBehavior->GetInstancingMode() == EAvaTransitionInstancingMode::Reuse)
 				{
-					// Scene reused if the Tree Instancing Mode is set to reuse, and if there is an Enter Instance with matching Tree (i.e. Level)
+					// Scene reused if the Tree Instancing Mode is set to reuse, and if there is an Enter Instance with matching behavior (i.e. Level)
 					bSceneReused |= Pair.Value.EnterInstances.ContainsByPredicate(
-						[TransitionTree, &Instances = Instances](int32 InEnterInstanceIndex)
+						[TransitionBehavior, &Instances = Instances](int32 InEnterInstanceIndex)
 						{
 							const FAvaTransitionBehaviorInstance& EnterInstance = Instances[InEnterInstanceIndex];
-							return EnterInstance.GetTransitionTree() == TransitionTree;
+							return EnterInstance.GetBehavior() == TransitionBehavior;
 						});
 				}
 
diff --git a/Engine/Plugins/VirtualProduction/Avalanche/Source/AvalancheTransition/Private/Preview/AvaTransitionPreviewManager.cpp b/Engine/Plugins/VirtualProduction/Avalanche/Source/AvalancheTransition/Private/Preview/AvaTransitionPreviewManager.cpp
--- a/Engine/Plugins/VirtualProduction/Avalanche/Source/AvalancheTransition/Private/Preview/AvaTransitionPreviewManager.cpp
+++ b/Engine/Plugins/VirtualProduction/Avalanche/Source/AvalancheTransition/Private/Preview/AvaTransitionPreviewManager.cpp
@@ -98,15 +98,19 @@ void AAvaTransitionPreviewManager::TakeOut()
 		return;
 	}
 
+	bTransitionInProgress = true;
+	NextLevelState.LevelStreaming = nullptr;
+	NextLevelState.bEnableOverrideTransitionLayer = true;
+	NextLevelState.OverrideTransitionLayer = TransitionBehavior->GetTransitionLayer();
+
 	if (UAvaTransitionTree* TransitionTree = TransitionBehavior->GetTransitionTree())
 	{
-		bTransitionInProgress = true;
-
-		NextLevelState.LevelStreaming = nullptr;
-		NextLevelState.bEnableOverrideTransitionLayer = true;
-		NextLevelState.OverrideTransitionLayer = TransitionTree->GetTransitionLayer();
 		StartTransition();
 	}
+	else
+	{
+		OnTransitionEnded();
+	}
 }
 
 void AAvaTransitionPreviewManager::TransitionStop()
diff --git a/Engine/Plugins/VirtualProduction/Avalanche/Source/AvalancheTransition/Private/Tasks/AvaTransitionDelayTask.cpp b/Engine/Plugins/VirtualProduction/Avalanche/Source/AvalancheTransition/Private/Tasks/AvaTransitionDelayTask.cpp
--- a/Engine/Plugins/VirtualProduction/Avalanche/Source/AvalancheTransition/Private/Tasks/AvaTransitionDelayTask.cpp
+++ b/Engine/Plugins/VirtualProduction/Avalanche/Source/AvalancheTransition/Private/Tasks/AvaTransitionDelayTask.cpp
@@ -3,8 +3,8 @@
 #include "Tasks/AvaTransitionDelayTask.h"
 #include "AvaTransitionContext.h"
 #include "AvaTransitionScene.h"
-#include "AvaTransitionTree.h"
 #include "AvaTransitionUtils.h"
+#include "Behavior/IAvaTransitionBehavior.h"
 #include "Engine/Level.h"
 #include "Rendering/AvaTransitionRenderingSubsystem.h"
 #include "StateTreeExecutionContext.h"
@@ -83,8 +83,8 @@ bool FAvaTransitionDelayTask::ShouldHideLevel(const FStateTreeExecutionContext&
 	}
 
 	// If Instancing Mode is set to Reuse, and user set to Hide Level regardless, allow Hiding Reused Level
-	const UAvaTransitionTree* TransitionTree = Cast<UAvaTransitionTree>(InContext.GetStateTree());
-	if (TransitionTree && TransitionTree->GetInstancingMode() == EAvaTransitionInstancingMode::Reuse)
+	const IAvaTransitionBehavior* TransitionBehavior = Cast<IAvaTransitionBehavior>(InContext.GetOwner());
+	if (TransitionBehavior && TransitionBehavior->GetInstancingMode() == EAvaTransitionInstancingMode::Reuse)
 	{
 		return InInstanceData.HideMode == EAvaTransitionLevelHideMode::AlwaysHide;
 	}
diff --git a/Engine/Plugins/VirtualProduction/Avalanche/Source/AvalancheTransition/Private/Tasks/AvaTransitionWaitForLayerTask.cpp b/Engine/Plugins/VirtualProduction/Avalanche/Source/AvalancheTransition/Private/Tasks/AvaTransitionWaitForLayerTask.cpp
--- a/Engine/Plugins/VirtualProduction/Avalanche/Source/AvalancheTransition/Private/Tasks/AvaTransitionWaitForLayerTask.cpp
+++ b/Engine/Plugins/VirtualProduction/Avalanche/Source/AvalancheTransition/Private/Tasks/AvaTransitionWaitForLayerTask.cpp
@@ -2,8 +2,8 @@
 
 #include "Tasks/AvaTransitionWaitForLayerTask.h"
 #include "AvaTransitionLayerUtils.h"
-#include "AvaTransitionTree.h"
 #include "Behavior/AvaTransitionBehaviorInstance.h"
+#include "Behavior/IAvaTransitionBehavior.h"
 #include "Engine/Level.h"
 #include "Rendering/AvaTransitionRenderingSubsystem.h"
 #include "StateTreeExecutionContext.h"
@@ -87,8 +87,8 @@ bool FAvaTransitionWaitForLayerTask::ShouldHideLevel(const FStateTreeExecutionCo
 	}
 
 	// If Instancing Mode is set to Reuse, and user set to Hide Level regardless, allow Hiding Reused Level
-	const UAvaTransitionTree* TransitionTree = Cast<UAvaTransitionTree>(InContext.GetStateTree());
-	if (TransitionTree && TransitionTree->GetInstancingMode() == EAvaTransitionInstancingMode::Reuse)
+	const IAvaTransitionBehavior* TransitionBehavior = Cast<IAvaTransitionBehavior>(InContext.GetOwner());
+	if (TransitionBehavior && TransitionBehavior->GetInstancingMode() == EAvaTransitionInstancingMode::Reuse)
 	{
 		return InInstanceData.HideMode == EAvaTransitionLevelHideMode::AlwaysHide;
 	}
diff --git a/Engine/Plugins/VirtualProduction/Avalanche/Source/AvalancheTransition/Public/AvaTransitionLibrary.h b/Engine/Plugins/VirtualProduction/Avalanche/Source/AvalancheTransition/Public/AvaTransitionLibrary.h
--- a/Engine/Plugins/VirtualProduction/Avalanche/Source/AvalancheTransition/Public/AvaTransitionLibrary.h
+++ b/Engine/Plugins/VirtualProduction/Avalanche/Source/AvalancheTransition/Public/AvaTransitionLibrary.h
@@ -8,6 +8,7 @@
 
 #define UE_API AVALANCHETRANSITION_API
 
+class IAvaTransitionBehavior;
 class UAvaTransitionTree;
 struct FAvaTagHandleContainer;
 struct FAvaTransitionContext;
@@ -62,6 +63,12 @@ public:
 		, const FAvaTagHandleContainer& InLayers
 		, const TArray<TSoftObjectPtr<UWorld>>& InScenesToIgnore);
 
+	/**
+	 * Retrieves the transition behavior from this context
+	 */
+	UFUNCTION(BlueprintPure, Category="Transition Logic", meta=(DefaultToSelf="InContextObject"))
+	static UE_API TScriptInterface<IAvaTransitionBehavior> GetTransitionBehavior(UObject* InContextObject);
+
 	/**
 	 * Retrieves the transition tree from this context
 	 */
diff --git a/Engine/Plugins/VirtualProduction/Avalanche/Source/AvalancheTransition/Public/AvaTransitionTree.h b/Engine/Plugins/VirtualProduction/Avalanche/Source/AvalancheTransition/Public/AvaTransitionTree.h
--- a/Engine/Plugins/VirtualProduction/Avalanche/Source/AvalancheTransition/Public/AvaTransitionTree.h
+++ b/Engine/Plugins/VirtualProduction/Avalanche/Source/AvalancheTransition/Public/AvaTransitionTree.h
@@ -7,6 +7,8 @@
 #include "StateTree.h"
 #include "AvaTransitionTree.generated.h"
 
+#define UE_API AVALANCHETRANSITION_API
+
 struct FAvaTransitionTask;
 
 /**
@@ -19,16 +21,20 @@ class UAvaTransitionTree : public UStateTree
 	GENERATED_BODY()
 
 public:
-	AVALANCHETRANSITION_API FAvaTagHandle GetTransitionLayer() const;
+	UE_DEPRECATED(5.8, "GetTransitionLayer deprecated. Use IAvaTransitionBehavior::GetTransitionLayer instead")
+	UE_API FAvaTagHandle GetTransitionLayer() const;
 
-	AVALANCHETRANSITION_API void SetTransitionLayer(FAvaTagHandle InTransitionLayer);
+	UE_DEPRECATED(5.8, "SetTransitionLayer deprecated. Use IAvaTransitionBehavior::SetTransitionLayer instead")
+	UE_API void SetTransitionLayer(FAvaTagHandle InTransitionLayer);
 
-	AVALANCHETRANSITION_API bool IsEnabled() const;
+	UE_DEPRECATED(5.8, "IsEnabled deprecated. Use IAvaTransitionBehavior::IsEnabled instead")
+	UE_API bool IsEnabled() const;
 
-	AVALANCHETRANSITION_API void SetEnabled(bool bInEnabled);
+	UE_DEPRECATED(5.8, "SetEnabled deprecated. Use IAvaTransitionBehavior::SetEnabled instead")
+	UE_API void SetEnabled(bool bInEnabled);
 
 	/** Returns whether a Task of a given type exists and is enabled within an enabled state in the Transition Tree */
-	AVALANCHETRANSITION_API bool ContainsTask(const UScriptStruct* InTaskStruct) const;
+	UE_API bool ContainsTask(const UScriptStruct* InTaskStruct) const;
 
 	/** Returns whether a Task of a given type exists and is enabled within an enabled state in the Transition Tree */
 	template <class InTaskType UE_REQUIRES(std::is_base_of_v<FAvaTransitionTask, InTaskType>)>
@@ -37,34 +43,28 @@ public:
 		return ContainsTask(InTaskType::StaticStruct());
 	}
 
-	void SetInstancingMode(EAvaTransitionInstancingMode InInstancingMode)
-	{
-		InstancingMode = InInstancingMode;
-	}
+	UE_DEPRECATED(5.8, "SetInstancingMode deprecated. Use IAvaTransitionBehavior::SetInstancingMode instead")
+	UE_API void SetInstancingMode(EAvaTransitionInstancingMode InInstancingMode);
 
-	UFUNCTION(BlueprintPure, Category = "Transition Logic")
-	EAvaTransitionInstancingMode GetInstancingMode() const
-	{
-		return InstancingMode;
-	}
+	UE_DEPRECATED(5.8, "GetInstancingMode deprecated. Use IAvaTransitionBehavior::GetInstancingMode instead")
+	UFUNCTION(BlueprintPure, Category = "Transition Logic", meta=(DeprecatedFunction, DeprecationMessage="Use GetInstancingMode from the Transition Behavior Interface (via GetTransitionBehavior)instead."))
+	UE_API EAvaTransitionInstancingMode GetInstancingMode() const;
 
-	static FName GetEnabledPropertyName()
-	{
-		return GET_MEMBER_NAME_CHECKED(UAvaTransitionTree, bEnabled);
-	}
+	UE_DEPRECATED(5.8, "GetEnabledPropertyName deprecated. Use enabled via IAvaTransitionBehavior instead")
+	UE_API static FName GetEnabledPropertyName();
 
 private:
-	/** The Layer this Transition Logic Tree deals with */
+	UE_DEPRECATED(5.8, "TransitionLayer deprecated. Use IAvaTransitionBehavior instead")
 	UPROPERTY()
-	FAvaTagHandle TransitionLayer;
+	FAvaTagHandle TransitionLayer_DEPRECATED;
 
-	/**
-	 * Determines whether this Transition Logic is enabled, by default.
-	 * Can be overriden by a Transition Instance to force the logic to run regardless
-	 */
+	UE_DEPRECATED(5.8, "bEnabled deprecated. Use IAvaTransitionBehavior instead")
 	UPROPERTY()
-	bool bEnabled = true;
+	bool bEnabled_DEPRECATED = true;
 
+	UE_DEPRECATED(5.8, "InstancingMode deprecated. Use IAvaTransitionBehavior instead")
 	UPROPERTY()
-	EAvaTransitionInstancingMode InstancingMode = EAvaTransitionInstancingMode::New;
+	EAvaTransitionInstancingMode InstancingMode_DEPRECATED = EAvaTransitionInstancingMode::New;
 };
+
+#undef UE_API
diff --git a/Engine/Plugins/VirtualProduction/Avalanche/Source/AvalancheTransition/Public/AvaTransitionVersion.h b/Engine/Plugins/VirtualProduction/Avalanche/Source/AvalancheTransition/Public/AvaTransitionVersion.h
new file mode 100644
--- /dev/null
+++ b/Engine/Plugins/VirtualProduction/Avalanche/Source/AvalancheTransition/Public/AvaTransitionVersion.h
@@ -0,0 +1,39 @@
+// Copyright Epic Games, Inc. All Rights Reserved.
+
+#pragma once
+
+#include "Misc/Guid.h"
+
+#define UE_API AVALANCHETRANSITION_API
+
+namespace UE::AvaTransition
+{
+
+/** Describes the versions of transition behaviors */
+struct FBehaviorVersion
+{
+private:
+	FBehaviorVersion() = delete;
+
+public:
+	enum Type : uint8
+	{
+		PreVersioning = 0,
+
+		/**
+		 * Proper location for instance-only properties like Transition Layer. Moved away from UAvaTransitionTree into IAvaTransitionBehavior 
+		 * Transition Trees are designed to be re-usable, so these instance-related properties now live in the Transition Behavior.
+		 */
+		InstanceProperties,
+
+		/* ------------------------------------------------------ */
+		VersionPlusOne,
+		LatestVersion = VersionPlusOne - 1
+	};
+
+	UE_API const static FGuid Guid;
+};
+
+} // UE::AvaTransition
+
+#undef UE_API
diff --git a/Engine/Plugins/VirtualProduction/Avalanche/Source/AvalancheTransition/Public/Behavior/AvaTransitionBehaviorInstanceCache.h b/Engine/Plugins/VirtualProduction/Avalanche/Source/AvalancheTransition/Public/Behavior/AvaTransitionBehaviorInstanceCache.h
--- a/Engine/Plugins/VirtualProduction/Avalanche/Source/AvalancheTransition/Public/Behavior/AvaTransitionBehaviorInstanceCache.h
+++ b/Engine/Plugins/VirtualProduction/Avalanche/Source/AvalancheTransition/Public/Behavior/AvaTransitionBehaviorInstanceCache.h
@@ -4,7 +4,7 @@
 
 #include "AvaTransitionScene.h"
 
-class UAvaTransitionTree;
+class IAvaTransitionBehavior;
 struct FAvaTransitionBehaviorInstance;
 struct FAvaTransitionContext;
 struct FStateTreeExecutionContext;
@@ -13,7 +13,7 @@ struct FAvaTransitionBehaviorInstanceCache
 {
 	AVALANCHETRANSITION_API const FAvaTransitionContext* GetTransitionContext() const;
 
-	const UAvaTransitionTree* GetTransitionTree() const;
+	IAvaTransitionBehavior* GetTransitionBehavior() const;
 
 	void SetCachedInstanceDataFromContext(const FStateTreeExecutionContext& InExecutionContext);
 
diff --git a/Engine/Plugins/VirtualProduction/Avalanche/Source/AvalancheTransition/Public/Behavior/IAvaTransitionBehavior.h b/Engine/Plugins/VirtualProduction/Avalanche/Source/AvalancheTransition/Public/Behavior/IAvaTransitionBehavior.h
--- a/Engine/Plugins/VirtualProduction/Avalanche/Source/AvalancheTransition/Public/Behavior/IAvaTransitionBehavior.h
+++ b/Engine/Plugins/VirtualProduction/Avalanche/Source/AvalancheTransition/Public/Behavior/IAvaTransitionBehavior.h
@@ -17,7 +17,7 @@ enum class EStateTreeRunStatus : uint8;
 struct FAvaTransitionScene;
 struct FStateTreeReference;
 
-UINTERFACE(MinimalAPI, NotBlueprintType, meta = (CannotImplementInterfaceInBlueprint))
+UINTERFACE(MinimalAPI, BlueprintType, meta = (CannotImplementInterfaceInBlueprint))
 class UAvaTransitionBehavior : public UInterface
 {
 	GENERATED_BODY()
@@ -33,6 +33,26 @@ public:
 	/** Gets the underlying Transition Tree that this Transition Behavior runs */
 	virtual UAvaTransitionTree* GetTransitionTree() const = 0;
 
-	/** Gets the underlying State Tree Reference that this Transition Behavior runs */
+	/** Gets the transition layer that this transition behavior uses */
+	UFUNCTION(BlueprintCallable, Category = "Transition Logic")
+	virtual FAvaTagHandle GetTransitionLayer() const = 0;
+
+	/** Set the transition layer that this transition behavior will use */
+	virtual void SetTransitionLayer(FAvaTagHandle InTransitionLayer) = 0;
+
+	/** Sets whether this transition behavior is enabled */
+	virtual void SetEnabled(bool bInEnabled) = 0;
+
+	/** Determines whether the transition behavior is enabled */
+	virtual bool IsEnabled() const = 0;
+
+	/** Determines the instancing mode of the transition behavior */
+	UFUNCTION(BlueprintCallable, Category = "Transition Logic")
+	virtual EAvaTransitionInstancingMode GetInstancingMode() const = 0;
+
+	/** Sets the instancing mode for this transition behavior */
+	virtual void SetInstancingMode(EAvaTransitionInstancingMode InInstancingMode) = 0;
+
+	/** Gets the underlying State Tree reference that this transition behavior runs */
 	virtual const FStateTreeReference& GetStateTreeReference() const = 0;
 };
diff --git a/Engine/Plugins/VirtualProduction/Avalanche/Source/AvalancheTransitionEditor/Private/Actions/AvaTransitionTreeActions.cpp b/Engine/Plugins/VirtualProduction/Avalanche/Source/AvalancheTransitionEditor/Private/Actions/AvaTransitionTreeActions.cpp
--- a/Engine/Plugins/VirtualProduction/Avalanche/Source/AvalancheTransitionEditor/Private/Actions/AvaTransitionTreeActions.cpp
+++ b/Engine/Plugins/VirtualProduction/Avalanche/Source/AvalancheTransitionEditor/Private/Actions/AvaTransitionTreeActions.cpp
@@ -119,22 +119,12 @@ void FAvaTransitionTreeActions::ImportTransitionTree(UAvaTransitionTree* InTempl
 		return;
 	}
 
-	FAvaTagHandle TransitionLayer;
-	if (UAvaTransitionTreeEditorData* CurrentEditorData = Cast<UAvaTransitionTreeEditorData>(TransitionTree->EditorData))
-	{
-		TransitionLayer = CurrentEditorData->GetTransitionLayer();
-	}
-
 	FScopedTransaction Transaction(LOCTEXT("ImportTransitionTree", "Import Tree"));
 	TransitionTree->Modify();
 
-	// Duplicate the Entire Tree, but keep the same Transition Layer
 	if (UAvaTransitionTreeEditorData* NewEditorData = DuplicateObject<UAvaTransitionTreeEditorData>(TemplateEditorData, TransitionTree))
 	{
-		NewEditorData->SetTransitionLayer(TransitionLayer);
-
 		TransitionTree->EditorData = NewEditorData;
-
 		Owner.UpdateEditorData();
 		Owner.Refresh();
 	}
diff --git a/Engine/Plugins/VirtualProduction/Avalanche/Source/AvalancheTransitionEditor/Private/AvaTransitionEditorModule.cpp b/Engine/Plugins/VirtualProduction/Avalanche/Source/AvalancheTransitionEditor/Private/AvaTransitionEditorModule.cpp
--- a/Engine/Plugins/VirtualProduction/Avalanche/Source/AvalancheTransitionEditor/Private/AvaTransitionEditorModule.cpp
+++ b/Engine/Plugins/VirtualProduction/Avalanche/Source/AvalancheTransitionEditor/Private/AvaTransitionEditorModule.cpp
@@ -29,8 +29,6 @@ void FAvaTransitionEditorModule::StartupModule()
 	FAvaTransitionEditorCommands::Register();
 
 	IAvaTransitionModule::Get().GetOnValidateTransitionTree().BindRaw(this, &FAvaTransitionEditorModule::ValidateStateTree);
-
-	OnPostCompileHandle = UE::StateTree::Delegates::OnPostCompile.AddRaw(this, &FAvaTransitionEditorModule::OnPostCompile);
 }
 
 void FAvaTransitionEditorModule::ShutdownModule()
@@ -38,9 +36,6 @@ void FAvaTransitionEditorModule::ShutdownModule()
 	FAvaTransitionEditorCommands::Unregister();
 
 	IAvaTransitionModule::Get().GetOnValidateTransitionTree().Unbind();
-
-	UE::StateTree::Delegates::OnPostCompile.Remove(OnPostCompileHandle);
-	OnPostCompileHandle.Reset();
 }
 
 IAvaTransitionEditorModule::FOnBuildDefaultTransitionTree& FAvaTransitionEditorModule::GetOnBuildDefaultTransitionTree()
@@ -61,7 +56,7 @@ void FAvaTransitionEditorModule::GenerateTransitionTreeOptionsMenu(UToolMenu* In
 		return;
 	}
 
-	TWeakObjectPtr<UAvaTransitionTree> TransitionTreeWeak = TransitionTree;
+	TWeakInterfacePtr<IAvaTransitionBehavior> TransitionBehaviorWeak = InTransitionBehavior;
 
 	FToolMenuSection& GeneralSection = InMenu->FindOrAddSection(TEXT("GeneralSection"), LOCTEXT("GeneralSectionLabel", "General"));
 
@@ -69,12 +64,12 @@ void FAvaTransitionEditorModule::GenerateTransitionTreeOptionsMenu(UToolMenu* In
 		, LOCTEXT("TransitionTreeEnabledLabel", "Enabled")
 		, LOCTEXT("TransitionTreeEnabledTooltip", "Determines whether Transition Tree is Enabled")
 		, TAttribute<FSlateIcon>()
-		, FUIAction(FExecuteAction::CreateStatic(&UE::AvaTransitionEditor::ToggleTransitionTreeEnabled, TransitionTreeWeak)
+		, FUIAction(FExecuteAction::CreateStatic(&UE::AvaTransitionEditor::ToggleTransitionTreeEnabled, TransitionBehaviorWeak)
 			, FCanExecuteAction()
-			, FIsActionChecked::CreateStatic(&UE::AvaTransitionEditor::IsTransitionTreeEnabled, TransitionTreeWeak))
+			, FIsActionChecked::CreateStatic(&UE::AvaTransitionEditor::IsTransitionTreeEnabled, TransitionBehaviorWeak))
 		, EUserInterfaceActionType::ToggleButton);
 
-	if (const TSharedPtr<SWidget> ModeSelector = UE::AvaTransitionEditor::CreateTransitionInstancingModeSelector(TransitionTree))
+	if (const TSharedPtr<SWidget> ModeSelector = UE::AvaTransitionEditor::CreateTransitionInstancingModeSelector(InTransitionBehavior))
 	{
 		GeneralSection.AddEntry(FToolMenuEntry::InitWidget(TEXT("TransitionModeSelector")
 			, SNew(SBox)
@@ -87,7 +82,7 @@ void FAvaTransitionEditorModule::GenerateTransitionTreeOptionsMenu(UToolMenu* In
 			, /*bNoPadding*/true));
 	}
 
-	if (TSharedPtr<SWidget> LayerPicker = UE::AvaTransitionEditor::CreateTransitionLayerPicker(Cast<UAvaTransitionTreeEditorData>(TransitionTree->EditorData), /*bInCompileOnLayerPicked*/true))
+	if (TSharedPtr<SWidget> LayerPicker = UE::AvaTransitionEditor::CreateTransitionLayerPicker(InTransitionBehavior))
 	{
 		FToolMenuSection& LayerSection = InMenu->FindOrAddSection(TEXT("LayerSection"), LOCTEXT("LayerSectionLabel", "Layer"));
 
@@ -112,9 +107,6 @@ void FAvaTransitionEditorModule::ValidateStateTree(UAvaTransitionTree* InTransit
 		return;
 	}
 
-	// Disable Tree by default if being set up for the first time
-	InTransitionTree->SetEnabled(false);
-
 	UAvaTransitionTreeEditorData* EditorData;
 
 	const UAvaTransitionEditorSettings* TransitionEditorSettings = GetDefault<UAvaTransitionEditorSettings>();
@@ -150,18 +142,4 @@ void FAvaTransitionEditorModule::ValidateStateTree(UAvaTransitionTree* InTransit
 	Compiler.Compile(EAvaTransitionEditorMode::Advanced);
 }
 
-void FAvaTransitionEditorModule::OnPostCompile(const UStateTree& InStateTree)
-{
-	// TODO: Currently State Tree Compiler does not have any hooks to compile additional data,
-	// so this is where we inject the additional properties
-	if (const UAvaTransitionTree* TransitionTree = Cast<UAvaTransitionTree>(&InStateTree))
-	{
-		UAvaTransitionTree* TransitionTreeMutable = const_cast<UAvaTransitionTree*>(TransitionTree);
-
-		const UAvaTransitionTreeEditorData* EditorData = CastChecked<UAvaTransitionTreeEditorData>(InStateTree.EditorData);
-
-		TransitionTreeMutable->SetTransitionLayer(EditorData->GetTransitionLayer());
-	}
-}
-
 #undef LOCTEXT_NAMESPACE
diff --git a/Engine/Plugins/VirtualProduction/Avalanche/Source/AvalancheTransitionEditor/Private/AvaTransitionEditorModule.h b/Engine/Plugins/VirtualProduction/Avalanche/Source/AvalancheTransitionEditor/Private/AvaTransitionEditorModule.h
--- a/Engine/Plugins/VirtualProduction/Avalanche/Source/AvalancheTransitionEditor/Private/AvaTransitionEditorModule.h
+++ b/Engine/Plugins/VirtualProduction/Avalanche/Source/AvalancheTransitionEditor/Private/AvaTransitionEditorModule.h
@@ -21,9 +21,5 @@ class FAvaTransitionEditorModule : public IAvaTransitionEditorModule
 
 	void ValidateStateTree(UAvaTransitionTree* InTransitionTree);
 
-	void OnPostCompile(const UStateTree& InStateTree);
-
 	FOnBuildDefaultTransitionTree OnBuildDefaultTransitionTree;
-
-	FDelegateHandle OnPostCompileHandle;
 };
diff --git a/Engine/Plugins/VirtualProduction/Avalanche/Source/AvalancheTransitionEditor/Private/AvaTransitionEditorUtils.cpp b/Engine/Plugins/VirtualProduction/Avalanche/Source/AvalancheTransitionEditor/Private/AvaTransitionEditorUtils.cpp
--- a/Engine/Plugins/VirtualProduction/Avalanche/Source/AvalancheTransitionEditor/Private/AvaTransitionEditorUtils.cpp
+++ b/Engine/Plugins/VirtualProduction/Avalanche/Source/AvalancheTransitionEditor/Private/AvaTransitionEditorUtils.cpp
@@ -10,6 +10,7 @@
 #include "ContentBrowserModule.h"
 #include "IContentBrowserSingleton.h"
 #include "ISinglePropertyView.h"
+#include "Menu/Widgets/AvaTransitionTreeStatus.h"
 #include "Menu/Widgets/SAvaTransitionTreeInstancingModeSelector.h"
 #include "Modules/ModuleManager.h"
 #include "PropertyEditorModule.h"
@@ -24,10 +25,18 @@
 
 namespace UE::AvaTransitionEditor
 {
+TSharedPtr<SWidget> CreateTransitionStatusWidget(IAvaTransitionBehavior* InTransitionBehavior)
+{
+	if (!InTransitionBehavior)
+	{
+		return nullptr;
+	}
+	return SNew(SAvaTransitionTreeStatus, InTransitionBehavior);
+}
 
-TSharedPtr<SWidget> CreateTransitionLayerPicker(UAvaTransitionTreeEditorData* InEditorData, bool bInCompileOnLayerPicked)
+TSharedPtr<SWidget> CreateTransitionLayerPicker(IAvaTransitionBehavior* InTransitionBehavior)
 {
-	if (!InEditorData)
+	if (!InTransitionBehavior)
 	{
 		return nullptr;
 	}
@@ -37,27 +46,10 @@ TSharedPtr<SWidget> CreateTransitionLayerPicker(UAvaTransitionTreeEditorData* In
 	FSinglePropertyParams SinglePropertyParams;
 	SinglePropertyParams.NamePlacement = EPropertyNamePlacement::Hidden;
 
-	TSharedPtr<ISinglePropertyView> PropertyView = PropertyEditorModule.CreateSingleProperty(InEditorData
-		, UAvaTransitionTreeEditorData::GetTransitionLayerPropertyName()
+	TSharedPtr<ISinglePropertyView> PropertyView = PropertyEditorModule.CreateSingleProperty(&InTransitionBehavior->AsUObject()
+		, TEXT("TransitionLayer")
 		, SinglePropertyParams);
 
-	if (bInCompileOnLayerPicked)
-	{
-		if (UAvaTransitionTree* TransitionTree = InEditorData->GetTypedOuter<UAvaTransitionTree>())
-		{
-			FSimpleDelegate OnLayerPicked = FSimpleDelegate::CreateWeakLambda(TransitionTree,
-				[TransitionTree]
-				{
-					// Compile in Advanced Mode here so that no new nodes are generated from outside
-					FAvaTransitionCompiler Compiler;
-					Compiler.SetTransitionTree(TransitionTree);
-					Compiler.Compile(EAvaTransitionEditorMode::Advanced);
-				});
-
-			PropertyView->SetOnPropertyValueChanged(OnLayerPicked);
-		}		
-	}
-
 	if (PropertyView.IsValid())
 	{
 		return SNew(SBox)
@@ -80,7 +72,7 @@ TSharedPtr<SWidget> CreateTransitionLayerPicker(UAvaTransitionTreeEditorData* In
 	return nullptr;
 }
 
-TSharedPtr<SWidget> CreateTransitionInstancingModeSelector(UAvaTransitionTree* InTransitionTree)
+TSharedPtr<SWidget> CreateTransitionInstancingModeSelector(IAvaTransitionBehavior* InTransitionBehavior)
 {
 	const FMargin PropertyPadding(2.0f, 0.0f, 2.0f, 0.0f);
 	
@@ -99,7 +91,7 @@ TSharedPtr<SWidget> CreateTransitionInstancingModeSelector(UAvaTransitionTree* I
 		.HAlign(HAlign_Fill)
 		.VAlign(VAlign_Center)
 		[
-			SNew(SAvaTransitionTreeInstancingModeSelector, InTransitionTree)
+			SNew(SAvaTransitionTreeInstancingModeSelector, InTransitionBehavior)
 		];
 }
 
@@ -283,19 +275,20 @@ bool PickTransitionTreeAsset(const FText& InDialogTitle, UAvaTransitionTree*& Ou
 	return true;
 }
 
-void ToggleTransitionTreeEnabled(TWeakObjectPtr<UAvaTransitionTree> InTransitionTreeWeak)
+void ToggleTransitionTreeEnabled(TWeakInterfacePtr<IAvaTransitionBehavior> InTransitionBehaviorWeak)
 {
-	if (UAvaTransitionTree* TransitionTree = InTransitionTreeWeak.Get())
+	if (IAvaTransitionBehavior* const TransitionBehavior = InTransitionBehaviorWeak.Get())
 	{
 		FScopedTransaction Transaction(LOCTEXT("ToggleTransitionTreeEnabled", "Toggle Transition Tree Enabled"));
-		TransitionTree->Modify();
-		TransitionTree->SetEnabled(!TransitionTree->IsEnabled());
+		TransitionBehavior->AsUObject().Modify();
+		TransitionBehavior->SetEnabled(!TransitionBehavior->IsEnabled());
 	}
 }
 
-bool IsTransitionTreeEnabled(TWeakObjectPtr<UAvaTransitionTree> InTransitionTreeWeak)
+bool IsTransitionTreeEnabled(TWeakInterfacePtr<IAvaTransitionBehavior> InTransitionBehaviorWeak)
 {
-	return InTransitionTreeWeak.IsValid() && InTransitionTreeWeak->IsEnabled();
+	IAvaTransitionBehavior* const TransitionBehavior = InTransitionBehaviorWeak.Get();
+	return TransitionBehavior && TransitionBehavior->IsEnabled();
 }
 
 }
diff --git a/Engine/Plugins/VirtualProduction/Avalanche/Source/AvalancheTransitionEditor/Private/AvaTransitionEditorUtils.h b/Engine/Plugins/VirtualProduction/Avalanche/Source/AvalancheTransitionEditor/Private/AvaTransitionEditorUtils.h
--- a/Engine/Plugins/VirtualProduction/Avalanche/Source/AvalancheTransitionEditor/Private/AvaTransitionEditorUtils.h
+++ b/Engine/Plugins/VirtualProduction/Avalanche/Source/AvalancheTransitionEditor/Private/AvaTransitionEditorUtils.h
@@ -12,9 +12,11 @@ class UAvaTransitionTreeEditorData;
 
 namespace UE::AvaTransitionEditor
 {
-	TSharedPtr<SWidget> CreateTransitionLayerPicker(UAvaTransitionTreeEditorData* InEditorData, bool bInCompileOnLayerPicked);
+	TSharedPtr<SWidget> CreateTransitionStatusWidget(IAvaTransitionBehavior* InTransitionBehavior);
 
-	TSharedPtr<SWidget> CreateTransitionInstancingModeSelector(UAvaTransitionTree* InTransitionTree);
+	TSharedPtr<SWidget> CreateTransitionLayerPicker(IAvaTransitionBehavior* InTransitionBehavior);
+
+	TSharedPtr<SWidget> CreateTransitionInstancingModeSelector(IAvaTransitionBehavior* InTransitionBehavior);
 	
 	void ValidateTree(UAvaTransitionTree& InTransitionTree);
 
@@ -22,7 +24,7 @@ namespace UE::AvaTransitionEditor
 
 	bool PickTransitionTreeAsset(const FText& InDialogTitle, UAvaTransitionTree*& OutTransitionTree);
 
-	void ToggleTransitionTreeEnabled(TWeakObjectPtr<UAvaTransitionTree> InTransitionTreeWeak);
+	void ToggleTransitionTreeEnabled(TWeakInterfacePtr<IAvaTransitionBehavior> InTransitionBehaviorWeak);
 
-	bool IsTransitionTreeEnabled(TWeakObjectPtr<UAvaTransitionTree> InTransitionTreeWeak);
+	bool IsTransitionTreeEnabled(TWeakInterfacePtr<IAvaTransitionBehavior> InTransitionBehaviorWeak);
 }
diff --git a/Engine/Plugins/VirtualProduction/Avalanche/Source/AvalancheTransitionEditor/Private/AvaTransitionTreeEditorData.cpp b/Engine/Plugins/VirtualProduction/Avalanche/Source/AvalancheTransitionEditor/Private/AvaTransitionTreeEditorData.cpp
--- a/Engine/Plugins/VirtualProduction/Avalanche/Source/AvalancheTransitionEditor/Private/AvaTransitionTreeEditorData.cpp
+++ b/Engine/Plugins/VirtualProduction/Avalanche/Source/AvalancheTransitionEditor/Private/AvaTransitionTreeEditorData.cpp
@@ -5,8 +5,9 @@
 
 UAvaTransitionTreeEditorData::UAvaTransitionTreeEditorData()
 {
-	const UE::AvaTransition::FDefaultTags& DefaultTags = UE::AvaTransition::FDefaultTags::Get();
-	TransitionLayer = DefaultTags.DefaultLayer.MakeTagHandle();
+PRAGMA_DISABLE_DEPRECATION_WARNINGS
+	TransitionLayer_DEPRECATED = UE::AvaTransition::FDefaultTags::Get().DefaultLayer.MakeTagHandle();
+PRAGMA_ENABLE_DEPRECATION_WARNINGS
 
 	FStateTreeEditorColor DefaultColor;
 	DefaultColor.ColorRef = FStateTreeEditorColorRef(UE::AvaTransitionEditor::ColorId_Default);
@@ -65,3 +66,22 @@ bool UAvaTransitionTreeEditorData::Compare(const UAvaTransitionTreeEditorData* I
 	// When duplicating trees, even though the parameter itself isn't duplicate transient, it will assign a new guid in UStateTreeEditorData::DuplicateIDs
 	return InEditorData && (this == InEditorData || GetRootParametersGuid() == InEditorData->GetRootParametersGuid());
 }
+
+FAvaTagHandle UAvaTransitionTreeEditorData::GetTransitionLayer() const
+{
+PRAGMA_DISABLE_DEPRECATION_WARNINGS
+	return TransitionLayer_DEPRECATED;
+PRAGMA_ENABLE_DEPRECATION_WARNINGS
+}
+
+void UAvaTransitionTreeEditorData::SetTransitionLayer(const FAvaTagHandle& InTransitionLayer)
+{
+PRAGMA_DISABLE_DEPRECATION_WARNINGS
+	TransitionLayer_DEPRECATED = InTransitionLayer;
+PRAGMA_ENABLE_DEPRECATION_WARNINGS
+}
+
+FName UAvaTransitionTreeEditorData::GetTransitionLayerPropertyName()
+{
+	return NAME_None;
+}
diff --git a/Engine/Plugins/VirtualProduction/Avalanche/Source/AvalancheTransitionEditor/Private/AvaTransitionTreeFactory.cpp b/Engine/Plugins/VirtualProduction/Avalanche/Source/AvalancheTransitionEditor/Private/AvaTransitionTreeFactory.cpp
--- a/Engine/Plugins/VirtualProduction/Avalanche/Source/AvalancheTransitionEditor/Private/AvaTransitionTreeFactory.cpp
+++ b/Engine/Plugins/VirtualProduction/Avalanche/Source/AvalancheTransitionEditor/Private/AvaTransitionTreeFactory.cpp
@@ -28,9 +28,6 @@ UObject* UAvaTransitionTreeFactory::FactoryCreateNew(UClass* InClass, UObject* I
 		if (TransitionTree)
 		{
 			IAvaTransitionModule::Get().GetOnValidateTransitionTree().ExecuteIfBound(TransitionTree);
-
-			// Transition Tree Assets are always Enabled by Default
-			TransitionTree->SetEnabled(true);
 		}
 		return TransitionTree;
 	}
diff --git a/Engine/Plugins/VirtualProduction/Avalanche/Source/AvalancheTransitionEditor/Private/Compiler/AvaTransitionCompiler.cpp b/Engine/Plugins/VirtualProduction/Avalanche/Source/AvalancheTransitionEditor/Private/Compiler/AvaTransitionCompiler.cpp
--- a/Engine/Plugins/VirtualProduction/Avalanche/Source/AvalancheTransitionEditor/Private/Compiler/AvaTransitionCompiler.cpp
+++ b/Engine/Plugins/VirtualProduction/Avalanche/Source/AvalancheTransitionEditor/Private/Compiler/AvaTransitionCompiler.cpp
@@ -72,9 +72,6 @@ bool FAvaTransitionCompiler::Compile(EAvaTransitionEditorMode InCompileMode)
 
 	if (UAvaTransitionTreeEditorData* EditorData = Cast<UAvaTransitionTreeEditorData>(TransitionTree->EditorData))
 	{
-		// Update Tree's Transition Layer
-		TransitionTree->SetTransitionLayer(EditorData->GetTransitionLayer());
-
 		// Remove all invalid Sub Trees, if any
 		EditorData->SubTrees.RemoveAll(
 			[](const UStateTreeState* InState)
diff --git a/Engine/Plugins/VirtualProduction/Avalanche/Source/AvalancheTransitionEditor/Private/Customizations/AvaTransitionTreeEditorDataCustomization.cpp b/Engine/Plugins/VirtualProduction/Avalanche/Source/AvalancheTransitionEditor/Private/Customizations/AvaTransitionTreeEditorDataCustomization.cpp
--- a/Engine/Plugins/VirtualProduction/Avalanche/Source/AvalancheTransitionEditor/Private/Customizations/AvaTransitionTreeEditorDataCustomization.cpp
+++ b/Engine/Plugins/VirtualProduction/Avalanche/Source/AvalancheTransitionEditor/Private/Customizations/AvaTransitionTreeEditorDataCustomization.cpp
@@ -33,12 +33,6 @@ void FAvaTransitionTreeEditorDataCustomization::CustomizeDetails(IDetailLayoutBu
 		}
 	}
 
-	// Hide property as it's going to show in the Toolbar
-	TSharedRef<IPropertyHandle> LayerHandle = InDetailBuilder.GetProperty(UAvaTransitionTreeEditorData::GetTransitionLayerPropertyName()
-		, UAvaTransitionTreeEditorData::StaticClass());
-
-	InDetailBuilder.HideProperty(LayerHandle);
-
 	if (EditorMode != EAvaTransitionEditorMode::Advanced)
 	{
 		InDetailBuilder.HideCategory(TEXT("Evaluators"));
diff --git a/Engine/Plugins/VirtualProduction/Avalanche/Source/AvalancheTransitionEditor/Private/Menu/AvaTransitionToolbar.cpp b/Engine/Plugins/VirtualProduction/Avalanche/Source/AvalancheTransitionEditor/Private/Menu/AvaTransitionToolbar.cpp
--- a/Engine/Plugins/VirtualProduction/Avalanche/Source/AvalancheTransitionEditor/Private/Menu/AvaTransitionToolbar.cpp
+++ b/Engine/Plugins/VirtualProduction/Avalanche/Source/AvalancheTransitionEditor/Private/Menu/AvaTransitionToolbar.cpp
@@ -91,9 +91,13 @@ void FAvaTransitionToolbar::ExtendEditorToolbar(UToolMenu* InToolbarMenu)
 			return InSection.AddSeparator(FName(TEXT("Separator"), SeparatorIndex++));
 		};
 
+	UAvaTransitionTree* const TransitionTree = Owner.GetTransitionTree();
+
+	IAvaTransitionBehavior* const TransitionBehavior = TransitionTree ? TransitionTree->GetImplementingOuter<IAvaTransitionBehavior>() : nullptr;
+
 	// Tree Status
+	if (TSharedPtr<SWidget> TreeStatusWidget = CreateTransitionStatusWidget(TransitionBehavior))
 	{
-		TSharedRef<SWidget> TreeStatusWidget = SNew(SAvaTransitionTreeStatus, Owner.GetTransitionTree());
 		TreeStatusWidget->SetEnabled(!bReadOnly);
 
 		AllowInReadOnly(MakeSeparator(Section));
@@ -101,12 +105,12 @@ void FAvaTransitionToolbar::ExtendEditorToolbar(UToolMenu* InToolbarMenu)
 		AllowInReadOnly(Private::AddLabelEntry(Section, TEXT("TreeStatusLabel"), LOCTEXT("StatusLabel", "Status")));
 
 		AllowInReadOnly(Section.AddEntry(FToolMenuEntry::InitWidget(TEXT("TreeStatusWidget")
-			, TreeStatusWidget
+			, TreeStatusWidget.ToSharedRef()
 			, FText::GetEmpty())));
 	}
 
 	// Layer Picker
-	if (TSharedPtr<SWidget> LayerPicker = CreateTransitionLayerPicker(Owner.GetEditorData(), /*bInCompileOnLayerPicked*/false))
+	if (TSharedPtr<SWidget> LayerPicker = CreateTransitionLayerPicker(TransitionBehavior))
 	{
 		LayerPicker->SetEnabled(!bReadOnly);
 
diff --git a/Engine/Plugins/VirtualProduction/Avalanche/Source/AvalancheTransitionEditor/Private/Menu/Widgets/AvaTransitionTreeStatus.cpp b/Engine/Plugins/VirtualProduction/Avalanche/Source/AvalancheTransitionEditor/Private/Menu/Widgets/AvaTransitionTreeStatus.cpp
--- a/Engine/Plugins/VirtualProduction/Avalanche/Source/AvalancheTransitionEditor/Private/Menu/Widgets/AvaTransitionTreeStatus.cpp
+++ b/Engine/Plugins/VirtualProduction/Avalanche/Source/AvalancheTransitionEditor/Private/Menu/Widgets/AvaTransitionTreeStatus.cpp
@@ -1,7 +1,7 @@
 // Copyright Epic Games, Inc. All Rights Reserved.
 
 #include "AvaTransitionTreeStatus.h"
-#include "AvaTransitionTree.h"
+#include "Behavior/IAvaTransitionBehavior.h"
 #include "ScopedTransaction.h"
 #include "Styling/AvaTransitionEditorStyle.h"
 #include "Widgets/Images/SImage.h"
@@ -24,9 +24,9 @@ namespace UE::AvaTransitionEditor::Private
 	};
 }
 
-void SAvaTransitionTreeStatus::Construct(const FArguments& InArgs, UAvaTransitionTree* InTransitionTree)
+void SAvaTransitionTreeStatus::Construct(const FArguments& InArgs, IAvaTransitionBehavior* InTransitionBehavior)
 {
-	TransitionTreeWeak = InTransitionTree;
+	TransitionBehaviorWeak = InTransitionBehavior;
 
 	ChildSlot
 		.HAlign(HAlign_Center)
@@ -53,31 +53,25 @@ bool SAvaTransitionTreeStatus::IsStatusButtonEnabled() const
 
 FReply SAvaTransitionTreeStatus::OnStatusButtonClicked()
 {
-	UAvaTransitionTree* TransitionTree = TransitionTreeWeak.Get();
-	if (TransitionTree && !TransitionTree->IsEnabled())
+	IAvaTransitionBehavior* const TransitionBehavior = TransitionBehaviorWeak.Get();
+	if (TransitionBehavior && !TransitionBehavior->IsEnabled())
 	{
 		FScopedTransaction Transaction(LOCTEXT("EnableTransitionTree", "Enable Transition Tree"));
-
-		FProperty* const EnabledProperty = FindFProperty<FBoolProperty>(UAvaTransitionTree::StaticClass(), UAvaTransitionTree::GetEnabledPropertyName());
-		TransitionTree->PreEditChange(EnabledProperty);
-
-		TransitionTree->SetEnabled(true);
-
-		FPropertyChangedEvent PropertyChangedEvent(EnabledProperty, EPropertyChangeType::ValueSet);
-		TransitionTree->PostEditChangeProperty(PropertyChangedEvent);
+		TransitionBehavior->AsUObject().Modify();
+		TransitionBehavior->SetEnabled(true);
 	}
 	return FReply::Handled();
 }
 
 SAvaTransitionTreeStatus::EStatus SAvaTransitionTreeStatus::GetStatus() const
 {
-	UAvaTransitionTree* TransitionTree = TransitionTreeWeak.Get();
-	if (!TransitionTree)
+	IAvaTransitionBehavior* const TransitionBehavior = TransitionBehaviorWeak.Get();
+	if (!TransitionBehavior)
 	{
 		return EStatus::Invalid;
 	}
 
-	return TransitionTree->IsEnabled()
+	return TransitionBehavior->IsEnabled()
 		? EStatus::Enabled
 		: EStatus::Disabled;
 }
diff --git a/Engine/Plugins/VirtualProduction/Avalanche/Source/AvalancheTransitionEditor/Private/Menu/Widgets/AvaTransitionTreeStatus.h b/Engine/Plugins/VirtualProduction/Avalanche/Source/AvalancheTransitionEditor/Private/Menu/Widgets/AvaTransitionTreeStatus.h
--- a/Engine/Plugins/VirtualProduction/Avalanche/Source/AvalancheTransitionEditor/Private/Menu/Widgets/AvaTransitionTreeStatus.h
+++ b/Engine/Plugins/VirtualProduction/Avalanche/Source/AvalancheTransitionEditor/Private/Menu/Widgets/AvaTransitionTreeStatus.h
@@ -2,10 +2,10 @@
 
 #pragma once
 
-#include "UObject/WeakObjectPtr.h"
+#include "UObject/WeakInterfacePtr.h"
 #include "Widgets/SCompoundWidget.h"
 
-class UAvaTransitionTree;
+class IAvaTransitionBehavior;
 
 class SAvaTransitionTreeStatus : public SCompoundWidget
 {
@@ -20,7 +20,7 @@ public:
 	SLATE_BEGIN_ARGS(SAvaTransitionTreeStatus) {}
 	SLATE_END_ARGS()
 
-	void Construct(const FArguments& InArgs, UAvaTransitionTree* InTransitionTree);
+	void Construct(const FArguments& InArgs, IAvaTransitionBehavior* InTransitionBehavior);
 
 private:
 	bool IsStatusButtonEnabled() const;
@@ -35,5 +35,5 @@ private:
 
 	const FSlateBrush* GetStatusIcon() const;
 
-	TWeakObjectPtr<UAvaTransitionTree> TransitionTreeWeak;
+	TWeakInterfacePtr<IAvaTransitionBehavior> TransitionBehaviorWeak;
 };
diff --git a/Engine/Plugins/VirtualProduction/Avalanche/Source/AvalancheTransitionEditor/Private/Menu/Widgets/SAvaTransitionTreeInstancingModeSelector.cpp b/Engine/Plugins/VirtualProduction/Avalanche/Source/AvalancheTransitionEditor/Private/Menu/Widgets/SAvaTransitionTreeInstancingModeSelector.cpp
--- a/Engine/Plugins/VirtualProduction/Avalanche/Source/AvalancheTransitionEditor/Private/Menu/Widgets/SAvaTransitionTreeInstancingModeSelector.cpp
+++ b/Engine/Plugins/VirtualProduction/Avalanche/Source/AvalancheTransitionEditor/Private/Menu/Widgets/SAvaTransitionTreeInstancingModeSelector.cpp
@@ -1,17 +1,17 @@
 // Copyright Epic Games, Inc. All Rights Reserved.
 
 #include "SAvaTransitionTreeInstancingModeSelector.h"
-
-#include "AvaTransitionTree.h"
+#include "AvaTransitionEnums.h"
+#include "Behavior/IAvaTransitionBehavior.h"
 #include "ScopedTransaction.h"
 #include "Widgets/Input/SComboBox.h"
 
 #define LOCTEXT_NAMESPACE "AvaTransitionTreeInstancingModeSelector"
 
-void SAvaTransitionTreeInstancingModeSelector::Construct(const FArguments& InArgs, UAvaTransitionTree* InTransitionTree)
+void SAvaTransitionTreeInstancingModeSelector::Construct(const FArguments& InArgs, IAvaTransitionBehavior* InTransitionBehavior)
 {
 	UpdateItems();
-	TransitionTreeWeak = InTransitionTree;
+	TransitionBehaviorWeak = InTransitionBehavior;
 
 	ChildSlot
 	[
@@ -35,29 +35,29 @@ TSharedRef<SWidget> SAvaTransitionTreeInstancingModeSelector::GenerateWidget(FNa
 
 void SAvaTransitionTreeInstancingModeSelector::HandleSelectionChanged(FName InProposedSelection, ESelectInfo::Type InSelectInfo)
 {
-	if (UAvaTransitionTree* TransitionTree = TransitionTreeWeak.Get())
+	if (IAvaTransitionBehavior* TransitionBehavior = TransitionBehaviorWeak.Get())
 	{
 		FScopedTransaction Transaction(LOCTEXT("SetTransitionTreeInstancingMode", "Set Transition Tree Instancing Mode")); 
-		TransitionTree->Modify();
+		TransitionBehavior->AsUObject().Modify();
 		int64 Value = StaticEnum<EAvaTransitionInstancingMode>()->GetValueByName(InProposedSelection); 
-		TransitionTree->SetInstancingMode(static_cast<EAvaTransitionInstancingMode>(Value));
+		TransitionBehavior->SetInstancingMode(static_cast<EAvaTransitionInstancingMode>(Value));
 	}
 }
 
 FText SAvaTransitionTreeInstancingModeSelector::GetDisplayTextFromProperty() const
 {
-	if (const UAvaTransitionTree* TransitionTree = TransitionTreeWeak.Get())
+	if (const IAvaTransitionBehavior* TransitionBehavior = TransitionBehaviorWeak.Get())
 	{
-		return StaticEnum<EAvaTransitionInstancingMode>()->GetDisplayNameTextByValue(static_cast<int64>(TransitionTree->GetInstancingMode()));
+		return StaticEnum<EAvaTransitionInstancingMode>()->GetDisplayNameTextByValue(static_cast<int64>(TransitionBehavior->GetInstancingMode()));
 	}
 	return FText::GetEmpty();
 }
 
 FName SAvaTransitionTreeInstancingModeSelector::GetItemFromProperty() const
 {
-	if (const UAvaTransitionTree* TransitionTree = TransitionTreeWeak.Get())
+	if (const IAvaTransitionBehavior* TransitionBehavior = TransitionBehaviorWeak.Get())
 	{
-		return StaticEnum<EAvaTransitionInstancingMode>()->GetNameByValue(static_cast<int64>(TransitionTree->GetInstancingMode()));
+		return StaticEnum<EAvaTransitionInstancingMode>()->GetNameByValue(static_cast<int64>(TransitionBehavior->GetInstancingMode()));
 	}
 	return NAME_None;
 }
diff --git a/Engine/Plugins/VirtualProduction/Avalanche/Source/AvalancheTransitionEditor/Private/Menu/Widgets/SAvaTransitionTreeInstancingModeSelector.h b/Engine/Plugins/VirtualProduction/Avalanche/Source/AvalancheTransitionEditor/Private/Menu/Widgets/SAvaTransitionTreeInstancingModeSelector.h
--- a/Engine/Plugins/VirtualProduction/Avalanche/Source/AvalancheTransitionEditor/Private/Menu/Widgets/SAvaTransitionTreeInstancingModeSelector.h
+++ b/Engine/Plugins/VirtualProduction/Avalanche/Source/AvalancheTransitionEditor/Private/Menu/Widgets/SAvaTransitionTreeInstancingModeSelector.h
@@ -2,10 +2,10 @@
 
 #pragma once
 
-#include "UObject/WeakObjectPtr.h"
+#include "UObject/WeakInterfacePtr.h"
 #include "Widgets/SCompoundWidget.h"
 
-class UAvaTransitionTree;
+class IAvaTransitionBehavior;
 template <typename OptionType> class SComboBox;
 
 class SAvaTransitionTreeInstancingModeSelector : public SCompoundWidget
@@ -14,7 +14,7 @@ public:
 	SLATE_BEGIN_ARGS(SAvaTransitionTreeInstancingModeSelector){}
 	SLATE_END_ARGS()
 
-	void Construct(const FArguments& InArgs, UAvaTransitionTree* InTransitionTree);
+	void Construct(const FArguments& InArgs, IAvaTransitionBehavior* InTransitionBehavior);
 
 protected:
 	TSharedRef<SWidget> GenerateWidget(FName InItem) const;
@@ -25,7 +25,7 @@ protected:
 	void UpdateItems();
 
 protected:
-	TWeakObjectPtr<UAvaTransitionTree> TransitionTreeWeak;
+	TWeakInterfacePtr<IAvaTransitionBehavior> TransitionBehaviorWeak;
 	TSharedPtr<SComboBox<FName>> Combo;
 	TArray<FName> Items;
 };
diff --git a/Engine/Plugins/VirtualProduction/Avalanche/Source/AvalancheTransitionEditor/Public/AvaTransitionTreeEditorData.h b/Engine/Plugins/VirtualProduction/Avalanche/Source/AvalancheTransitionEditor/Public/AvaTransitionTreeEditorData.h
--- a/Engine/Plugins/VirtualProduction/Avalanche/Source/AvalancheTransitionEditor/Public/AvaTransitionTreeEditorData.h
+++ b/Engine/Plugins/VirtualProduction/Avalanche/Source/AvalancheTransitionEditor/Public/AvaTransitionTreeEditorData.h
@@ -34,20 +34,14 @@ public:
 	/** Returns true if this has the same data as the given editor data */
 	UE_API bool Compare(const UAvaTransitionTreeEditorData* InEditorData) const;
 
-	FAvaTagHandle GetTransitionLayer() const
-	{
-		return TransitionLayer;
-	}
+	UE_DEPRECATED(5.8, "GetTransitionLayer is deprecated. Use IAvaTransitionBehavior::GetTransitionLayer instead")
+	UE_API FAvaTagHandle GetTransitionLayer() const;
 
-	void SetTransitionLayer(const FAvaTagHandle& InTransitionLayer)
-	{
-		TransitionLayer = InTransitionLayer;
-	}
+	UE_DEPRECATED(5.8, "SetTransitionLayer is deprecated. Use IAvaTransitionBehavior::SetTransitionLayer instead")
+	UE_API void SetTransitionLayer(const FAvaTagHandle& InTransitionLayer);
 
-	static FName GetTransitionLayerPropertyName()
-	{
-		return GET_MEMBER_NAME_CHECKED(UAvaTransitionTreeEditorData, TransitionLayer);
-	}
+	UE_DEPRECATED(5.8, "Editor data TransitionLayer is deprecated. Use IAvaTransitionBehavior instead")
+	UE_API static FName GetTransitionLayerPropertyName();
 
 	const FAvaTransitionStateMetadata* FindStateMetadata(const FGuid& InStateId) const
 	{
@@ -65,9 +59,9 @@ public:
 	}
 
 private:
-	/** The Layer this Transition Logic Tree deals with */
-	UPROPERTY(EditAnywhere, Category="Transition Logic")
-	FAvaTagHandle TransitionLayer;
+	UE_DEPRECATED(5.8, "TransitionLayer is deprecated. Use IAvaTransitionBehavior Set/Get TransitionLayer instead")
+	UPROPERTY()
+	FAvaTagHandle TransitionLayer_DEPRECATED;
 
 	/** Map of a state's id to its metadata */
 	UPROPERTY()
