Motion Design: added experimental ability to delete Remote Control presets from a motion design level. Remote control presets can be deleted by going into the additional options menu next to the 'Remote Control' button and selecting the delete entry.

#Synced-CL 52400047

diff --git a/Engine/Plugins/VirtualProduction/Avalanche/Source/Avalanche/Private/AvaScene.cpp b/Engine/Plugins/VirtualProduction/Avalanche/Source/Avalanche/Private/AvaScene.cpp
--- a/Engine/Plugins/VirtualProduction/Avalanche/Source/Avalanche/Private/AvaScene.cpp
+++ b/Engine/Plugins/VirtualProduction/Avalanche/Source/Avalanche/Private/AvaScene.cpp
@@ -91,8 +91,6 @@ PRAGMA_DISABLE_DEPRECATION_WARNINGS
 PRAGMA_ENABLE_DEPRECATION_WARNINGS
 #endif
 
-	RemoteControlPreset = CreateDefaultSubobject<URemoteControlPreset>(TEXT("RemoteControlPreset"));
-
 #if WITH_EDITOR
 	if (!HasAnyFlags(EObjectFlags::RF_ClassDefaultObject))
 	{
@@ -172,6 +170,45 @@ IAvaSequencePlaybackObject* AAvaScene::GetPlaybackObject() const
 	return GetScenePlayback();
 }
 
+URemoteControlPreset* AAvaScene::GetOrCreateRemoteControlPreset(bool bInCallModify)
+{
+	if (!RemoteControlPreset)
+	{
+		if (bInCallModify)
+		{
+			Modify();
+		}
+		RemoteControlPreset = NewObject<URemoteControlPreset>(this);
+		bRemoteControlPresetPresent = true;
+		FAvaRemoteControlUtils::RegisterRemoteControlPreset(RemoteControlPreset, /*bInEnsureUniqueId*/ true);
+	}
+	return RemoteControlPreset;
+}
+
+URemoteControlPreset* AAvaScene::GetRemoteControlPreset() const
+{
+	return RemoteControlPreset;
+}
+
+void AAvaScene::DestroyRemoteControlPreset(bool bInCallModify)
+{
+	if (RemoteControlPreset)
+	{
+		if (bInCallModify)
+		{
+			Modify();
+		}
+		FAvaRemoteControlUtils::UnregisterRemoteControlPreset(RemoteControlPreset);
+		RemoteControlPreset = nullptr;
+		bRemoteControlPresetPresent = false;
+	}
+}
+
+bool AAvaScene::IsRemoteControlPresetExpectedToBePresent() const
+{
+	return bRemoteControlPresetPresent;
+}
+
 UObject* AAvaScene::ToUObject()
 {
 	return this;
@@ -312,6 +349,10 @@ void AAvaScene::OnValuesApplied_Implementation()
 void AAvaScene::PostActorCreated()
 {
 	Super::PostActorCreated();
+
+	// by default create an RC preset.
+	GetOrCreateRemoteControlPreset();
+
 	RegisterObjects();
 }
 
@@ -349,6 +390,12 @@ void AAvaScene::EndPlay(const EEndPlayReason::Type InEndPlayReason)
 void AAvaScene::PostLoad()
 {
 	Super::PostLoad();
+
+	if (RemoteControlPreset)
+	{
+		RemoteControlPreset->ClearFlags(RF_DefaultSubObject);
+	}
+
 	RegisterObjects();
 
 	if (AttributeContainer)
@@ -378,6 +425,16 @@ void AAvaScene::PostEditImport()
 	RegisterObjects();
 }
 
+#if WITH_EDITOR
+void AAvaScene::PostEditUndo()
+{
+	Super::PostEditUndo();
+
+	// Ensure RC is registered
+	FAvaRemoteControlUtils::RegisterRemoteControlPreset(RemoteControlPreset, /*bInEnsureUniqueId*/ true);
+}
+#endif
+
 void AAvaScene::BeginDestroy()
 {
 	Super::BeginDestroy();
diff --git a/Engine/Plugins/VirtualProduction/Avalanche/Source/Avalanche/Public/AvaScene.h b/Engine/Plugins/VirtualProduction/Avalanche/Source/Avalanche/Public/AvaScene.h
--- a/Engine/Plugins/VirtualProduction/Avalanche/Source/Avalanche/Public/AvaScene.h
+++ b/Engine/Plugins/VirtualProduction/Avalanche/Source/Avalanche/Public/AvaScene.h
@@ -69,7 +69,10 @@ public:
 	virtual IAvaSequencePlaybackObject* GetPlaybackObject() const override;
 	virtual IAvaSequenceProvider* GetSequenceProvider() override { return this; }
 	virtual const IAvaSequenceProvider* GetSequenceProvider() const override { return this; }
-	virtual URemoteControlPreset* GetRemoteControlPreset() const override { return RemoteControlPreset; }
+	virtual URemoteControlPreset* GetOrCreateRemoteControlPreset(bool bInCallModify=true) override;
+	virtual URemoteControlPreset* GetRemoteControlPreset() const override;
+	virtual void DestroyRemoteControlPreset(bool bInCallModify=true) override;
+	virtual bool IsRemoteControlPresetExpectedToBePresent() const  override;
 #if WITH_EDITOR
 	virtual FNavigationToolSaveState& GetNavigationToolSaveState() override { return NavigationToolState; }
 #endif
@@ -118,6 +121,9 @@ public:
 	virtual void PostInitializeComponents() override;
 	virtual void PostDuplicate(EDuplicateMode::Type InDuplicateMode) override;
 	virtual void PostEditImport() override;
+#if WITH_EDITOR
+	virtual void PostEditUndo() override;
+#endif
 	virtual void BeginDestroy() override;
 	//~ End UObject
 
@@ -139,7 +145,7 @@ protected:
 	TObjectPtr<UAvaAttributeContainer> SceneState;
 #endif // WITH_EDITORONLY_DATA
 
-	UPROPERTY()
+	UPROPERTY(Instanced)
 	TObjectPtr<URemoteControlPreset> RemoteControlPreset;
 
 	/** A List of All Animations, including those that are nested in other Animations */
@@ -176,6 +182,10 @@ protected:
 
 	bool bPendingAnimTreeUpdate = false;
 
+	/** Determines whether the remote control preset should be present in this scene */
+	UPROPERTY()
+	bool bRemoteControlPresetPresent = true;
+
 #if WITH_EDITORONLY_DATA
 	/** Whether the motion design mode should auto start */
 	UPROPERTY()
diff --git a/Engine/Plugins/VirtualProduction/Avalanche/Source/Avalanche/Public/IAvaSceneInterface.h b/Engine/Plugins/VirtualProduction/Avalanche/Source/Avalanche/Public/IAvaSceneInterface.h
--- a/Engine/Plugins/VirtualProduction/Avalanche/Source/Avalanche/Public/IAvaSceneInterface.h
+++ b/Engine/Plugins/VirtualProduction/Avalanche/Source/Avalanche/Public/IAvaSceneInterface.h
@@ -54,8 +54,19 @@ public:
 
 	virtual IAvaSequencePlaybackObject* GetPlaybackObject() const = 0;
 
+	/** Retrieve the remote control preset, or create one if there's no preset present */
+	virtual URemoteControlPreset* GetOrCreateRemoteControlPreset(bool bInCallModify=true) = 0;
+
+	/** Gets the remote control preset of the scene */
 	virtual URemoteControlPreset* GetRemoteControlPreset() const = 0;
 
+	/** Destroys the remote control preset of the scene */
+	UE_EXPERIMENTAL(5.8, "Destroying the Remote Control Preset in Motion Design Scenes is experimental.")
+	virtual void DestroyRemoteControlPreset(bool bInCallModify=true) = 0;
+
+	/** Whether the remote control preset is expected to be present in this scene */
+	virtual bool IsRemoteControlPresetExpectedToBePresent() const = 0;
+
 #if WITH_EDITOR
 	virtual FNavigationToolSaveState& GetNavigationToolSaveState() = 0;
 #endif
diff --git a/Engine/Plugins/VirtualProduction/Avalanche/Source/AvalancheEditor/Private/RemoteControl/AvaRCExtension.cpp b/Engine/Plugins/VirtualProduction/Avalanche/Source/AvalancheEditor/Private/RemoteControl/AvaRCExtension.cpp
--- a/Engine/Plugins/VirtualProduction/Avalanche/Source/AvalancheEditor/Private/RemoteControl/AvaRCExtension.cpp
+++ b/Engine/Plugins/VirtualProduction/Avalanche/Source/AvalancheEditor/Private/RemoteControl/AvaRCExtension.cpp
@@ -7,6 +7,7 @@
 #include "IAvaSceneInterface.h"
 #include "IRemoteControlUIModule.h"
 #include "LevelEditor.h"
+#include "Misc/MessageDialog.h"
 #include "RemoteControlPreset.h"
 #include "RemoteControlTrackerComponent.h"
 #include "Subsystems/AssetEditorSubsystem.h"
@@ -24,7 +25,7 @@ URemoteControlPreset* FAvaRCExtension::GetRemoteControlPreset() const
 void FAvaRCExtension::Activate()
 {
 	RegisterSignatureCustomization();
-	OpenRemoteControlTab();
+	OpenRemoteControlTab(/*bCreateIfNotFound*/false);
 }
 
 void FAvaRCExtension::Deactivate()
@@ -46,13 +47,23 @@ void FAvaRCExtension::ExtendToolbarMenu(UToolMenu& InMenu)
 {
 	FToolMenuSection& Section = InMenu.FindOrAddSection(DefaultSectionName);
 
+	const FSlateIcon RemoteControlIcon(FAppStyle::GetAppStyleSetName(), "LevelEditor.Tabs.Details");
+
 	FToolMenuEntry& Entry = Section.AddEntry(FToolMenuEntry::InitToolBarButton(TEXT("RemoteControlButton")
-		,  FExecuteAction::CreateSP(this, &FAvaRCExtension::OpenRemoteControlTab)
+		,  FExecuteAction::CreateSP(this, &FAvaRCExtension::OpenRemoteControlTab, /*bCreateIfNotFound*/true)
 		, LOCTEXT("RemoteControlLabel"  , "Remote Control")
 		, LOCTEXT("RemoteControlTooltip", "Opens the Remote Control Editor for the given Scene")
-		, FSlateIcon(FAppStyle::GetAppStyleSetName(), "LevelEditor.Tabs.Details")));
+		, RemoteControlIcon));
 
 	Entry.StyleNameOverride = "CalloutToolbar";
+
+	Section.AddEntry(FToolMenuEntry::InitComboButton(TEXT("RemoteControlComboButton")
+		, FUIAction()
+		, FNewToolMenuDelegate::CreateSP(this, &FAvaRCExtension::GenerateRemoteControlOptions)
+		, LOCTEXT("RemoteControlOptionsLabel", "Remote Control Options")
+		, LOCTEXT("RemoteControlOptionsTooltip", "Remote Control Options")
+		, RemoteControlIcon
+		, /*bSimpleComboBox*/true));
 }
 
 void FAvaRCExtension::OnSceneObjectChanged(UObject* InOldSceneObject, UObject* InNewSceneObject)
@@ -95,8 +106,38 @@ void FAvaRCExtension::OnSceneObjectChanged(UObject* InOldSceneObject, UObject* I
 	}
 }
 
-void FAvaRCExtension::OpenRemoteControlTab() const
+void FAvaRCExtension::OpenRemoteControlTab(bool bInCreateIfNotFound) const
 {
+	IAvaSceneInterface* const Scene = GetSceneObject<IAvaSceneInterface>();
+	if (!Scene)
+	{
+		return;
+	}
+
+	URemoteControlPreset* RemoteControlPreset = Scene->GetRemoteControlPreset();
+	if (!RemoteControlPreset)
+	{
+		if (!bInCreateIfNotFound)
+		{
+			return;
+		}
+
+		const EAppReturnType::Type Response = FMessageDialog::Open(EAppMsgType::YesNo
+			, LOCTEXT("CreateRCPresetMessage", "Remote Control Preset is not present in this level. Do you wish to create it?")
+			, LOCTEXT("CreateRCPresetTitle", "Create Remote Control Preset"));
+
+		if (Response != EAppReturnType::Yes)
+		{
+			return;
+		}
+
+		RemoteControlPreset = Scene->GetOrCreateRemoteControlPreset();
+		if (!RemoteControlPreset)
+		{
+			return;
+		}
+	}
+
 	TSharedPtr<IAvaEditor> Editor = GetEditor();
 	if (!Editor.IsValid())
 	{
@@ -109,12 +150,9 @@ void FAvaRCExtension::OpenRemoteControlTab() const
 		return;
 	}
 
-	if (URemoteControlPreset* const RemoteControlPreset = GetRemoteControlPreset())
-	{
-		UAssetEditorSubsystem* AssetEditorSubsystem = GEditor->GetEditorSubsystem<UAssetEditorSubsystem>();
-		check(AssetEditorSubsystem);
-		AssetEditorSubsystem->OpenEditorForAsset(RemoteControlPreset, EToolkitMode::WorldCentric, ToolkitHost);
-	}
+	UAssetEditorSubsystem* AssetEditorSubsystem = GEditor->GetEditorSubsystem<UAssetEditorSubsystem>();
+	check(AssetEditorSubsystem);
+	AssetEditorSubsystem->OpenEditorForAsset(RemoteControlPreset, EToolkitMode::WorldCentric, ToolkitHost);
 }
 
 void FAvaRCExtension::CloseRemoteControlTab() const
@@ -151,4 +189,44 @@ void FAvaRCExtension::UnregisterSignatureCustomization()
 	SignatureCustomization.Reset();
 }
 
+bool FAvaRCExtension::CanDeleteRemoteControlPreset() const
+{
+	return !!GetRemoteControlPreset();
+}
+
+void FAvaRCExtension::DeleteRemoteControlPreset()
+{
+	const EAppReturnType::Type Response = FMessageDialog::Open(EAppMsgType::YesNo
+		, LOCTEXT("DeleteRemoteControlMessage", "Do you wish to delete the Remote Control Preset? (experimental)")
+		, LOCTEXT("DeleteRemoteControlTitle", "Delete Remote Control Preset"));
+
+	if (Response != EAppReturnType::Yes)
+	{
+		return;
+	}
+
+	CloseRemoteControlTab();
+
+	if (IAvaSceneInterface* const Scene = GetSceneObject<IAvaSceneInterface>())
+	{
+		constexpr bool bCallModify = true;
+PRAGMA_DISABLE_EXPERIMENTAL_WARNINGS
+		Scene->DestroyRemoteControlPreset(bCallModify);
+PRAGMA_ENABLE_EXPERIMENTAL_WARNINGS
+	}
+}
+
+void FAvaRCExtension::GenerateRemoteControlOptions(UToolMenu* InMenu)
+{
+	FToolMenuSection& GeneralSection = InMenu->FindOrAddSection(TEXT("GeneralSection"), LOCTEXT("GeneralSectionLabel", "General"));
+
+	GeneralSection.AddMenuEntry(TEXT("DeleteRemoteControl")
+		, LOCTEXT("DeleteRemoteControlLabel", "Delete Remote Control (experimental)")
+		, LOCTEXT("DeleteRemoteControlTooltip", "Deletes the existing embedded Remote Control.")
+		, TAttribute<FSlateIcon>()
+		, FUIAction(FExecuteAction::CreateSP(this, &FAvaRCExtension::DeleteRemoteControlPreset)
+			, FCanExecuteAction::CreateSP(this, &FAvaRCExtension::CanDeleteRemoteControlPreset))
+		, EUserInterfaceActionType::Button);
+}
+
 #undef LOCTEXT_NAMESPACE
diff --git a/Engine/Plugins/VirtualProduction/Avalanche/Source/AvalancheEditor/Private/RemoteControl/AvaRCExtension.h b/Engine/Plugins/VirtualProduction/Avalanche/Source/AvalancheEditor/Private/RemoteControl/AvaRCExtension.h
--- a/Engine/Plugins/VirtualProduction/Avalanche/Source/AvalancheEditor/Private/RemoteControl/AvaRCExtension.h
+++ b/Engine/Plugins/VirtualProduction/Avalanche/Source/AvalancheEditor/Private/RemoteControl/AvaRCExtension.h
@@ -26,12 +26,20 @@ public:
 	//~ End IAvaEditorExtension
 
 protected:
-	void OpenRemoteControlTab() const;
+	void OpenRemoteControlTab(bool bInCreateIfNotFound) const;
 	void CloseRemoteControlTab() const;
 
 	void RegisterSignatureCustomization();
 	void UnregisterSignatureCustomization();
 
-private:
+	/** Determines whether the scene's rc preset can be deleted */
+	bool CanDeleteRemoteControlPreset() const;
+
+	/** Deletes the scene's rc preset (if any) */
+	void DeleteRemoteControlPreset();
+
+	/** Generates the sub-menu options for remote control */
+	void GenerateRemoteControlOptions(UToolMenu* InMenu);
+
 	TSharedPtr<IRCSignatureCustomization> SignatureCustomization;
 };
diff --git a/Engine/Plugins/VirtualProduction/Avalanche/Source/AvalancheMRQ/Private/AvaMRQRundownPageSetting.cpp b/Engine/Plugins/VirtualProduction/Avalanche/Source/AvalancheMRQ/Private/AvaMRQRundownPageSetting.cpp
--- a/Engine/Plugins/VirtualProduction/Avalanche/Source/AvalancheMRQ/Private/AvaMRQRundownPageSetting.cpp
+++ b/Engine/Plugins/VirtualProduction/Avalanche/Source/AvalancheMRQ/Private/AvaMRQRundownPageSetting.cpp
@@ -111,7 +111,7 @@ void UAvaMRQRundownPageSetting::SetupForPipelineImpl(UMoviePipeline* InPipeline)
 	URemoteControlPreset* RemoteControlPreset = Scene->GetRemoteControlPreset();
 	if (!RemoteControlPreset)
 	{
-		UE_LOG(LogAvaMRQSequenceData, Error
+		UE_CLOG(Scene->IsRemoteControlPresetExpectedToBePresent(), LogAvaMRQSequenceData, Error
 			, TEXT("A Remote Control Preset was not found in the provided World '%s'.")
 			, *World->GetFullName());
 		return;
diff --git a/Engine/Plugins/VirtualProduction/Avalanche/Source/AvalancheMedia/Private/Playable/AvaPlayable.cpp b/Engine/Plugins/VirtualProduction/Avalanche/Source/AvalancheMedia/Private/Playable/AvaPlayable.cpp
--- a/Engine/Plugins/VirtualProduction/Avalanche/Source/AvalancheMedia/Private/Playable/AvaPlayable.cpp
+++ b/Engine/Plugins/VirtualProduction/Avalanche/Source/AvalancheMedia/Private/Playable/AvaPlayable.cpp
@@ -279,15 +279,19 @@ EAvaPlayableCommandResult UAvaPlayable::UpdateRemoteControlCommand(const TShared
 	{
 		return EAvaPlayableCommandResult::ErrorDiscard;
 	}
-	
-	URemoteControlPreset* RemoteControlPreset = Scene->GetRemoteControlPreset();
-	
-	if (!IsValid(RemoteControlPreset))
+
+	URemoteControlPreset* const RemoteControlPreset = Scene->GetRemoteControlPreset();
+	if (!RemoteControlPreset)
 	{
-		UE_LOG(LogAvaPlayable, Error,
+		const bool bIsRemoteControlPresetExpectedToBePresent = Scene->IsRemoteControlPresetExpectedToBePresent();
+
+		UE_CLOG(bIsRemoteControlPresetExpectedToBePresent, LogAvaPlayable, Error,
 			TEXT("Remote Control command for asset \"%s\": Remote Control Preset is null."),
 			*GetSourceAssetPath().ToString());
-		return EAvaPlayableCommandResult::ErrorDiscard;
+
+		return bIsRemoteControlPresetExpectedToBePresent
+			? EAvaPlayableCommandResult::ErrorDiscard
+			: EAvaPlayableCommandResult::Executed;
 	}
 
 	using namespace UE::AvaPlayable::Private;
