Motion Design: added render state dirty reason scope for the modifier system to have a better idea of where a render state dirty event call might be for.

#Synced-CL 52400047

diff --git a/Engine/Plugins/VirtualProduction/ActorModifier/Source/ActorModifier/Private/Extensions/ActorModifierRenderStateUpdateExtension.cpp b/Engine/Plugins/VirtualProduction/ActorModifier/Source/ActorModifier/Private/Extensions/ActorModifierRenderStateUpdateExtension.cpp
--- a/Engine/Plugins/VirtualProduction/ActorModifier/Source/ActorModifier/Private/Extensions/ActorModifierRenderStateUpdateExtension.cpp
+++ b/Engine/Plugins/VirtualProduction/ActorModifier/Source/ActorModifier/Private/Extensions/ActorModifierRenderStateUpdateExtension.cpp
@@ -5,13 +5,14 @@
 #include "Components/PrimitiveComponent.h"
 #include "Containers/Ticker.h"
 #include "Modifiers/ActorModifierCoreBase.h"
+#include "Modifiers/ActorModifierRenderStateDirtyEvent.h"
 #include "Subsystems/ActorModifierCoreSubsystem.h"
 #include "Utilities/ActorModifierActorUtils.h"
 
-FActorModifierRenderStateUpdateExtension::FActorModifierRenderStateUpdateExtension(IActorModifierRenderStateUpdateHandler* InExtensionHandler)
+FActorModifierRenderStateUpdateExtension::FActorModifierRenderStateUpdateExtension(TNotNull<IActorModifierRenderStateUpdateHandler*> InExtensionHandler)
 	: ExtensionHandlerWeak(InExtensionHandler)
+	, RelevantReasons(InExtensionHandler->GetRenderStateDirtyEventTypes())
 {
-	check(InExtensionHandler);
 }
 
 void FActorModifierRenderStateUpdateExtension::TrackActorVisibility(AActor* InActor)
@@ -96,6 +97,11 @@ void FActorModifierRenderStateUpdateExtension::OnExtensionDisabled(EActorModifie
 
 void FActorModifierRenderStateUpdateExtension::OnRenderStateDirty(UActorComponent& InComponent)
 {
+	if (!UE::ActorModifierCore::IsRenderStateDirtyReasonPresent(RelevantReasons))
+	{
+		return;
+	}
+
 	const AActor* ModifierActor = GetModifierActor();
 	AActor* ActorDirty = InComponent.GetOwner();
 
@@ -120,7 +126,9 @@ void FActorModifierRenderStateUpdateExtension::OnRenderStateDirty(UActorComponen
 		HandlerInterface->OnRenderStateUpdated(ActorDirty, &InComponent);
 	}
 
-	if (!TrackedActorsVisibility.Contains(ActorDirty))
+	// Skip if render state is not relevant or not tracking the actor
+	if (!UE::ActorModifierCore::IsRenderStateDirtyReasonPresent(UE::ActorModifierCore::ERenderStateDirtyReason::Visibility)
+		|| !TrackedActorsVisibility.Contains(ActorDirty))
 	{
 		return;
 	}
diff --git a/Engine/Plugins/VirtualProduction/ActorModifier/Source/ActorModifier/Private/Extensions/ActorModifierSceneTreeUpdateExtension.cpp b/Engine/Plugins/VirtualProduction/ActorModifier/Source/ActorModifier/Private/Extensions/ActorModifierSceneTreeUpdateExtension.cpp
--- a/Engine/Plugins/VirtualProduction/ActorModifier/Source/ActorModifier/Private/Extensions/ActorModifierSceneTreeUpdateExtension.cpp
+++ b/Engine/Plugins/VirtualProduction/ActorModifier/Source/ActorModifier/Private/Extensions/ActorModifierSceneTreeUpdateExtension.cpp
@@ -6,6 +6,7 @@
 #include "Engine/Level.h"
 #include "Engine/World.h"
 #include "GameFramework/Actor.h"
+#include "Modifiers/ActorModifierRenderStateDirtyEvent.h"
 #include "Subsystems/ActorModifierCoreSubsystem.h"
 #include "Utilities/ActorModifierActorUtils.h"
 
@@ -151,6 +152,15 @@ void FActorModifierSceneTreeUpdateExtension::OnRefreshTrackedActors(AActor* InAc
 
 void FActorModifierSceneTreeUpdateExtension::OnRenderStateDirty(UActorComponent& InComponent)
 {
+	// Scene tree update is only concerned with scene tree / visibility changes
+	constexpr UE::ActorModifierCore::ERenderStateDirtyReason RelevantReasons = UE::ActorModifierCore::ERenderStateDirtyReason::Visibility 
+		| UE::ActorModifierCore::ERenderStateDirtyReason::SceneTree;
+
+	if (!UE::ActorModifierCore::IsRenderStateDirtyReasonPresent(RelevantReasons))
+	{
+		return;
+	}
+
 	AActor* OwningActor = InComponent.GetOwner();
 	if (!IsValid(OwningActor))
 	{
diff --git a/Engine/Plugins/VirtualProduction/ActorModifier/Source/ActorModifier/Private/Shared/ActorModifierVisibilityShared.cpp b/Engine/Plugins/VirtualProduction/ActorModifier/Source/ActorModifier/Private/Shared/ActorModifierVisibilityShared.cpp
--- a/Engine/Plugins/VirtualProduction/ActorModifier/Source/ActorModifier/Private/Shared/ActorModifierVisibilityShared.cpp
+++ b/Engine/Plugins/VirtualProduction/ActorModifier/Source/ActorModifier/Private/Shared/ActorModifierVisibilityShared.cpp
@@ -3,6 +3,7 @@
 #include "Shared/ActorModifierVisibilityShared.h"
 
 #include "GameFramework/Actor.h"
+#include "Modifiers/ActorModifierRenderStateDirtyEvent.h"
 
 void FActorModifierVisibilitySharedModifierState::Save(const AActor* InActor)
 {
@@ -10,6 +11,7 @@ void FActorModifierVisibilitySharedModifierState::Save(const AActor* InActor)
 	{
 		if (InActor)
 		{
+			UE::ActorModifierCore::FRenderStateDirtyReasonScope Scope(UE::ActorModifierCore::ERenderStateDirtyReason::Visibility);
 #if WITH_EDITOR
 			bActorHiddenInEditor = InActor->IsTemporarilyHiddenInEditor();
 #endif
@@ -24,6 +26,7 @@ void FActorModifierVisibilitySharedModifierState::Restore(AActor* InActor) const
 	{
 		if (InActor)
 		{
+			UE::ActorModifierCore::FRenderStateDirtyReasonScope Scope(UE::ActorModifierCore::ERenderStateDirtyReason::Visibility);
 #if WITH_EDITOR
 			InActor->SetIsTemporarilyHiddenInEditor(bActorHiddenInEditor);
 #endif
@@ -47,6 +50,7 @@ void FActorModifierVisibilitySharedActorState::Restore() const
 {
 	if (AActor* Actor = ActorWeak.Get())
 	{
+		UE::ActorModifierCore::FRenderStateDirtyReasonScope Scope(UE::ActorModifierCore::ERenderStateDirtyReason::Visibility);
 #if WITH_EDITOR
 		Actor->SetIsTemporarilyHiddenInEditor(bActorHiddenInEditor);
 #endif
@@ -141,6 +145,8 @@ void UActorModifierVisibilityShared::SetActorsVisibility(UActorModifierCoreBase*
 		
 		SaveActorState(InModifierContext, Actor);
 
+		UE::ActorModifierCore::FRenderStateDirtyReasonScope Scope(UE::ActorModifierCore::ERenderStateDirtyReason::Visibility);
+
 #if WITH_EDITOR
 		if (EnumHasAnyFlags(InActorVisibility, EActorModifierVisibilityActor::Editor))
 		{
diff --git a/Engine/Plugins/VirtualProduction/ActorModifier/Source/ActorModifier/Public/Extensions/ActorModifierRenderStateUpdateExtension.h b/Engine/Plugins/VirtualProduction/ActorModifier/Source/ActorModifier/Public/Extensions/ActorModifierRenderStateUpdateExtension.h
--- a/Engine/Plugins/VirtualProduction/ActorModifier/Source/ActorModifier/Public/Extensions/ActorModifierRenderStateUpdateExtension.h
+++ b/Engine/Plugins/VirtualProduction/ActorModifier/Source/ActorModifier/Public/Extensions/ActorModifierRenderStateUpdateExtension.h
@@ -3,6 +3,7 @@
 #pragma once
 
 #include "Modifiers/ActorModifierCoreExtension.h"
+#include "Modifiers/ActorModifierRenderStateDirtyEvent.h"
 #include "UObject/WeakInterfacePtr.h"
 #include "ActorModifierRenderStateUpdateExtension.generated.h"
 
@@ -21,6 +22,15 @@ class IActorModifierRenderStateUpdateHandler
 	GENERATED_BODY()
 
 public:
+	/** 
+	 * Returns the relevant render state dirty event types relevant to this handler 
+	 * Defaults to all for legacy
+	 */
+	virtual UE::ActorModifierCore::ERenderStateDirtyReason GetRenderStateDirtyEventTypes() const
+	{
+		return UE::ActorModifierCore::ERenderStateDirtyReason::All;
+	}
+
 	/** Callback when a render state actor in this world changes */
 	virtual void OnRenderStateUpdated(AActor* InActor, UActorComponent* InComponent) = 0;
 
@@ -36,7 +46,7 @@ class FActorModifierRenderStateUpdateExtension : public FActorModifierCoreExtens
 {
 
 public:
-	ACTORMODIFIER_API explicit FActorModifierRenderStateUpdateExtension(IActorModifierRenderStateUpdateHandler* InExtensionHandler);
+	ACTORMODIFIER_API explicit FActorModifierRenderStateUpdateExtension(TNotNull<IActorModifierRenderStateUpdateHandler*> InExtensionHandler);
 
 	/** Adds an actor to track for visibility */
 	ACTORMODIFIER_API void TrackActorVisibility(AActor* InActor);
@@ -68,4 +78,7 @@ private:
 	TWeakInterfacePtr<IActorModifierRenderStateUpdateHandler> ExtensionHandlerWeak;
 
 	TMap<TWeakObjectPtr<AActor>, bool> TrackedActorsVisibility;
+
+	/** The render state dirty reasons that are relevant to this extension */
+	UE::ActorModifierCore::ERenderStateDirtyReason RelevantReasons;
 };
diff --git a/Engine/Plugins/VirtualProduction/ActorModifierCore/Source/ActorModifierCore/Private/Modifiers/ActorModifierRenderStateDirtyEvent.cpp b/Engine/Plugins/VirtualProduction/ActorModifierCore/Source/ActorModifierCore/Private/Modifiers/ActorModifierRenderStateDirtyEvent.cpp
new file mode 100644
--- /dev/null
+++ b/Engine/Plugins/VirtualProduction/ActorModifierCore/Source/ActorModifierCore/Private/Modifiers/ActorModifierRenderStateDirtyEvent.cpp
@@ -0,0 +1,36 @@
+// Copyright Epic Games, Inc. All Rights Reserved.
+
+#include "Modifiers/ActorModifierRenderStateDirtyEvent.h"
+
+namespace UE::ActorModifierCore
+{
+namespace Private
+{
+	// Defaults to unknown as most MarkRenderStateDirty calls will not have a reason scope to it
+	static ERenderStateDirtyReason RenderStateDirtyReason = ERenderStateDirtyReason::Unknown;
+
+} // UE::ActorModifierCore::Private
+
+ERenderStateDirtyReason GetRenderStateDirtyReason()
+{
+	return Private::RenderStateDirtyReason;
+}
+
+bool IsRenderStateDirtyReasonPresent(ERenderStateDirtyReason InReason)
+{
+	return Private::RenderStateDirtyReason == ERenderStateDirtyReason::Unknown
+		|| EnumHasAnyFlags(Private::RenderStateDirtyReason, InReason); 
+}
+
+FRenderStateDirtyReasonScope::FRenderStateDirtyReasonScope(ERenderStateDirtyReason InReason)
+	: PreviousReason(Private::RenderStateDirtyReason)
+{
+	Private::RenderStateDirtyReason |= InReason;
+}
+
+FRenderStateDirtyReasonScope::~FRenderStateDirtyReasonScope()
+{
+	Private::RenderStateDirtyReason = PreviousReason;
+}
+
+} //  UE::ActorModifierCore
diff --git a/Engine/Plugins/VirtualProduction/ActorModifierCore/Source/ActorModifierCore/Public/Modifiers/ActorModifierRenderStateDirtyEvent.h b/Engine/Plugins/VirtualProduction/ActorModifierCore/Source/ActorModifierCore/Public/Modifiers/ActorModifierRenderStateDirtyEvent.h
new file mode 100644
--- /dev/null
+++ b/Engine/Plugins/VirtualProduction/ActorModifierCore/Source/ActorModifierCore/Public/Modifiers/ActorModifierRenderStateDirtyEvent.h
@@ -0,0 +1,57 @@
+// Copyright Epic Games, Inc. All Rights Reserved.
+
+#pragma once
+
+#include "HAL/Platform.h"
+#include "Misc/EnumClassFlags.h"
+#include "Templates/UnrealTemplate.h"
+
+#define UE_API ACTORMODIFIERCORE_API
+
+namespace UE::ActorModifierCore
+{
+
+enum class ERenderStateDirtyReason : uint8
+{
+	/** Render state change reason is not known */
+	Unknown = 0,
+	/** Render state was changed because a primitive became visible/hidden */
+	Visibility = 1 << 0,
+	/** Render state was changed because a material was changed */
+	Material = 1 << 1,
+	/** Render state was dirtied because of translucency priority */
+	TranslucencyPriority = 1 << 2,
+	/** Render state was dirtied because of geometry change */
+	Geometry = 1 << 3,
+	/** Render state dirtied because of a scene tree change */
+	SceneTree = 1 << 4,
+
+	All = 0xFF,
+};
+ENUM_CLASS_FLAGS(ERenderStateDirtyReason);
+
+/** Get the current render state reason */
+UE_API ERenderStateDirtyReason GetRenderStateDirtyReason();
+
+/** 
+ * Determines whether the given reason is part of the current render state dirty reason
+ * When the current reason is 'unknown' this returns true.
+ */
+UE_API bool IsRenderStateDirtyReasonPresent(ERenderStateDirtyReason InReason);
+
+/**
+ * Tracks the reason a render state dirty happened.
+ * Used to avoid unnecessary computations from modifiers.
+ */
+struct FRenderStateDirtyReasonScope : public FNoncopyable
+{
+	UE_API explicit FRenderStateDirtyReasonScope(ERenderStateDirtyReason InReason);
+	UE_API ~FRenderStateDirtyReasonScope();
+
+private:
+	ERenderStateDirtyReason PreviousReason;
+};
+
+} // UE::ActorModifierCore
+
+#undef UE_API
diff --git a/Engine/Plugins/VirtualProduction/Avalanche/Source/AvalancheMask/Private/Mask2D/AvaMask2DReadModifier.cpp b/Engine/Plugins/VirtualProduction/Avalanche/Source/AvalancheMask/Private/Mask2D/AvaMask2DReadModifier.cpp
--- a/Engine/Plugins/VirtualProduction/Avalanche/Source/AvalancheMask/Private/Mask2D/AvaMask2DReadModifier.cpp
+++ b/Engine/Plugins/VirtualProduction/Avalanche/Source/AvalancheMask/Private/Mask2D/AvaMask2DReadModifier.cpp
@@ -25,6 +25,7 @@
 #include "Materials/Material.h"
 #include "Materials/MaterialInstanceDynamic.h"
 #include "Modifiers/ActorModifierCoreStack.h"
+#include "Modifiers/ActorModifierRenderStateDirtyEvent.h"
 #include "Modifiers/Utilities/ActorModifierCoreUtilities.h"
 
 #if WITH_EDITOR
@@ -254,6 +255,8 @@ bool UAvaMask2DReadModifier::ApplyRead(TNotNull<UObject*> InMaterialContainer, T
 
 void UAvaMask2DReadModifier::ApplyMaskState(const UE::Ava::FMaterialBridgeWriteSlotContext& InContext, UE::Ava::FMaterialBridgeWriteSlot& InSlot, const FAvaMask2DMaskState& InMaskState)
 {
+	UE::ActorModifierCore::FRenderStateDirtyReasonScope Scope(UE::ActorModifierCore::ERenderStateDirtyReason::Material);
+
 	// Request the container to support the current blend mode.
 	// Do this first as the slot material could change from this request.
 	FAvaMaterialBridgeBlendModeFeature BlendModeFeature;
@@ -403,6 +406,7 @@ void UAvaMask2DReadModifier::RestorePreState()
 		// Actors have a built-in bridge so there should always be a valid material bridge found
 		if (ensureMsgf(MaterialBridge, TEXT("Could not find a valid material bridge for actor of type '%s'!"), *GetNameSafe(Pair.Key->GetClass())))
 		{
+			UE::ActorModifierCore::FRenderStateDirtyReasonScope Scope(UE::ActorModifierCore::ERenderStateDirtyReason::Material);
 			MaterialBridge->ApplyState(UE::Ava::FMaterialBridgeApplyStateContext(Pair.Key), Pair.Value, /*Options*/{});
 		}
 	}
diff --git a/Engine/Plugins/VirtualProduction/Avalanche/Source/AvalancheModifiers/Private/Modifiers/AvaTranslucentPriorityModifier.cpp b/Engine/Plugins/VirtualProduction/Avalanche/Source/AvalancheModifiers/Private/Modifiers/AvaTranslucentPriorityModifier.cpp
--- a/Engine/Plugins/VirtualProduction/Avalanche/Source/AvalancheModifiers/Private/Modifiers/AvaTranslucentPriorityModifier.cpp
+++ b/Engine/Plugins/VirtualProduction/Avalanche/Source/AvalancheModifiers/Private/Modifiers/AvaTranslucentPriorityModifier.cpp
@@ -2,15 +2,15 @@
 
 #include "Modifiers/AvaTranslucentPriorityModifier.h"
 
-#include "EngineUtils.h"
 #include "Camera/CameraActor.h"
 #include "Camera/CameraComponent.h"
 #include "Components/PrimitiveComponent.h"
-#include "EngineUtils.h"
 #include "Engine/RendererSettings.h"
+#include "EngineUtils.h"
 #include "Framework/AvaGameInstance.h"
 #include "GameFramework/Actor.h"
 #include "Modifiers/ActorModifierCoreBase.h"
+#include "Modifiers/ActorModifierRenderStateDirtyEvent.h"
 #include "Shared/AvaTranslucentPriorityModifierShared.h"
 
 #if WITH_EDITOR
@@ -163,6 +163,8 @@ void UAvaTranslucentPriorityModifier::Apply()
 				if (UPrimitiveComponent* const Component = SortedComponentState.PrimitiveComponent)
 				{
 					LastSortPriorities.Add(Component, TranslucentSortPriority);
+					
+					UE::ActorModifierCore::FRenderStateDirtyReasonScope Scope(UE::ActorModifierCore::ERenderStateDirtyReason::TranslucencyPriority);
 					Component->SetTranslucentSortPriority(TranslucentSortPriority);
 
 					TranslucentSortPriority += GlobalStep;
@@ -198,6 +200,8 @@ void UAvaTranslucentPriorityModifier::Apply()
 			if (SortedComponentState.Modifier == this)
 			{
 				LastSortPriorities.Add(Component, TranslucentSortPriority);
+
+				UE::ActorModifierCore::FRenderStateDirtyReasonScope Scope(UE::ActorModifierCore::ERenderStateDirtyReason::TranslucencyPriority);
 				Component->SetTranslucentSortPriority(TranslucentSortPriority);
 			}
 			// Another modifier handles this component
@@ -403,6 +407,12 @@ void UAvaTranslucentPriorityModifier::OnSceneTreeTrackedActorRearranged(int32 In
 	}
 }
 
+UE::ActorModifierCore::ERenderStateDirtyReason UAvaTranslucentPriorityModifier::GetRenderStateDirtyEventTypes() const
+{
+	return UE::ActorModifierCore::ERenderStateDirtyReason::TranslucencyPriority
+		| UE::ActorModifierCore::ERenderStateDirtyReason::SceneTree;
+}
+
 void UAvaTranslucentPriorityModifier::OnRenderStateUpdated(AActor* InActor, UActorComponent* InComponent)
 {
 	Super::OnRenderStateUpdated(InActor, InComponent);
diff --git a/Engine/Plugins/VirtualProduction/Avalanche/Source/AvalancheModifiers/Private/Shared/AvaTranslucentPriorityModifierShared.cpp b/Engine/Plugins/VirtualProduction/Avalanche/Source/AvalancheModifiers/Private/Shared/AvaTranslucentPriorityModifierShared.cpp
--- a/Engine/Plugins/VirtualProduction/Avalanche/Source/AvalancheModifiers/Private/Shared/AvaTranslucentPriorityModifierShared.cpp
+++ b/Engine/Plugins/VirtualProduction/Avalanche/Source/AvalancheModifiers/Private/Shared/AvaTranslucentPriorityModifierShared.cpp
@@ -9,6 +9,7 @@
 #include "Camera/CameraComponent.h"
 #include "Framework/AvaGameInstance.h"
 #include "IAvaSceneInterface.h"
+#include "Modifiers/ActorModifierRenderStateDirtyEvent.h"
 #include "Modifiers/AvaTranslucentPriorityModifier.h"
 #include "Version/AvaTranslucentPriorityModifierVersion.h"
 
@@ -352,6 +353,7 @@ void FAvaTranslucentPriorityModifierComponentState::Restore() const
 {
 	if (PrimitiveComponent)
 	{
+		UE::ActorModifierCore::FRenderStateDirtyReasonScope Scope(UE::ActorModifierCore::ERenderStateDirtyReason::TranslucencyPriority);
 		PrimitiveComponent->SetTranslucentSortPriority(SortPriority);
 	}
 }
diff --git a/Engine/Plugins/VirtualProduction/Avalanche/Source/AvalancheModifiers/Public/Modifiers/AvaTranslucentPriorityModifier.h b/Engine/Plugins/VirtualProduction/Avalanche/Source/AvalancheModifiers/Public/Modifiers/AvaTranslucentPriorityModifier.h
--- a/Engine/Plugins/VirtualProduction/Avalanche/Source/AvalancheModifiers/Public/Modifiers/AvaTranslucentPriorityModifier.h
+++ b/Engine/Plugins/VirtualProduction/Avalanche/Source/AvalancheModifiers/Public/Modifiers/AvaTranslucentPriorityModifier.h
@@ -118,6 +118,7 @@ protected:
 	//~ End IAvaSceneTreeUpdateModifierExtension
 
 	//~ Begin IAvaRenderStateUpdateExtension
+	virtual UE::ActorModifierCore::ERenderStateDirtyReason GetRenderStateDirtyEventTypes() const override;
 	virtual void OnRenderStateUpdated(AActor* InActor, UActorComponent* InComponent) override;
 	//~ End IAvaRenderStateUpdateExtension
 
diff --git a/Engine/Plugins/VirtualProduction/Avalanche/Source/AvalancheShapes/Private/DynamicMeshes/AvaShapeDynMeshBase.cpp b/Engine/Plugins/VirtualProduction/Avalanche/Source/AvalancheShapes/Private/DynamicMeshes/AvaShapeDynMeshBase.cpp
--- a/Engine/Plugins/VirtualProduction/Avalanche/Source/AvalancheShapes/Private/DynamicMeshes/AvaShapeDynMeshBase.cpp
+++ b/Engine/Plugins/VirtualProduction/Avalanche/Source/AvalancheShapes/Private/DynamicMeshes/AvaShapeDynMeshBase.cpp
@@ -7,12 +7,12 @@
 #include "AvaSceneSubsystem.h"
 #include "AvaShapeActor.h"
 #include "AvaShapeBuildSystem.h"
-#include "DynamicMeshMikkTWrapper.h"
 #include "Components/DynamicMeshComponent.h"
 #include "Containers/Ticker.h"
 #include "DynamicMesh/DynamicMesh3.h"
 #include "DynamicMesh/DynamicMeshAttributeSet.h"
 #include "DynamicMesh/MeshTransforms.h"
+#include "DynamicMeshMikkTWrapper.h"
 #include "DynamicMeshes/AvaShape2DDynMeshBase.h"
 #include "Engine/StaticMesh.h"
 #include "Framework/AvaGizmoComponent.h"
@@ -26,6 +26,7 @@
 #include "Materials/MaterialInstanceDynamic.h"
 #include "Materials/MaterialInterface.h"
 #include "Modifiers/ActorModifierCoreStack.h"
+#include "Modifiers/ActorModifierRenderStateDirtyEvent.h"
 #include "Modules/ModuleManager.h"
 #include "Parameterization/DynamicMeshUVEditor.h"
 #include "PhysicsEngine/AggregateGeom.h"
@@ -418,6 +419,8 @@ void UAvaShapeDynamicMeshBase::PostUpdateMesh()
 			UpdateBaseMaterial();
 		}
 
+		UE::ActorModifierCore::FRenderStateDirtyReasonScope Scope(UE::ActorModifierCore::ERenderStateDirtyReason::Material | UE::ActorModifierCore::ERenderStateDirtyReason::Geometry);
+
 		if (MeshIndex != INDEX_NONE)
 		{
 			// Call SetMaterial on last index/material to mark render state dirty, update physical materials, etc.
@@ -2708,6 +2711,7 @@ void UAvaShapeDynamicMeshBase::MarkMeshRenderStateDirty() const
 {
 	if (UDynamicMeshComponent* MeshComponent = GetShapeMeshComponent())
 	{
+		UE::ActorModifierCore::FRenderStateDirtyReasonScope Scope(UE::ActorModifierCore::ERenderStateDirtyReason::Material);
 		MeshComponent->MarkRenderStateDirty();
 	}
 }
diff --git a/Engine/Plugins/VirtualProduction/Avalanche/Source/AvalancheText/AvalancheText.Build.cs b/Engine/Plugins/VirtualProduction/Avalanche/Source/AvalancheText/AvalancheText.Build.cs
--- a/Engine/Plugins/VirtualProduction/Avalanche/Source/AvalancheText/AvalancheText.Build.cs
+++ b/Engine/Plugins/VirtualProduction/Avalanche/Source/AvalancheText/AvalancheText.Build.cs
@@ -24,7 +24,8 @@ public class AvalancheText : ModuleRules
 		PrivateDependencyModuleNames.AddRange(
 			new string[]
 			{
-				"Avalanche"
+				"ActorModifierCore",
+				"Avalanche",
 			}
 		);
 
diff --git a/Engine/Plugins/VirtualProduction/Avalanche/Source/AvalancheText/Private/AvaTextModule.cpp b/Engine/Plugins/VirtualProduction/Avalanche/Source/AvalancheText/Private/AvaTextModule.cpp
--- a/Engine/Plugins/VirtualProduction/Avalanche/Source/AvalancheText/Private/AvaTextModule.cpp
+++ b/Engine/Plugins/VirtualProduction/Avalanche/Source/AvalancheText/Private/AvaTextModule.cpp
@@ -6,7 +6,9 @@
 #include "MaterialBridge/AvaText3DDefaultExtensionMaterialBridge.h"
 #include "MaterialBridge/AvaText3DExtensionMaterialBridge.h"
 #include "MaterialBridge/AvaText3DManagedComponentMaterialBridge.h"
+#include "Modifiers/ActorModifierRenderStateDirtyEvent.h"
 #include "Modules/ModuleManager.h"
+#include "Text3DDelegates.h"
 
 IMPLEMENT_MODULE(FAvaTextModule, AvalancheText)
 
@@ -23,4 +25,43 @@ void FAvaTextModule::StartupModule()
 	constexpr uint32 HigherPriority = 100;
 
 	MaterialBridgeRegistry.Register<UE::Ava::FText3DManagedComponentMaterialBridge>(HigherPriority);
+
+	RegisterTextRendererDelegates();
+}
+
+void FAvaTextModule::ShutdownModule()
+{
+	UnregisterTextRendererDelegates();
+}
+
+void FAvaTextModule::RegisterTextRendererDelegates()
+{
+	static TArray<TUniquePtr<UE::ActorModifierCore::FRenderStateDirtyReasonScope>, TInlineAllocator<8>> Scope;
+
+	TextPreRendererUpdateHandle = UE::Text3D::OnText3DPreRendererUpdate().AddLambda(
+		[](const UE::Text3D::FText3DBuildContext& InBuildContext)
+		{
+			if (InBuildContext.CurrentFlag == EText3DRendererFlags::Material)
+			{
+				Scope.Push(MakeUnique<UE::ActorModifierCore::FRenderStateDirtyReasonScope>(UE::ActorModifierCore::ERenderStateDirtyReason::Material));
+			}
+		});
+
+	TextPostRendererUpdateHandle = UE::Text3D::OnText3DPostRendererUpdate().AddLambda(
+		[](const UE::Text3D::FText3DBuildContext& InBuildContext)
+		{
+			if (InBuildContext.CurrentFlag == EText3DRendererFlags::Material)
+			{
+				Scope.Pop();
+			}
+		});
+}
+
+void FAvaTextModule::UnregisterTextRendererDelegates()
+{
+	UE::Text3D::OnText3DPreRendererUpdate().Remove(TextPreRendererUpdateHandle);
+	TextPreRendererUpdateHandle.Reset();
+
+	UE::Text3D::OnText3DPostRendererUpdate().Remove(TextPostRendererUpdateHandle);
+	TextPostRendererUpdateHandle.Reset();
 }
diff --git a/Engine/Plugins/VirtualProduction/Avalanche/Source/AvalancheText/Private/AvaTextModule.h b/Engine/Plugins/VirtualProduction/Avalanche/Source/AvalancheText/Private/AvaTextModule.h
--- a/Engine/Plugins/VirtualProduction/Avalanche/Source/AvalancheText/Private/AvaTextModule.h
+++ b/Engine/Plugins/VirtualProduction/Avalanche/Source/AvalancheText/Private/AvaTextModule.h
@@ -2,6 +2,7 @@
 
 #pragma once
 
+#include "Delegates/IDelegateInstance.h"
 #include "Modules/ModuleInterface.h"
 
 class FAvaTextModule : public IModuleInterface
@@ -9,5 +10,12 @@ class FAvaTextModule : public IModuleInterface
 public:
 	//~ Begin IModuleInterface
 	virtual void StartupModule() override;
+	virtual void ShutdownModule() override;
 	//~ End IModuleInterface
+
+	void RegisterTextRendererDelegates();
+	void UnregisterTextRendererDelegates();
+
+	FDelegateHandle TextPreRendererUpdateHandle;
+	FDelegateHandle TextPostRendererUpdateHandle;
 };
diff --git a/Engine/Plugins/VirtualProduction/Text3D/Source/Text3D/Private/Renderers/Text3DRendererBase.cpp b/Engine/Plugins/VirtualProduction/Text3D/Source/Text3D/Private/Renderers/Text3DRendererBase.cpp
--- a/Engine/Plugins/VirtualProduction/Text3D/Source/Text3D/Private/Renderers/Text3DRendererBase.cpp
+++ b/Engine/Plugins/VirtualProduction/Text3D/Source/Text3D/Private/Renderers/Text3DRendererBase.cpp
@@ -14,6 +14,7 @@
 #include "Logs/Text3DLogs.h"
 #include "Settings/Text3DProjectSettings.h"
 #include "Text3DComponent.h"
+#include "Text3DDelegates.h"
 
 #if WITH_EDITOR
 void UText3DRendererBase::BindDebugDelegate()
@@ -96,7 +97,14 @@ void UText3DRendererBase::Update(const UE::Text3D::Renderer::FUpdateParameters&
 {
 	if (bInitialized)
 	{
+		const UE::Text3D::FText3DBuildContext BuildContext
+			{
+				.CurrentFlag = InParameters.CurrentFlag
+			};
+
+		UE::Text3D::BroadcastPreRendererUpdate(BuildContext);
 		OnUpdate(InParameters);
+		UE::Text3D::BroadcastPostRendererUpdate(BuildContext);
 	}
 }
 
diff --git a/Engine/Plugins/VirtualProduction/Text3D/Source/Text3D/Private/Text3DDelegates.cpp b/Engine/Plugins/VirtualProduction/Text3D/Source/Text3D/Private/Text3DDelegates.cpp
new file mode 100644
--- /dev/null
+++ b/Engine/Plugins/VirtualProduction/Text3D/Source/Text3D/Private/Text3DDelegates.cpp
@@ -0,0 +1,44 @@
+// Copyright Epic Games, Inc. All Rights Reserved.
+
+#include "Text3DDelegates.h"
+
+namespace UE::Text3D
+{
+namespace Private
+{
+
+TMulticastDelegate<void(const FText3DBuildContext&)>& GetPreRendererUpdateDelegate()
+{
+	static TMulticastDelegate<void(const FText3DBuildContext&)> StaticDelegate;
+	return StaticDelegate;
+}
+
+TMulticastDelegate<void(const FText3DBuildContext&)>& GetPostRendererUpdateDelegate()
+{
+	static TMulticastDelegate<void(const FText3DBuildContext&)> StaticDelegate;
+	return StaticDelegate;
+}
+
+} // UE::Text3D::Private
+
+void BroadcastPreRendererUpdate(const FText3DBuildContext& InBuildContext)
+{
+	Private::GetPreRendererUpdateDelegate().Broadcast(InBuildContext);
+}
+
+void BroadcastPostRendererUpdate(const FText3DBuildContext& InBuildContext)
+{
+	Private::GetPostRendererUpdateDelegate().Broadcast(InBuildContext);
+}
+
+TMulticastDelegateRegistration<void(const FText3DBuildContext&)>& OnText3DPreRendererUpdate()
+{
+	return Private::GetPreRendererUpdateDelegate();
+}
+
+TMulticastDelegateRegistration<void(const FText3DBuildContext&)>& OnText3DPostRendererUpdate()
+{
+	return Private::GetPostRendererUpdateDelegate();
+}
+
+} // UE::Text3D
diff --git a/Engine/Plugins/VirtualProduction/Text3D/Source/Text3D/Public/Text3DDelegates.h b/Engine/Plugins/VirtualProduction/Text3D/Source/Text3D/Public/Text3DDelegates.h
new file mode 100644
--- /dev/null
+++ b/Engine/Plugins/VirtualProduction/Text3D/Source/Text3D/Public/Text3DDelegates.h
@@ -0,0 +1,31 @@
+// Copyright Epic Games, Inc. All Rights Reserved.
+
+#pragma once
+
+#include "Delegates/Delegate.h"
+#include "Text3DTypes.h"
+
+#define UE_API TEXT3D_API
+
+namespace UE::Text3D
+{
+
+struct UE_INTERNAL FText3DBuildContext
+{
+	/** The current flag to update in the renderer */
+	EText3DRendererFlags CurrentFlag = EText3DRendererFlags::None;
+};
+
+/** Called before the Text3D renderer update */
+void BroadcastPreRendererUpdate(const FText3DBuildContext& InBuildContext);
+/** Called after the Text3D renderer update */
+void BroadcastPostRendererUpdate(const FText3DBuildContext& InBuildContext);
+
+/** Delegate called before the Text3D renderer update */
+UE_INTERNAL UE_API TMulticastDelegateRegistration<void(const FText3DBuildContext&)>& OnText3DPreRendererUpdate();
+/** Delegate called after the Text3D renderer update */
+UE_INTERNAL UE_API TMulticastDelegateRegistration<void(const FText3DBuildContext&)>& OnText3DPostRendererUpdate();
+
+} // UE::Text3D
+
+#undef UE_API
