AvalancheFunctionalTest - Fix race condition in AScreenshotFunctionalTestAvalanche::FinishTest by pinning Niagara assets against GC during in-flight CPU compilation

#was-backout 52469351
#Synced-CL 52468003

#Synced-CL 52537439

diff --git a/Engine/Plugins/VirtualProduction/Avalanche/Source/AvalancheFunctionalTest/Private/ScreenshotFunctionalTestAvalanche.cpp b/Engine/Plugins/VirtualProduction/Avalanche/Source/AvalancheFunctionalTest/Private/ScreenshotFunctionalTestAvalanche.cpp
--- a/Engine/Plugins/VirtualProduction/Avalanche/Source/AvalancheFunctionalTest/Private/ScreenshotFunctionalTestAvalanche.cpp
+++ b/Engine/Plugins/VirtualProduction/Avalanche/Source/AvalancheFunctionalTest/Private/ScreenshotFunctionalTestAvalanche.cpp
@@ -7,8 +7,10 @@
 #include "Tests/AutomationCommon.h"
 #include "UnrealClient.h"
 #include "NiagaraComponent.h"
+#include "NiagaraSystem.h"
 #include "NiagaraSystemInstanceController.h"
 #include "Components/PrimitiveComponent.h"
+#include "UObject/UObjectIterator.h"
 
 DEFINE_LOG_CATEGORY_STATIC(LogScreenshotAvalanche, Log, All);
 
@@ -254,3 +256,74 @@ bool AScreenshotFunctionalTestAvalanche::IsActorFullyInitialized(AActor* Actor)
 
 	return true;
 }
+
+void AScreenshotFunctionalTestAvalanche::FinishTest(EFunctionalTestResult TestResult, const FString& Message)
+{
+	// Unroot any systems left over from a previous FinishTest call (safety measure).
+	for (UNiagaraSystem* NiagaraSystem : RootedNiagaraSystems)
+	{
+		if (IsValid(NiagaraSystem) && NiagaraSystem->IsRooted())
+		{
+			NiagaraSystem->RemoveFromRoot();
+		}
+	}
+	RootedNiagaraSystems.Reset();
+
+	// Root all UNiagaraSystem assets referenced by NiagaraComponents in this world.
+	// Niagara async compilation tasks (kicked off by scalability CVar changes during
+	// PrepareTest) hold a TWeakObjectPtr<UNiagaraSystem>. If the Cloner actor is torn
+	// down and GC runs before DigestSystemInfo() executes, that weak ptr becomes null
+	// and the task crashes at EmitterHandles (offset 0x390). AddToRoot() prevents GC
+	// from marking these assets pending-kill until BeginDestroy() cleans them up.
+	if (UWorld* World = GetWorld())
+	{
+		for (TObjectIterator<UNiagaraComponent> It; It; ++It)
+		{
+			UNiagaraComponent* NiagaraComp = *It;
+			if (!IsValid(NiagaraComp) || NiagaraComp->GetWorld() != World)
+			{
+				continue;
+			}
+
+			if (UNiagaraSystem* NiagaraSystem = NiagaraComp->GetAsset())
+			{
+				if (!NiagaraSystem->IsRooted())
+				{
+					NiagaraSystem->AddToRoot();
+					RootedNiagaraSystems.Add(NiagaraSystem);
+					UE_LOG(LogScreenshotAvalanche, Log,
+						TEXT("FinishTest '%s': Rooted NiagaraSystem '%s' to outlive async compilation tasks"),
+						*GetName(), *NiagaraSystem->GetName());
+				}
+			}
+		}
+
+		UE_LOG(LogScreenshotAvalanche, Log,
+			TEXT("FinishTest '%s': Rooted %d NiagaraSystem(s) to survive async compilation during teardown (result=%d)"),
+			*GetName(), RootedNiagaraSystems.Num(), (int32)TestResult);
+	}
+
+	Super::FinishTest(TestResult, Message);
+}
+
+void AScreenshotFunctionalTestAvalanche::BeginDestroy()
+{
+	// Unroot NiagaraSystems pinned in FinishTest. Releasing the root here is a
+	// best-effort cleanup; FinishTest pins assets to extend their lifetime beyond
+	// teardown, but BeginDestroy does not wait for in-flight compilation tasks to
+	// complete. If tasks are still running when this fires, they hold
+	// TWeakObjectPtr references that will become invalid on the next GC cycle.
+	for (UNiagaraSystem* NiagaraSystem : RootedNiagaraSystems)
+	{
+		if (IsValid(NiagaraSystem) && NiagaraSystem->IsRooted())
+		{
+			NiagaraSystem->RemoveFromRoot();
+			UE_LOG(LogScreenshotAvalanche, Log,
+				TEXT("BeginDestroy '%s': Unrooted NiagaraSystem '%s'"),
+				*GetName(), *NiagaraSystem->GetName());
+		}
+	}
+	RootedNiagaraSystems.Reset();
+
+	Super::BeginDestroy();
+}
diff --git a/Engine/Plugins/VirtualProduction/Avalanche/Source/AvalancheFunctionalTest/Public/ScreenshotFunctionalTestAvalanche.h b/Engine/Plugins/VirtualProduction/Avalanche/Source/AvalancheFunctionalTest/Public/ScreenshotFunctionalTestAvalanche.h
--- a/Engine/Plugins/VirtualProduction/Avalanche/Source/AvalancheFunctionalTest/Public/ScreenshotFunctionalTestAvalanche.h
+++ b/Engine/Plugins/VirtualProduction/Avalanche/Source/AvalancheFunctionalTest/Public/ScreenshotFunctionalTestAvalanche.h
@@ -7,6 +7,8 @@
 #include "ScreenshotFunctionalTest.h"
 #include "ScreenshotFunctionalTestAvalanche.generated.h"
 
+class UNiagaraSystem;
+
 /**
  * Screenshot functional test optimized for Avalanche (Motion Design) levels.
  *
@@ -80,6 +82,13 @@ protected:
 	AVALANCHEFUNCTIONALTEST_API virtual void RequestScreenshot() override;
 	AVALANCHEFUNCTIONALTEST_API virtual void OnScreenShotCaptured(int32 InSizeX, int32 InSizeY, const TArray<FColor>& InImageData) override;
 
+	// Override to root NiagaraSystem assets before test teardown, preventing GC from
+	// freeing them while async Niagara compilation tasks still hold TWeakObjectPtr refs.
+	AVALANCHEFUNCTIONALTEST_API virtual void FinishTest(EFunctionalTestResult TestResult, const FString& Message) override;
+
+	// Override to unroot NiagaraSystem assets once the test actor is being destroyed.
+	AVALANCHEFUNCTIONALTEST_API virtual void BeginDestroy() override;
+
 private:
 	/**
 	 * Checks if the specified actor is fully initialized and ready for screenshot.
@@ -94,6 +103,10 @@ private:
 	// Time when PrepareTest() was called, used to enforce MaxActorWaitTime
 	float ActorWaitStartTime = 0.0f;
 
+	// NiagaraSystem assets temporarily rooted in FinishTest() to outlive async
+	// compilation tasks that may fire during test teardown. Unrooted in BeginDestroy().
+	TArray<TObjectPtr<UNiagaraSystem>> RootedNiagaraSystems;
+
 private:
 	/**
 	 * Helper struct to manage CVar overrides at Console priority.
