Motion Design: added a bp function to gather actors and order them based on the motion design scene tree order.

#Synced-CL 52400047

diff --git a/Engine/Plugins/VirtualProduction/Avalanche/Source/Avalanche/Private/AvaSceneSubsystem.cpp b/Engine/Plugins/VirtualProduction/Avalanche/Source/Avalanche/Private/AvaSceneSubsystem.cpp
--- a/Engine/Plugins/VirtualProduction/Avalanche/Source/Avalanche/Private/AvaSceneSubsystem.cpp
+++ b/Engine/Plugins/VirtualProduction/Avalanche/Source/Avalanche/Private/AvaSceneSubsystem.cpp
@@ -17,6 +17,22 @@ namespace UE::Ava::Private
 			EWorldType::Type::Inactive,
 			EWorldType::Type::EditorPreview,
 		};
+
+	const ULevel* GetLevel(const UObject* InContextObject, const UWorld* InWorld)
+	{
+		if (InContextObject)
+		{
+			if (const ULevel* ContextLevel = Cast<ULevel>(InContextObject))
+			{
+				return ContextLevel;
+			}
+			if (const ULevel* OuterLevel = InContextObject->GetTypedOuter<ULevel>())
+			{
+				return OuterLevel;
+			}
+		}
+		return InWorld ? InWorld->PersistentLevel : nullptr;
+	}
 }
 
 void UAvaSceneSubsystem::RegisterSceneInterface(ULevel* InLevel, IAvaSceneInterface* InSceneInterface)
@@ -46,11 +62,21 @@ IAvaSceneInterface* UAvaSceneSubsystem::GetSceneInterface(const ULevel* InLevel)
 
 IAvaSceneInterface* UAvaSceneSubsystem::FindSceneInterface(const ULevel* InLevel)
 {
-	AAvaScene* SceneActor = nullptr;
-	if (InLevel)
+	if (!InLevel)
 	{
-		InLevel->Actors.FindItemByClass(&SceneActor);
+		return nullptr;
 	}
+
+	if (const UAvaSceneSubsystem* SceneSubsystem = InLevel->OwningWorld ? InLevel->OwningWorld->GetSubsystem<UAvaSceneSubsystem>() : nullptr)
+	{
+		if (IAvaSceneInterface* SceneInterface = SceneSubsystem->GetSceneInterface(InLevel))
+		{
+			return SceneInterface;
+		}
+	}
+
+	AAvaScene* SceneActor = nullptr;
+	InLevel->Actors.FindItemByClass(&SceneActor);
 	return SceneActor;
 }
 
@@ -60,6 +86,56 @@ bool UAvaSceneSubsystem::IsReadyToPlay(const UObject* InContextObject) const
 	return !TextBuildSystem || !TextBuildSystem->IsBuildInProgress();
 }
 
+TArray<AActor*> UAvaSceneSubsystem::GatherSceneTreeActors(const UObject* InContextObject, const AActor* InParentActor, bool bInIncludeDescendants)
+{
+	UWorld* const World = GEngine ? GEngine->GetWorldFromContextObject(InContextObject, EGetWorldErrorMode::LogAndReturnNull) : nullptr;
+	if (!World)
+	{
+		return {};
+	}
+
+	const ULevel* const Level = UE::Ava::Private::GetLevel(InContextObject, World);
+	if (!Level)
+	{
+		return {};
+	}
+
+	IAvaSceneInterface* const SceneInterface = UAvaSceneSubsystem::FindSceneInterface(Level);
+	if (!SceneInterface)
+	{
+		return {};
+	}
+
+	FAvaSceneTree& SceneTree = SceneInterface->GetSceneTree();
+
+	TArray<AActor*> Actors;
+	if (InParentActor)
+	{
+		InParentActor->GetAttachedActors(Actors, /*bResetArray*/false, bInIncludeDescendants);
+	}
+	else
+	{
+		Actors = Level->Actors;	
+	}
+
+	if (!SceneTree.IsSorted())
+	{
+		SceneTree.SortItems();
+		SceneTree.ResolveObjects(World);
+	}
+
+	Actors.Remove(nullptr);
+	Actors.Sort(
+		[&SceneTree, World](const AActor& A, const AActor& B)
+		{
+			const FAvaSceneTreeNode* NodeA = SceneTree.FindObjectTreeNode(&A, World);
+			const FAvaSceneTreeNode* NodeB = SceneTree.FindObjectTreeNode(&B, World);
+			return FAvaSceneTree::CompareTreeItemOrder(NodeA, NodeB);
+		});
+
+	return Actors;
+}
+
 void UAvaSceneSubsystem::PostInitialize()
 {
 	Super::PostInitialize();
diff --git a/Engine/Plugins/VirtualProduction/Avalanche/Source/Avalanche/Public/AvaSceneSubsystem.h b/Engine/Plugins/VirtualProduction/Avalanche/Source/Avalanche/Public/AvaSceneSubsystem.h
--- a/Engine/Plugins/VirtualProduction/Avalanche/Source/Avalanche/Public/AvaSceneSubsystem.h
+++ b/Engine/Plugins/VirtualProduction/Avalanche/Source/Avalanche/Public/AvaSceneSubsystem.h
@@ -37,6 +37,15 @@ public:
 	UFUNCTION(BlueprintPure, Category="Motion Design", meta=(DefaultToSelf="InContextObject"))
 	UE_API bool IsReadyToPlay(const UObject* InContextObject) const;
 
+	/**
+	 * Gathers actors in order of their appearance in the tree 
+	 * @param InContextObject the context object to get the level from
+	 * @param InParentActor if specified, only the children (and optionally including descendants) of this actor will be considered. If not specified, all the actors in the context level will be used.
+	 * @param bInIncludeDescendants if specified, descendants of the given parent actor will also be considered.
+	 */
+	UFUNCTION(BlueprintCallable, Category="Motion Design", meta=(DefaultToSelf="InContextObject"))
+	static UE_API TArray<AActor*> GatherSceneTreeActors(const UObject* InContextObject, const AActor* InParentActor=nullptr, bool bInIncludeDescendants=true);
+
 protected:
 	//~ Begin UWorldSubsystem
 	UE_API virtual void PostInitialize() override;
diff --git a/Engine/Plugins/VirtualProduction/Avalanche/Source/AvalancheSceneTree/Private/AvaSceneTree.cpp b/Engine/Plugins/VirtualProduction/Avalanche/Source/AvalancheSceneTree/Private/AvaSceneTree.cpp
--- a/Engine/Plugins/VirtualProduction/Avalanche/Source/AvalancheSceneTree/Private/AvaSceneTree.cpp
+++ b/Engine/Plugins/VirtualProduction/Avalanche/Source/AvalancheSceneTree/Private/AvaSceneTree.cpp
@@ -194,7 +194,7 @@ void FAvaSceneTree::UpdateTreeNodes()
 	}
 }
 
-TArray<AActor*> FAvaSceneTree::GetChildActors(AActor* const InParentActor) const
+TArray<AActor*> FAvaSceneTree::GetChildActors(const AActor* InParentActor) const
 {
 	TArray<AActor*> OutActors;
 
diff --git a/Engine/Plugins/VirtualProduction/Avalanche/Source/AvalancheSceneTree/Public/AvaSceneTree.h b/Engine/Plugins/VirtualProduction/Avalanche/Source/AvalancheSceneTree/Public/AvaSceneTree.h
--- a/Engine/Plugins/VirtualProduction/Avalanche/Source/AvalancheSceneTree/Public/AvaSceneTree.h
+++ b/Engine/Plugins/VirtualProduction/Avalanche/Source/AvalancheSceneTree/Public/AvaSceneTree.h
@@ -20,12 +20,23 @@ struct AVALANCHESCENETREE_API FAvaSceneTree
 
 	void PostSerialize(const FArchive& Ar);
 
-	FAvaSceneTreeNode& GetRootNode() { return RootNode; }
+	FAvaSceneTreeNode& GetRootNode()
+	{
+		return RootNode;
+	}
 
-	const FAvaSceneTreeNode& GetRootNode() const { return RootNode; }
+	const FAvaSceneTreeNode& GetRootNode() const
+	{
+		return RootNode;
+	}
 
 	FAvaSceneTreeNode* FindTreeNode(const FAvaSceneItem& InItem);
 
+	int32 GetSceneItemCount() const
+	{
+		return SceneItems.Num();
+	}
+
 	/**
 	 * Finds tree node using the object map to avoid FAvaSceneItem allocation
 	 * ResolveObjects() must've been called for the tree for the object map to be up-to-date.
@@ -45,7 +56,7 @@ struct AVALANCHESCENETREE_API FAvaSceneTree
 	void Reset();
 
 	/** Returns all resolved child actors of a specified parent actor in the scene tree. */
-	TArray<AActor*> GetChildActors(AActor* const InParentActor) const;
+	TArray<AActor*> GetChildActors(const AActor* InParentActor) const;
 
 	/** Returns all root actors in this tree */
 	TArray<AActor*> GetRootActors(const ULevel* InLevel) const;
