Motion Design: added ability to specify transition tree assets in TL for shared logic.

#Synced-CL 52400047

diff --git a/Engine/Plugins/VirtualProduction/Avalanche/Source/AvalancheEditor/Private/Transition/AvaTransitionExtension.cpp b/Engine/Plugins/VirtualProduction/Avalanche/Source/AvalancheEditor/Private/Transition/AvaTransitionExtension.cpp
--- a/Engine/Plugins/VirtualProduction/Avalanche/Source/AvalancheEditor/Private/Transition/AvaTransitionExtension.cpp
+++ b/Engine/Plugins/VirtualProduction/Avalanche/Source/AvalancheEditor/Private/Transition/AvaTransitionExtension.cpp
@@ -139,8 +139,15 @@ void FAvaTransitionExtension::ExtendSettingsCategory(IDetailCategoryBuilder& InC
 {
 	if (IAvaTransitionBehavior* TransitionBehavior = GetTransitionBehavior())
 	{
+		const TArray<UObject*> TransitionBehaviorObjects = { &TransitionBehavior->AsUObject() };
+
 		InCategoryBuilder.SetDisplayName(LOCTEXT("CategoryDisplayName", "Transition Logic"));
-		InCategoryBuilder.AddExternalObjectProperty({ &TransitionBehavior->AsUObject() }, TEXT("StateTreeReference"));
+
+		TransitionBehavior->ForEachDetailsEditableProperty(
+			[&InCategoryBuilder, &TransitionBehaviorObjects](const IAvaTransitionBehavior::FPropertyContext& InPropertyContext)
+			{
+				InCategoryBuilder.AddExternalObjectProperty(TransitionBehaviorObjects, InPropertyContext.Name);
+			});
 	}
 }
 
diff --git a/Engine/Plugins/VirtualProduction/Avalanche/Source/AvalancheTransition/Private/Behavior/AvaTransitionBehaviorActor.cpp b/Engine/Plugins/VirtualProduction/Avalanche/Source/AvalancheTransition/Private/Behavior/AvaTransitionBehaviorActor.cpp
--- a/Engine/Plugins/VirtualProduction/Avalanche/Source/AvalancheTransition/Private/Behavior/AvaTransitionBehaviorActor.cpp
+++ b/Engine/Plugins/VirtualProduction/Avalanche/Source/AvalancheTransition/Private/Behavior/AvaTransitionBehaviorActor.cpp
@@ -30,7 +30,7 @@ AAvaTransitionBehaviorActor::AAvaTransitionBehaviorActor()
 
 UAvaTransitionTree* AAvaTransitionBehaviorActor::GetTransitionTree() const
 {
-	return TransitionTree;
+	return Cast<UAvaTransitionTree>(StateTreeReference.GetMutableStateTree());
 }
 
 FAvaTagHandle AAvaTransitionBehaviorActor::GetTransitionLayer() const
@@ -68,6 +68,23 @@ const FStateTreeReference& AAvaTransitionBehaviorActor::GetStateTreeReference()
 	return StateTreeReference;
 }
 
+#if WITH_EDITOR
+void AAvaTransitionBehaviorActor::ForEachDetailsEditableProperty(TFunctionRef<void(const FPropertyContext&)> InFunc) const
+{
+	IAvaTransitionBehavior::FPropertyContext DefaultPropertyContext;
+	auto CallFunc = [&InFunc, &DefaultPropertyContext](FName InPropertyName)
+		{
+			DefaultPropertyContext.Name = InPropertyName;
+			InFunc(DefaultPropertyContext);
+		};
+
+	CallFunc(GET_MEMBER_NAME_CHECKED(AAvaTransitionBehaviorActor, bEnabled));
+	CallFunc(GET_MEMBER_NAME_CHECKED(AAvaTransitionBehaviorActor, InstancingMode));
+	CallFunc(GET_MEMBER_NAME_CHECKED(AAvaTransitionBehaviorActor, TransitionLayer));
+	CallFunc(GET_MEMBER_NAME_CHECKED(AAvaTransitionBehaviorActor, StateTreeReference));
+}
+#endif
+
 void AAvaTransitionBehaviorActor::PostActorCreated()
 {
 	Super::PostActorCreated();
@@ -106,36 +123,37 @@ PRAGMA_ENABLE_DEPRECATION_WARNINGS
 }
 
 #if WITH_EDITOR
-bool AAvaTransitionBehaviorActor::CanEditChange(const FProperty* InProperty) const
+void AAvaTransitionBehaviorActor::PostEditChangeProperty(FPropertyChangedEvent& InPropertyChangedEvent)
 {
-	return Super::CanEditChange(InProperty);
+	Super::PostEditChangeProperty(InPropertyChangedEvent);
+	UpdateStateTreeReference();
 }
 
-bool AAvaTransitionBehaviorActor::CanEditChange(const FEditPropertyChain& InPropertyChain) const
+void AAvaTransitionBehaviorActor::PostEditImport()
 {
-	if (!UObject::CanEditChange(InPropertyChain))
-	{
-		return false;
-	}
+	Super::PostEditImport();
+	UpdateStateTreeReference();
+}
 
-	auto GetPropertyName =
-		[](TDoubleLinkedList<FProperty*>::TDoubleLinkedListNode* InChainNode)->FName
-		{
-			if (InChainNode && InChainNode->GetValue())
-			{
-				return InChainNode->GetValue()->GetFName();
-			}
-			return NAME_None;	
-		};
+void AAvaTransitionBehaviorActor::PostEditUndo()
+{
+	Super::PostEditUndo();
+	UpdateStateTreeReference();
+}
+#endif
 
-	const FName MemberPropertyName = GetPropertyName(InPropertyChain.GetActiveMemberNode());
-	const FName ActivePropertyName = GetPropertyName(InPropertyChain.GetActiveNode());
+void AAvaTransitionBehaviorActor::UpdateStateTreeReference()
+{
+	if (!StateTreeReference.GetStateTree())
+	{
+		StateTreeReference.SetStateTree(TransitionTree);
+	}
+}
 
-	// Disable change for StateTree in FStateTreeReference
-	return MemberPropertyName != GET_MEMBER_NAME_CHECKED(AAvaTransitionBehaviorActor, StateTreeReference)
-		|| ActivePropertyName != TEXT("StateTree");
+bool AAvaTransitionBehaviorActor::IsTransitionTreeOverridden() const
+{
+	return GetTransitionTree() != TransitionTree;
 }
-#endif
 
 UAvaTransitionSubsystem* AAvaTransitionBehaviorActor::GetTransitionSubsystem() const
 {
@@ -145,6 +163,11 @@ UAvaTransitionSubsystem* AAvaTransitionBehaviorActor::GetTransitionSubsystem() c
 
 void AAvaTransitionBehaviorActor::ValidateTransitionTree()
 {
+	if (IsTransitionTreeOverridden())
+	{
+		return;
+	}
+
 	if (!ensureAlwaysMsgf(TransitionTree, TEXT("Transition Tree is null. Cannot not validate tree")))
 	{
 		return;
diff --git a/Engine/Plugins/VirtualProduction/Avalanche/Source/AvalancheTransition/Private/Behavior/AvaTransitionBehaviorActor.h b/Engine/Plugins/VirtualProduction/Avalanche/Source/AvalancheTransition/Private/Behavior/AvaTransitionBehaviorActor.h
--- a/Engine/Plugins/VirtualProduction/Avalanche/Source/AvalancheTransition/Private/Behavior/AvaTransitionBehaviorActor.h
+++ b/Engine/Plugins/VirtualProduction/Avalanche/Source/AvalancheTransition/Private/Behavior/AvaTransitionBehaviorActor.h
@@ -33,6 +33,9 @@ protected:
 	virtual EAvaTransitionInstancingMode GetInstancingMode() const override;
 	virtual void SetInstancingMode(EAvaTransitionInstancingMode InInstancingMode) override;
 	virtual const FStateTreeReference& GetStateTreeReference() const override;
+#if WITH_EDITOR
+	virtual void ForEachDetailsEditableProperty(TFunctionRef<void(const FPropertyContext&)> InFunc) const override;
+#endif
 	//~ End IAvaTransitionBehavior
 
 	//~ Begin AActor
@@ -47,11 +50,18 @@ protected:
 	virtual void Serialize(FArchive& Ar) override;
 	virtual void PostLoad() override;
 #if WITH_EDITOR
-	virtual bool CanEditChange(const FProperty* InProperty) const override;
-	virtual bool CanEditChange(const FEditPropertyChain& InPropertyChain) const;
+	virtual void PostEditChangeProperty(FPropertyChangedEvent& InPropertyChangedEvent) override;
+	virtual void PostEditImport() override;
+	virtual void PostEditUndo() override;
 #endif
 	//~ End UObject
 
+	/** Updates the underlying state tree reference, setting it to TransitionTree if null. */
+	void UpdateStateTreeReference();
+
+	/** Determines whether the transition tree to use in transition logic is the overridden tree instead of the owned one */
+	bool IsTransitionTreeOverridden() const;
+
 private:
 	UAvaTransitionSubsystem* GetTransitionSubsystem() const;
 
@@ -60,8 +70,8 @@ private:
 	UPROPERTY()
 	TObjectPtr<UAvaTransitionTree> TransitionTree;
 
-	UPROPERTY(EditAnywhere, Category="Transition Logic")
-	FStateTreeReference StateTreeReference;
+	UPROPERTY(EditAnywhere, Category="Transition Logic", DisplayName = "Transition Tree", meta=(Schema="/Script/AvalancheTransition.AvaTransitionTreeSchema"))
+	mutable FStateTreeReference StateTreeReference;
 
 	/** Transition Layer to use */
 	UPROPERTY(EditAnywhere, Category="Transition Logic")
diff --git a/Engine/Plugins/VirtualProduction/Avalanche/Source/AvalancheTransition/Private/Behavior/AvaTransitionBehaviorInstance.cpp b/Engine/Plugins/VirtualProduction/Avalanche/Source/AvalancheTransition/Private/Behavior/AvaTransitionBehaviorInstance.cpp
--- a/Engine/Plugins/VirtualProduction/Avalanche/Source/AvalancheTransition/Private/Behavior/AvaTransitionBehaviorInstance.cpp
+++ b/Engine/Plugins/VirtualProduction/Avalanche/Source/AvalancheTransition/Private/Behavior/AvaTransitionBehaviorInstance.cpp
@@ -242,6 +242,13 @@ bool FAvaTransitionBehaviorInstance::MakeContext(IAvaTransitionBehavior* InBehav
 	}
 
 	const FStateTreeReference& StateTreeReference = InBehavior->GetStateTreeReference();
+	if (!StateTreeReference.GetStateTree())
+	{
+		UE_LOG(LogAvaTransition, Error
+			, TEXT("SetContextRequirements failed for '%s'. StateTree is invalid")
+			, *LogContext);
+		return false;
+	}
 
 	FAvaTransitionExecutionContext& Context = OutContext.Emplace(*this, InBehavior->AsUObject(), *StateTreeReference.GetStateTree(), InstanceData);
 	if (!Context.IsValid())
diff --git a/Engine/Plugins/VirtualProduction/Avalanche/Source/AvalancheTransition/Public/Behavior/IAvaTransitionBehavior.h b/Engine/Plugins/VirtualProduction/Avalanche/Source/AvalancheTransition/Public/Behavior/IAvaTransitionBehavior.h
--- a/Engine/Plugins/VirtualProduction/Avalanche/Source/AvalancheTransition/Public/Behavior/IAvaTransitionBehavior.h
+++ b/Engine/Plugins/VirtualProduction/Avalanche/Source/AvalancheTransition/Public/Behavior/IAvaTransitionBehavior.h
@@ -55,4 +55,14 @@ public:
 
 	/** Gets the underlying State Tree reference that this transition behavior runs */
 	virtual const FStateTreeReference& GetStateTreeReference() const = 0;
+
+#if WITH_EDITOR
+	struct FPropertyContext
+	{
+		/** Name of the property */
+		FName Name;
+	};
+	/** Calls the given functor on each property that is allowed to be visible in an external details panel */
+	virtual void ForEachDetailsEditableProperty(TFunctionRef<void(const FPropertyContext&)> InFunc) const = 0;
+#endif
 };
