Motion Design: outliner to reset the scene tree on refresh

#Synced-CL 52400047

diff --git a/Engine/Plugins/VirtualProduction/Avalanche/Source/AvalancheOutliner/Private/AvaOutliner.cpp b/Engine/Plugins/VirtualProduction/Avalanche/Source/AvalancheOutliner/Private/AvaOutliner.cpp
--- a/Engine/Plugins/VirtualProduction/Avalanche/Source/AvalancheOutliner/Private/AvaOutliner.cpp
+++ b/Engine/Plugins/VirtualProduction/Avalanche/Source/AvalancheOutliner/Private/AvaOutliner.cpp
@@ -463,7 +463,7 @@ void FAvaOutliner::Refresh()
 	});
 
 	// Save so that the Tree is updated to the latest Outliner State
-	SaveState->SaveSceneTree(*this, /*bInResetTree*/false);
+	SaveState->SaveSceneTree(*this, /*bInResetTree*/true);
 }
 
 TSharedRef<FAvaOutlinerItem> FAvaOutliner::GetTreeRoot() const
